Copying //UE4/Dev-Anim to Dev-Main (//UE4/Dev-Main) @ 8782600

diff --git a/Engine/Source/Runtime/Engine/Classes/Animation/AnimNodeBase.h b/Engine/Source/Runtime/Engine/Classes/Animation/AnimNodeBase.h
--- a/Engine/Source/Runtime/Engine/Classes/Animation/AnimNodeBase.h
+++ b/Engine/Source/Runtime/Engine/Classes/Animation/AnimNodeBase.h
@@ -4,6 +4,7 @@
 
 #include "CoreMinimal.h"
 #include "UObject/ObjectMacros.h"
+#include "UObject/ObjectKey.h"
 #include "UObject/Class.h"
 #include "Engine/EngineTypes.h"
 #include "Animation/AnimTypes.h"
@@ -22,12 +23,114 @@ class UAnimInstance;
 struct FAnimInstanceProxy;
 struct FAnimNode_Base;
 
+/**
+ * Utility container for tracking a stack of ancestor nodes by node type during graph traversal
+ * This is not an exhaustive list of all visited ancestors. During Update nodes must call
+ * FAnimationUpdateContext::TrackAncestor() to appear in the tracker.
+ */
+struct FAnimNodeTracker
+{
+	using FKey = FObjectKey;
+	using FNodeStack = TArray<FAnimNode_Base*, TInlineAllocator<4>>;
+	using FMap = TMap<FKey, FNodeStack, TInlineSetAllocator<4>>;
+
+	FMap Map;
+
+	template<typename NodeType>
+	static FKey GetKey()
+	{
+		return FKey(NodeType::StaticStruct());
+	}
+
+	template<typename NodeType>
+	FKey Push(NodeType* Node)
+	{
+		FKey Key(GetKey<NodeType>());
+		FNodeStack& Stack = Map.FindOrAdd(Key);
+		Stack.Push(Node);
+		return Key;
+	}
+
+	template<typename NodeType>
+	NodeType* Pop()
+	{
+		FNodeStack* Stack = Map.Find(GetKey<NodeType>());
+		return Stack ? static_cast<NodeType*>(Stack->Pop()) : nullptr;
+	}
+
+	FAnimNode_Base* Pop(FKey Key)
+	{
+		FNodeStack* Stack = Map.Find(Key);
+		return Stack ? Stack->Pop() : nullptr;
+	}
+
+	template<typename NodeType>
+	NodeType* Top() const
+	{
+		const FNodeStack* Stack = Map.Find(GetKey<NodeType>());
+		return Stack ? static_cast<NodeType*>(Stack->Top()) : nullptr;
+	}
+
+	void CopyTopsOnly(const FAnimNodeTracker& Source)
+	{
+		Map.Reset();
+		Map.Reserve(Source.Map.Num());
+		for (const auto& Iter : Source.Map)
+		{
+			if (Iter.Value.Top())
+			{
+				FNodeStack& Stack = Map.Add(Iter.Key);
+				Stack.Push(Iter.Value.Top());
+			}
+		}
+	}
+};
+
+
+/** Helper RAII object to cleanup a node added to the node tracker */
+class FScopedAnimNodeTracker
+{
+public:
+	FScopedAnimNodeTracker() = default;
+
+	FScopedAnimNodeTracker(FAnimNodeTracker* InTracker, FAnimNodeTracker::FKey InKey)
+		: Tracker(InTracker)
+		, TrackedKey(InKey)
+	{}
+
+	~FScopedAnimNodeTracker()
+	{
+		if (Tracker && TrackedKey != FAnimNodeTracker::FKey())
+		{
+			Tracker->Pop(TrackedKey);
+		}
+	}
+
+private:
+	FAnimNodeTracker* Tracker = nullptr;
+	FAnimNodeTracker::FKey TrackedKey;
+};
+
+
+/** Persistent state shared during animation tree update  */
+struct FAnimationUpdateSharedContext
+{
+	FAnimNodeTracker AncestorTracker;
+
+	void CopyForCachedUpdate(const FAnimationUpdateSharedContext& Source)
+	{
+		AncestorTracker.CopyTopsOnly(Source.AncestorTracker);
+	}
+};
+
 /** Base class for update/evaluate contexts */
 struct FAnimationBaseContext
 {
 public:
 	FAnimInstanceProxy* AnimInstanceProxy;
 
+	FAnimationBaseContext();
+
 protected:
 	// DEPRECATED - Please use constructor that uses an FAnimInstanceProxy*
 	ENGINE_API FAnimationBaseContext(UAnimInstance* InAnimInstance);
@@ -83,39 +186,67 @@ public:
 struct FAnimationUpdateContext : public FAnimationBaseContext
 {
 private:
+	FAnimationUpdateSharedContext* SharedContext;
+
 	float CurrentWeight;
 	float RootMotionWeightModifier;
 
 	float DeltaTime;
+
 public:
-	FAnimationUpdateContext(FAnimInstanceProxy* InAnimInstanceProxy, float InDeltaTime)
+	FAnimationUpdateContext(FAnimInstanceProxy* InAnimInstanceProxy = nullptr)
 		: FAnimationBaseContext(InAnimInstanceProxy)
+		, SharedContext(nullptr)
 		, CurrentWeight(1.0f)
-		, RootMotionWeightModifier(1.f)
-		, DeltaTime(InDeltaTime)
+		, RootMotionWeightModifier(1.0f)
+		, DeltaTime(0.0f)
+	{
+	}
+
+	FAnimationUpdateContext(FAnimInstanceProxy* InAnimInstanceProxy, float InDeltaTime, FAnimationUpdateSharedContext* InSharedContext = nullptr)
+		: FAnimationUpdateContext(InAnimInstanceProxy)
+	{
+		SharedContext = InSharedContext;
+		DeltaTime = InDeltaTime;
+	}
+
+
+	FAnimationUpdateContext(const FAnimationUpdateContext& Copy) = default;
+
+	FAnimationUpdateContext(const FAnimationUpdateContext& Copy, FAnimInstanceProxy* InAnimInstanceProxy)
+		: FAnimationBaseContext(InAnimInstanceProxy)
+		, SharedContext(Copy.SharedContext)
+		, CurrentWeight(Copy.CurrentWeight)
+		, RootMotionWeightModifier(Copy.RootMotionWeightModifier)
+		, DeltaTime(Copy.DeltaTime)
 	{
 	}
 
 public:
 	FAnimationUpdateContext WithOtherProxy(FAnimInstanceProxy* InAnimInstanceProxy) const
 	{
-		FAnimationUpdateContext Result(InAnimInstanceProxy, DeltaTime);
-		Result.CurrentWeight = CurrentWeight;
-		Result.RootMotionWeightModifier = RootMotionWeightModifier;
+		return FAnimationUpdateContext(*this, InAnimInstanceProxy);
+	}
+
+	FAnimationUpdateContext WithOtherSharedContext(FAnimationUpdateSharedContext* InSharedContext) const
+	{
+		FAnimationUpdateContext Result(*this);
+		Result.SharedContext = InSharedContext;
+
 		return Result;
 	}
 
-	FAnimationUpdateContext FractionalWeight(float Multiplier) const
+	FAnimationUpdateContext FractionalWeight(float WeightMultiplier) const
 	{
-		FAnimationUpdateContext Result(AnimInstanceProxy, DeltaTime);
-		Result.CurrentWeight = CurrentWeight * Multiplier;
-		Result.RootMotionWeightModifier = RootMotionWeightModifier;
+		FAnimationUpdateContext Result(*this);
+		Result.CurrentWeight = CurrentWeight * WeightMultiplier;
+
 		return Result;
 	}
 
 	FAnimationUpdateContext FractionalWeightAndRootMotion(float WeightMultiplier, float RootMotionMultiplier) const
 	{
-		FAnimationUpdateContext Result(AnimInstanceProxy, DeltaTime);
+		FAnimationUpdateContext Result(*this);
 		Result.CurrentWeight = CurrentWeight * WeightMultiplier;
 		Result.RootMotionWeightModifier = RootMotionWeightModifier * RootMotionMultiplier;
 
@@ -124,21 +255,54 @@ public:
 
 	FAnimationUpdateContext FractionalWeightAndTime(float WeightMultiplier, float TimeMultiplier) const
 	{
-		FAnimationUpdateContext Result(AnimInstanceProxy, DeltaTime * TimeMultiplier);
+		FAnimationUpdateContext Result(*this);
+		Result.DeltaTime = DeltaTime * TimeMultiplier;
 		Result.CurrentWeight = CurrentWeight * WeightMultiplier;
-		Result.RootMotionWeightModifier = RootMotionWeightModifier;
 		return Result;
 	}
 
 	FAnimationUpdateContext FractionalWeightTimeAndRootMotion(float WeightMultiplier, float TimeMultiplier, float RootMotionMultiplier) const
 	{
-		FAnimationUpdateContext Result(AnimInstanceProxy, DeltaTime * TimeMultiplier);
+		FAnimationUpdateContext Result(*this);
+		Result.DeltaTime = DeltaTime * TimeMultiplier;
 		Result.CurrentWeight = CurrentWeight * WeightMultiplier;
 		Result.RootMotionWeightModifier = RootMotionWeightModifier * RootMotionMultiplier;
 
 		return Result;
 	}
 
+	// Add a node to the list of tracked ancestors
+	template<typename NodeType>
+	FScopedAnimNodeTracker TrackAncestor(NodeType* Node) const
+	{
+		if (ensure(SharedContext != nullptr))
+		{
+			FAnimNodeTracker::FKey Key = SharedContext->AncestorTracker.Push<NodeType>(Node);
+			return FScopedAnimNodeTracker(&SharedContext->AncestorTracker, Key);
+		}
+
+		return FScopedAnimNodeTracker();
+	}
+
+	// Returns the nearest ancestor node of a particular type
+	template<typename NodeType>
+	NodeType* GetAncestor() const
+	{
+		if (ensure(SharedContext != nullptr))
+		{
+			FAnimNode_Base* Node = SharedContext->AncestorTracker.Top<NodeType>();
+			return static_cast<NodeType*>(Node);
+		}
+		
+		return nullptr;
+	}
+
+	// Returns persistent state that is tracked through animation tree update
+	FAnimationUpdateSharedContext* GetSharedContext() const
+	{
+		return SharedContext;
+	}
+
 	// Returns the final blend weight contribution for this stage
 	float GetFinalBlendWeight() const { return CurrentWeight; }
 
@@ -710,6 +874,18 @@ struct ENGINE_API FAnimNode_Base
 	/** Called to help dynamics-based updates to recover correctly from large movements/teleports */
 	virtual void ResetDynamics(ETeleportType InTeleportType);
 
+	/**
+	 * Override this if your node uses ancestor tracking and wants to be informed of Update() calls
+	 * that were skipped due to pose caching.
+	 */
+	virtual bool WantsSkippedUpdates() const { return false; }
+	
+	/**
+	 * Called on a tracked ancestor node when there are Update() calls that were skipped due to pose 
+	 * caching. Your node must implement WantsSkippedUpdates to receive this callback.
+	 */
+	virtual void OnUpdatesSkipped(TArrayView<const FAnimationUpdateContext*> SkippedUpdateContexts) {}
+
 	/** Called after compilation */
 	virtual void PostCompile(const class USkeleton* InSkeleton) {}
 
diff --git a/Engine/Source/Runtime/Engine/Classes/Animation/AnimNode_Inertialization.h b/Engine/Source/Runtime/Engine/Classes/Animation/AnimNode_Inertialization.h
new file mode 100644
--- /dev/null
+++ b/Engine/Source/Runtime/Engine/Classes/Animation/AnimNode_Inertialization.h
@@ -0,0 +1,254 @@
+// Copyright 1998-2019 Epic Games, Inc. All Rights Reserved.
+
+#pragma once
+
+#include "CoreMinimal.h"
+#include "UObject/ObjectMacros.h"
+#include "Animation/AnimNodeBase.h"
+#include "AnimNode_Inertialization.generated.h"
+
+
+// Inertialization: High-Performance Animation Transitions in 'Gears of War'
+// David Bollo
+// Game Developer Conference 2018
+//
+// https://www.gdcvault.com/play/1025165/Inertialization
+// https://www.gdcvault.com/play/1025331/Inertialization
+
+
+UENUM()
+enum class EInertializationState : uint8
+{
+	Inactive,		// Inertialization inactive
+	Pending,		// Inertialization request pending... prepare to capture the pose difference and then switch to active
+	Active			// Inertialization active... apply the previously captured pose difference
+};
+
+
+UENUM()
+enum class EInertializationBoneState : uint8
+{
+	Invalid,		// Invalid bone (ie: bone was present in the skeleton but was not present in the pose when it was captured)
+	Valid,			// Valid bone
+	Excluded		// Valid bone that is to be excluded from the inertialization request
+};
+
+
+UENUM()
+enum class EInertializationSpace : uint8
+{
+	Default,		// Inertialize in local space (default)
+	WorldSpace,		// Inertialize translation and rotation in world space (to conceal discontinuities in actor transform such snapping to a new attach parent)
+	WorldRotation	// Inertialize rotation only in world space (to conceal discontinuities in actor orientation)
+};
+
+
+USTRUCT()
+struct FInertializationPose
+{
+	GENERATED_BODY()
+
+	FTransform ComponentTransform;
+
+	// Bone transforms indexed by skeleton bone index.  Transforms are in local space except for direct descendants of
+	// the root which are in component space (ie: they have been multiplied by the root).  Invalid bones (ie: bones
+	// that are present in the skeleton but were not present in the pose when it was captured) are all zero
+	//
+	TArray<FTransform> BoneTransforms;
+
+	// Bone states indexed by skeleton bone index
+	//
+	TArray<EInertializationBoneState> BoneStates;
+
+	FName AttachParentName;
+	float DeltaTime;
+
+	FInertializationPose()
+		: ComponentTransform(FTransform::Identity)
+		, AttachParentName(NAME_None)
+		, DeltaTime(0.0f)
+	{
+	}
+
+	void InitFrom(const FCompactPose& Pose, const FTransform& InComponentTransform, const FName& InAttachParentName, float InDeltaTime);
+};
+
+
+USTRUCT()
+struct FInertializationBoneDiff
+{
+	GENERATED_BODY()
+
+	FVector TranslationDirection;
+	FVector RotationAxis;
+	FVector ScaleAxis;
+
+	float TranslationMagnitude;
+	float TranslationSpeed;
+
+	float RotationAngle;
+	float RotationSpeed;
+
+	float ScaleMagnitude;
+	float ScaleSpeed;
+
+	FInertializationBoneDiff()
+		: TranslationDirection(FVector::ZeroVector)
+		, RotationAxis(FVector::ZeroVector)
+		, ScaleAxis(FVector::ZeroVector)
+		, TranslationMagnitude(0.0f)
+		, TranslationSpeed(0.0f)
+		, RotationAngle(0.0f)
+		, RotationSpeed(0.0f)
+		, ScaleMagnitude(0.0f)
+		, ScaleSpeed(0.0f)
+	{
+	}
+
+	void Clear()
+	{
+		TranslationDirection = FVector::ZeroVector;
+		RotationAxis = FVector::ZeroVector;
+		ScaleAxis = FVector::ZeroVector;
+		TranslationMagnitude = 0.0f;
+		TranslationSpeed = 0.0f;
+		RotationAngle = 0.0f;
+		RotationSpeed = 0.0f;
+		ScaleMagnitude = 0.0f;
+		ScaleSpeed = 0.0f;
+	}
+};
+
+
+USTRUCT()
+struct FInertializationPoseDiff
+{
+	GENERATED_BODY()
+
+	FInertializationPoseDiff()
+		: InertializationSpace(EInertializationSpace::Default)
+	{
+	}
+
+	void Reset()
+	{
+		BoneDiffs.Empty();
+		InertializationSpace = EInertializationSpace::Default;
+	}
+
+	// Initialize the pose difference from the current pose and the two previous snapshots
+	//
+	// Pose					the current frame's pose
+	// ComponentTransform	the current frame's component to world transform
+	// AttachParentName		the current frame's attach parent name (for checking if the attachment has changed)
+	// Prev1				the previous frame's pose
+	// Prev2				the pose from two frames before
+	//
+	void InitFrom(const FCompactPose& Pose, const FTransform& ComponentTransform, const FName& AttachParentName, const FInertializationPose& Prev1, const FInertializationPose& Prev2);
+
+	// Apply this difference to a pose, decaying over time as InertializationElapsedTime approaches InertializationDuration
+	//
+	void ApplyTo(FCompactPose& Pose, float InertializationElapsedTime, float InertializationDuration) const;
+
+	// Get the inertialization space for this pose diff (for debug display)
+	//
+	EInertializationSpace GetInertializationSpace() const
+	{
+		return InertializationSpace;
+	}
+
+private:
+
+	static float CalcInertialFloat(float x0, float v0, float t, float t1);
+
+	// Bone differences indexed by skeleton bone index
+	TArray<FInertializationBoneDiff> BoneDiffs;
+
+	// Inertialization space (local vs world for situations where we wish to correct a world-space discontinuity such as an abrupt orientation change)
+	EInertializationSpace InertializationSpace;
+};
+
+
+USTRUCT(BlueprintInternalUseOnly)
+struct ENGINE_API FAnimNode_Inertialization : public FAnimNode_Base
+{
+	GENERATED_BODY()
+
+	UPROPERTY(EditAnywhere, BlueprintReadWrite, Category = Links)
+	FPoseLink Source;
+
+public: // FAnimNode_Inertialization
+
+	FAnimNode_Inertialization();
+	
+	// Request to activate inertialization for a duration.
+	// If multiple requests are made on the same inertialization node, the minimum requested time will be used.
+	//
+	virtual void Request(float Duration);
+
+	virtual float GetRequestedDuration() const { return RequestedDuration; }
+
+
+public: // FAnimNode_Base
+
+	virtual void Initialize_AnyThread(const FAnimationInitializeContext& Context) override;
+	virtual void CacheBones_AnyThread(const FAnimationCacheBonesContext& Context) override;
+	virtual void Update_AnyThread(const FAnimationUpdateContext& Context) override;
+	virtual void Evaluate_AnyThread(FPoseContext& Output) override;
+	virtual void GatherDebugData(FNodeDebugData& DebugData) override;
+
+	virtual bool NeedsDynamicReset() const override;
+	virtual void ResetDynamics(ETeleportType InTeleportType) override;
+
+	virtual bool WantsSkippedUpdates() const override;
+	virtual void OnUpdatesSkipped(TArrayView<const FAnimationUpdateContext *> SkippedUpdateContexts) override;
+
+
+protected:
+
+	// Consume Inertialization Request
+	//
+	// Returns any pending inertialization request and removes it from future processing.  Returns zero if there is no pending request.
+	// This function is virtual so that a derived class could optionally hook into other external sources of inertialization requests
+	// (for example from the owning actor for requests triggered from game code).
+	//
+	virtual float ConsumeInertializationRequest(FPoseContext& Context);
+
+	// Start Inertialization
+	//
+	// Computes the inertialization pose difference from the current pose and the two previous poses (to capture velocity).  This function
+	// is virtual so that a derived class could optionally regularize the pose snapshots to align better with the current frame's pose
+	// before computing the inertial difference (for example to correct for instantaneous changes in the root relative to its children).
+	//
+	virtual void StartInertialization(FPoseContext& Context, FInertializationPose& PreviousPose1, FInertializationPose& PreviousPose2, float Duration, /*OUT*/ FInertializationPoseDiff& OutPoseDiff);
+
+	// Apply Inertialization
+	//
+	// Applies the inertialization pose difference to the current pose (feathering down to zero as ElapsedTime approaches Duration).  This
+	// function is virtual so that a derived class could optionally adjust the pose based on any regularization done in StartInertialization.
+	//
+	virtual void ApplyInertialization(FPoseContext& Context, const FInertializationPoseDiff& PoseDiff, float ElapsedTime, float Duration);
+
+
+private:
+
+	// Snapshots of the actor pose from past frames
+	TArray<FInertializationPose> PoseSnapshots;
+
+	// Elapsed delta time between calls to evaluate
+	float DeltaTime;
+
+	// Pending inertialization request
+	float RequestedDuration;
+
+	// Teleport type
+	ETeleportType TeleportType;
+
+	// Inertialization state
+	EInertializationState InertializationState;
+	float InertializationElapsedTime;
+	float InertializationDuration;
+
+	// Inertialization pose differences
+	FInertializationPoseDiff InertializationPoseDiff;
+};
diff --git a/Engine/Source/Runtime/Engine/Classes/Animation/AnimNode_SaveCachedPose.h b/Engine/Source/Runtime/Engine/Classes/Animation/AnimNode_SaveCachedPose.h
--- a/Engine/Source/Runtime/Engine/Classes/Animation/AnimNode_SaveCachedPose.h
+++ b/Engine/Source/Runtime/Engine/Classes/Animation/AnimNode_SaveCachedPose.h
@@ -28,7 +28,13 @@ protected:
 	FCompactPose CachedPose;
 	FBlendedCurve CachedCurve;
 
-	TArray<FAnimationUpdateContext> CachedUpdateContexts;
+	struct FCachedUpdateContext
+	{
+		FAnimationUpdateContext Context;
+		FAnimationUpdateSharedContext SharedContext;
+	};
+
+	TArray<FCachedUpdateContext> CachedUpdateContexts;
 
 	FGraphTraversalCounter InitializationCounter;
 	FGraphTraversalCounter CachedBonesCounter;
diff --git a/Engine/Source/Runtime/Engine/Classes/Animation/AnimStateMachineTypes.h b/Engine/Source/Runtime/Engine/Classes/Animation/AnimStateMachineTypes.h
--- a/Engine/Source/Runtime/Engine/Classes/Animation/AnimStateMachineTypes.h
+++ b/Engine/Source/Runtime/Engine/Classes/Animation/AnimStateMachineTypes.h
@@ -30,6 +30,7 @@ namespace ETransitionLogicType
 	enum Type
 	{
 		TLT_StandardBlend UMETA(DisplayName="Standard Blend"),
+		TLT_Inertialization UMETA(DisplayName = "Inertialization"),
 		TLT_Custom UMETA(DisplayName="Custom")
 	};
 }
diff --git a/Engine/Source/Runtime/Engine/Classes/Engine/Blueprint.h b/Engine/Source/Runtime/Engine/Classes/Engine/Blueprint.h
--- a/Engine/Source/Runtime/Engine/Classes/Engine/Blueprint.h
+++ b/Engine/Source/Runtime/Engine/Classes/Engine/Blueprint.h
@@ -1003,6 +1003,14 @@ public:
 	*/
 	virtual void GetInstanceActions(FBlueprintActionDatabaseRegistrar& ActionRegistrar) const {}
 
+	/**
+	 * Allows child classes to override the undo behaviour on undo / redo.
+	 * Normally blueprints are marked as dirty / structurally modified on undo,
+	 * which might not be required for specialized classes.
+	 * @return True if the blueprint should be marked as structurally modified on undo redo.
+	 */
+	virtual bool RequiresMarkAsStructurallyModifiedOnUndo() const { return true; }
+
 #if WITH_EDITOR
 private:
 
@@ -1038,6 +1046,46 @@ protected:
 	 */
 	virtual void LoadModulesRequiredForCompilation();
 #endif
+
+#if WITH_EDITOR
+
+public:
+	/**
+	 * Returns true if this blueprint supports global variables
+	 */
+	virtual bool SupportsGlobalVariables() const { return true; }
+
+	/**
+	 * Returns true if this blueprint supports lobal variables
+	 */
+	virtual bool SupportsLocalVariables() const { return true; }
+
+	/**
+	 * Returns true if this blueprint supports functions
+	 */
+	virtual bool SupportsFunctions() const { return true; }
+
+	/**
+	 * Returns true if this blueprint supports macros
+	 */
+	virtual bool SupportsMacros() const { return true; }
+
+	/**
+	 * Returns true if this blueprint supports delegates
+	 */
+	virtual bool SupportsDelegates() const { return true; }
+
+	/**
+	 * Returns true if this blueprint supports event graphs
+	 */
+	virtual bool SupportsEventGraphs() const { return true; }
+
+	/**
+	 * Returns true if this blueprint supports animation layers
+	 */
+	virtual bool SupportsAnimLayers() const { return true; }
+
+#endif
 };
 
 
diff --git a/Engine/Source/Runtime/Engine/Private/Animation/AnimInstanceProxy.cpp b/Engine/Source/Runtime/Engine/Private/Animation/AnimInstanceProxy.cpp
--- a/Engine/Source/Runtime/Engine/Private/Animation/AnimInstanceProxy.cpp
+++ b/Engine/Source/Runtime/Engine/Private/Animation/AnimInstanceProxy.cpp
@@ -48,7 +48,9 @@ void FAnimInstanceProxy::UpdateAnimationNode_WithRoot(float DeltaSeconds, FAnimN
 			UpdateCounter.Increment();
 		}
 		
-		InRootNode->Update_AnyThread(FAnimationUpdateContext(this, DeltaSeconds));
+		FAnimationUpdateSharedContext SharedContext;
+		FAnimationUpdateContext UpdateContext(this, DeltaSeconds, &SharedContext);
+		InRootNode->Update_AnyThread(UpdateContext);
 
 		// We've updated the graph, now update the fractured saved pose sections
 		TArray<FAnimNode_SaveCachedPose*>& SavedPoseQueue = SavedPoseQueueMap.FindChecked(InLayerName);
diff --git a/Engine/Source/Runtime/Engine/Private/Animation/AnimNodeBase.cpp b/Engine/Source/Runtime/Engine/Private/Animation/AnimNodeBase.cpp
--- a/Engine/Source/Runtime/Engine/Private/Animation/AnimNodeBase.cpp
+++ b/Engine/Source/Runtime/Engine/Private/Animation/AnimNodeBase.cpp
@@ -55,6 +55,11 @@ const void* FExposedValueCopyRecord::GetSourceAddr(FAnimInstanceProxy* Proxy) co
 /////////////////////////////////////////////////////
 // FAnimationBaseContext
 
+FAnimationBaseContext::FAnimationBaseContext()
+	: AnimInstanceProxy(nullptr)
+{
+}
+
 FAnimationBaseContext::FAnimationBaseContext(FAnimInstanceProxy* InAnimInstanceProxy)
 	: AnimInstanceProxy(InAnimInstanceProxy)
 {
diff --git a/Engine/Source/Runtime/Engine/Private/Animation/AnimNode_SaveCachedPose.cpp b/Engine/Source/Runtime/Engine/Private/Animation/AnimNode_SaveCachedPose.cpp
--- a/Engine/Source/Runtime/Engine/Private/Animation/AnimNode_SaveCachedPose.cpp
+++ b/Engine/Source/Runtime/Engine/Private/Animation/AnimNode_SaveCachedPose.cpp
@@ -40,8 +40,16 @@ void FAnimNode_SaveCachedPose::CacheBones_AnyThread(const FAnimationCacheBonesCo
 
 void FAnimNode_SaveCachedPose::Update_AnyThread(const FAnimationUpdateContext& Context)
 {
+	FCachedUpdateContext& CachedUpdate = CachedUpdateContexts.AddDefaulted_GetRef();
+
+	// Make a minimal copy of the shared context for cached updates
+	if (FAnimationUpdateSharedContext* SharedContext = Context.GetSharedContext())
+	{
+		CachedUpdate.SharedContext.CopyForCachedUpdate(*SharedContext);
+	}
+
 	// Store this context for the post update
-	CachedUpdateContexts.Add(Context);
+	CachedUpdate.Context = Context.WithOtherSharedContext(&CachedUpdate.SharedContext);
 }
 
 void FAnimNode_SaveCachedPose::Evaluate_AnyThread(FPoseContext& Output)
@@ -79,21 +87,62 @@ void FAnimNode_SaveCachedPose::PostGraphUpdate()
 
 	// Update GlobalWeight based on highest weight calling us.
 	const int32 NumContexts = CachedUpdateContexts.Num();
-	if(NumContexts > 0)
+	if (NumContexts > 0)
 	{
-		GlobalWeight = CachedUpdateContexts[0].GetFinalBlendWeight();
+		GlobalWeight = CachedUpdateContexts[0].Context.GetFinalBlendWeight();
 		int32 MaxWeightIdx = 0;
-		for(int32 CurrIdx = 1; CurrIdx < NumContexts; ++CurrIdx)
+		for (int32 CurrIdx = 1; CurrIdx < NumContexts; ++CurrIdx)
 		{
-			const float BlendWeight = CachedUpdateContexts[CurrIdx].GetFinalBlendWeight();
-			if(BlendWeight > GlobalWeight)
+			const float BlendWeight = CachedUpdateContexts[CurrIdx].Context.GetFinalBlendWeight();
+			if (BlendWeight > GlobalWeight)
 			{
 				GlobalWeight = BlendWeight;
 				MaxWeightIdx = CurrIdx;
 			}
 		}
 
-		Pose.Update(CachedUpdateContexts[MaxWeightIdx]);
+		// Update the max weighted pose node
+		Pose.Update(CachedUpdateContexts[MaxWeightIdx].Context);
+
+		// Determine if any ancestors are interested in the other updates we'll be skipping
+		TArray<FAnimNode_Base*, TInlineAllocator<4>> AncestorsWithSkippedUpdateHandlers;
+		if (NumContexts > 1)
+		{
+			FAnimationUpdateSharedContext* SharedContext = CachedUpdateContexts[MaxWeightIdx].Context.GetSharedContext();
+			FAnimNodeTracker* AncestorTracker = SharedContext ? &SharedContext->AncestorTracker : nullptr;
+
+			if (AncestorTracker)
+			{
+				for (auto Iter : AncestorTracker->Map)
+				{
+					FAnimNode_Base* AncestorNode = Iter.Value.Top();
+					if (AncestorNode && AncestorNode->WantsSkippedUpdates())
+					{
+						AncestorsWithSkippedUpdateHandlers.Add(AncestorNode);
+					}
+				}
+			}
+		}
+
+		if (AncestorsWithSkippedUpdateHandlers.Num() > 0)
+		{
+			// Build a list of skipped updates
+			TArray<const FAnimationUpdateContext *, TInlineAllocator<4>> SkippedUpdateContexts;
+			for (int32 CurrIdx = 0; CurrIdx < NumContexts; ++CurrIdx)
+			{
+				if (CurrIdx != MaxWeightIdx)
+				{
+					SkippedUpdateContexts.Add(&CachedUpdateContexts[CurrIdx].Context);
+				}
+			}
+
+			// Inform any interested ancestors about the skipped updates
+			for (FAnimNode_Base* AncestorNode : AncestorsWithSkippedUpdateHandlers)
+			{
+				AncestorNode->OnUpdatesSkipped(SkippedUpdateContexts);
+			}
+		}
+
 	}
 
 	CachedUpdateContexts.Reset();
diff --git a/Engine/Source/Runtime/Engine/Private/Animation/AnimNode_StateMachine.cpp b/Engine/Source/Runtime/Engine/Private/Animation/AnimNode_StateMachine.cpp
--- a/Engine/Source/Runtime/Engine/Private/Animation/AnimNode_StateMachine.cpp
+++ b/Engine/Source/Runtime/Engine/Private/Animation/AnimNode_StateMachine.cpp
@@ -5,6 +5,7 @@
 #include "Animation/AnimNode_TransitionResult.h"
 #include "Animation/AnimNode_TransitionPoseEvaluator.h"
 #include "Animation/AnimNode_AssetPlayerBase.h"
+#include "Animation/AnimNode_Inertialization.h"
 #include "Animation/BlendProfile.h"
 #include "Animation/AnimNode_Layer.h"
 #include "Animation/AnimInstance.h"
@@ -44,7 +45,7 @@ FAnimationActiveTransitionEntry::FAnimationActiveTransitionEntry(int32 NextState
 	, bActive(true)
 {
 	const float Scaler = 1.0f - ExistingWeightOfNextState;
-	CrossfadeDuration = ReferenceTransitionInfo.CrossfadeDuration * CalculateInverseAlpha(BlendOption, Scaler);
+	CrossfadeDuration = (LogicType == ETransitionLogicType::TLT_Inertialization) ? 0.0f : ReferenceTransitionInfo.CrossfadeDuration * CalculateInverseAlpha(BlendOption, Scaler);
 
 	Blend.SetBlendTime(CrossfadeDuration);
 	Blend.SetBlendOption(BlendOption);
@@ -443,6 +444,19 @@ void FAnimNode_StateMachine::Update_AnyThread(const FAnimationUpdateContext& Con
 			{
 				NewTransition->InitializeCustomGraphLinks(Context, *(PotentialTransition.TransitionRule));
 
+				if (ReferenceTransition.LogicType == ETransitionLogicType::TLT_Inertialization)
+				{
+					FAnimNode_Inertialization* InertializationNode = Context.GetAncestor<FAnimNode_Inertialization>();
+					if (InertializationNode)
+					{
+						InertializationNode->Request(ReferenceTransition.CrossfadeDuration);
+					}
+					else
+					{
+						UE_LOG(LogAnimation, Error, TEXT("FAnimNode_StateMachine: No inertialization context found for transition"));
+					}
+				}
+
 #if WITH_EDITORONLY_DATA
 				NewTransition->SourceTransitionIndices = PotentialTransition.SourceTransitionIndices;
 #endif
@@ -686,6 +700,13 @@ void FAnimNode_StateMachine::UpdateTransitionStates(const FAnimationUpdateContex
 			}
 			break;
 
+		case ETransitionLogicType::TLT_Inertialization:
+			{
+				// update target state
+				UpdateState(Transition.NextState, Context);
+			}
+			break;
+
 		case ETransitionLogicType::TLT_Custom:
 			{
 				if (Transition.CustomTransitionGraph.LinkID != INDEX_NONE)
@@ -755,6 +776,9 @@ void FAnimNode_StateMachine::Evaluate_AnyThread(FPoseContext& Output)
 				case ETransitionLogicType::TLT_StandardBlend:
 					EvaluateTransitionStandardBlend(Output, ActiveTransition, bIntermediatePoseIsValid);
 					break;
+				case ETransitionLogicType::TLT_Inertialization:
+					EvaluateState(ActiveTransition.NextState, Output);
+					break;
 				case ETransitionLogicType::TLT_Custom:
 					EvaluateTransitionCustomBlend(Output, ActiveTransition, bIntermediatePoseIsValid);
 					break;
diff --git a/Engine/Source/Runtime/Engine/Private/SkeletalMesh.cpp b/Engine/Source/Runtime/Engine/Private/SkeletalMesh.cpp
--- a/Engine/Source/Runtime/Engine/Private/SkeletalMesh.cpp
+++ b/Engine/Source/Runtime/Engine/Private/SkeletalMesh.cpp
@@ -2524,7 +2524,7 @@ void USkeletalMesh::RebuildSocketMap()
 			USkeletalMeshSocket* Socket = Skeleton->Sockets[SocketIndex];
 			if (!SocketMap.Contains(Socket->SocketName))
 			{
-				SocketMap.Add(Socket->SocketName, FSocketInfo(this, Socket, SocketIndex));
+				SocketMap.Add(Socket->SocketName, FSocketInfo(this, Socket, Sockets.Num() + SocketIndex));
 			}
 		}
 	}
diff --git a/Engine/Source/Runtime/Engine/Private/UnrealClient.cpp b/Engine/Source/Runtime/Engine/Private/UnrealClient.cpp
--- a/Engine/Source/Runtime/Engine/Private/UnrealClient.cpp
+++ b/Engine/Source/Runtime/Engine/Private/UnrealClient.cpp
@@ -396,6 +396,7 @@ FString GetMemoryString( const double Value, const bool bAutoType )
 }
 
 FOnScreenshotRequestProcessed FScreenshotRequest::ScreenshotProcessedDelegate;
+FOnScreenshotCaptured FScreenshotRequest::ScreenshotCapturedDelegate;
 bool FScreenshotRequest::bIsScreenshotRequested = false;
 FString FScreenshotRequest::Filename;
 FString FScreenshotRequest::NextScreenshotName;
diff --git a/Engine/Source/Runtime/Engine/Public/UnrealClient.h b/Engine/Source/Runtime/Engine/Public/UnrealClient.h
--- a/Engine/Source/Runtime/Engine/Public/UnrealClient.h
+++ b/Engine/Source/Runtime/Engine/Public/UnrealClient.h
@@ -132,6 +132,8 @@ public:
 #define MAX_HITPROXYSIZE 200
 
 DECLARE_MULTICAST_DELEGATE(FOnScreenshotRequestProcessed);
+DECLARE_MULTICAST_DELEGATE_ThreeParams(FOnScreenshotCaptured, int32 /*Width*/, int32 /*Height*/, const TArray<FColor>& /*Colors*/);
+
 
 struct ENGINE_API FScreenshotRequest
 {
@@ -187,8 +189,14 @@ struct ENGINE_API FScreenshotRequest
 		return ScreenshotProcessedDelegate;
 	}
 
+	static FOnScreenshotCaptured& OnScreenshotCaptured()
+	{
+		return ScreenshotCapturedDelegate;
+	}
+
 private:
 	static FOnScreenshotRequestProcessed ScreenshotProcessedDelegate;
+	static FOnScreenshotCaptured ScreenshotCapturedDelegate;
 	static bool bIsScreenshotRequested;
 	static FString NextScreenshotName;
 	static FString Filename;
diff --git a/Engine/Source/Runtime/RigVM/Private/RigVMCore/RigVMByteCode.cpp b/Engine/Source/Runtime/RigVM/Private/RigVMCore/RigVMByteCode.cpp
new file mode 100644
--- /dev/null
+++ b/Engine/Source/Runtime/RigVM/Private/RigVMCore/RigVMByteCode.cpp
@@ -0,0 +1,244 @@
+// Copyright 1998-2019 Epic Games, Inc. All Rights Reserved.
+
+#include "RigVMCore/RigVMByteCode.h"
+
+FRigVMInstructionArray::FRigVMInstructionArray()
+{
+}
+
+FRigVMInstructionArray::FRigVMInstructionArray(const FRigVMByteCode& InByteCode)
+{
+	uint64 ByteIndex = 0;
+	while (ByteIndex < InByteCode.Num())
+	{
+		ERigVMOpCode OpCode = InByteCode.GetOpCodeAt(ByteIndex);
+		Instructions.Add(FRigVMInstruction(OpCode, ByteIndex));
+		ByteIndex += InByteCode.GetOpNumBytesAt(ByteIndex);
+	}
+}
+
+void FRigVMInstructionArray::Reset()
+{
+	Instructions.Reset();
+}
+
+FRigVMByteCode::FRigVMByteCode()
+{
+}
+
+void FRigVMByteCode::Reset()
+{
+	ByteCode.Reset();
+}
+
+uint64 FRigVMByteCode::Num() const
+{
+	return (uint64)ByteCode.Num();
+}
+
+uint64 FRigVMByteCode::GetOpNumBytesAt(uint64 InByteCodeIndex, bool bIncludeOperands) const
+{
+	ERigVMOpCode OpCode = GetOpCodeAt(InByteCodeIndex);
+	switch (OpCode)
+	{
+		case ERigVMOpCode::Execute_0_Operands:
+		case ERigVMOpCode::Execute_1_Operands:
+		case ERigVMOpCode::Execute_2_Operands:
+		case ERigVMOpCode::Execute_3_Operands:
+		case ERigVMOpCode::Execute_4_Operands:
+		case ERigVMOpCode::Execute_5_Operands:
+		case ERigVMOpCode::Execute_6_Operands:
+		case ERigVMOpCode::Execute_7_Operands:
+		case ERigVMOpCode::Execute_8_Operands:
+		case ERigVMOpCode::Execute_9_Operands:
+		case ERigVMOpCode::Execute_10_Operands:
+		case ERigVMOpCode::Execute_11_Operands:
+		case ERigVMOpCode::Execute_12_Operands:
+		case ERigVMOpCode::Execute_13_Operands:
+		case ERigVMOpCode::Execute_14_Operands:
+		case ERigVMOpCode::Execute_15_Operands:
+		case ERigVMOpCode::Execute_16_Operands:
+		case ERigVMOpCode::Execute_17_Operands:
+		case ERigVMOpCode::Execute_18_Operands:
+		case ERigVMOpCode::Execute_19_Operands:
+		case ERigVMOpCode::Execute_20_Operands:
+		case ERigVMOpCode::Execute_21_Operands:
+		case ERigVMOpCode::Execute_22_Operands:
+		case ERigVMOpCode::Execute_23_Operands:
+		case ERigVMOpCode::Execute_24_Operands:
+		case ERigVMOpCode::Execute_25_Operands:
+		case ERigVMOpCode::Execute_26_Operands:
+		case ERigVMOpCode::Execute_27_Operands:
+		case ERigVMOpCode::Execute_28_Operands:
+		case ERigVMOpCode::Execute_29_Operands:
+		case ERigVMOpCode::Execute_30_Operands:
+		case ERigVMOpCode::Execute_31_Operands:
+		case ERigVMOpCode::Execute_32_Operands:
+		case ERigVMOpCode::Execute_33_Operands:
+		case ERigVMOpCode::Execute_34_Operands:
+		case ERigVMOpCode::Execute_35_Operands:
+		case ERigVMOpCode::Execute_36_Operands:
+		case ERigVMOpCode::Execute_37_Operands:
+		case ERigVMOpCode::Execute_38_Operands:
+		case ERigVMOpCode::Execute_39_Operands:
+		case ERigVMOpCode::Execute_40_Operands:
+		case ERigVMOpCode::Execute_41_Operands:
+		case ERigVMOpCode::Execute_42_Operands:
+		case ERigVMOpCode::Execute_43_Operands:
+		case ERigVMOpCode::Execute_44_Operands:
+		case ERigVMOpCode::Execute_45_Operands:
+		case ERigVMOpCode::Execute_46_Operands:
+		case ERigVMOpCode::Execute_47_Operands:
+		case ERigVMOpCode::Execute_48_Operands:
+		case ERigVMOpCode::Execute_49_Operands:
+		case ERigVMOpCode::Execute_50_Operands:
+		case ERigVMOpCode::Execute_51_Operands:
+		case ERigVMOpCode::Execute_52_Operands:
+		case ERigVMOpCode::Execute_53_Operands:
+		case ERigVMOpCode::Execute_54_Operands:
+		case ERigVMOpCode::Execute_55_Operands:
+		case ERigVMOpCode::Execute_56_Operands:
+		case ERigVMOpCode::Execute_57_Operands:
+		case ERigVMOpCode::Execute_58_Operands:
+		case ERigVMOpCode::Execute_59_Operands:
+		case ERigVMOpCode::Execute_60_Operands:
+		case ERigVMOpCode::Execute_61_Operands:
+		case ERigVMOpCode::Execute_62_Operands:
+		case ERigVMOpCode::Execute_63_Operands:
+		case ERigVMOpCode::Execute_64_Operands:
+		{
+			uint64 NumBytes = (uint64)sizeof(FRigVMExecuteOp);
+			if(bIncludeOperands)
+			{
+				const FRigVMExecuteOp& ExecuteOp = GetOpAt<FRigVMExecuteOp>(InByteCodeIndex);
+				NumBytes += (uint64)ExecuteOp.GetOperandCount() * (uint64)sizeof(FRigVMOperand);
+			}
+			return NumBytes;
+		}
+		case ERigVMOpCode::Copy:
+		{
+			return (uint64)sizeof(FRigVMCopyOp);
+		}
+		case ERigVMOpCode::Zero:
+		case ERigVMOpCode::BoolFalse:
+		case ERigVMOpCode::BoolTrue:
+		case ERigVMOpCode::Increment:
+		case ERigVMOpCode::Decrement:
+		{
+			return (uint64)sizeof(FRigVMUnaryOp);
+		}
+		case ERigVMOpCode::Equals:
+		case ERigVMOpCode::NotEquals:
+		{
+			return (uint64)sizeof(FRigVMComparisonOp);
+		}
+		case ERigVMOpCode::JumpAbsolute:
+		case ERigVMOpCode::JumpForward:
+		case ERigVMOpCode::JumpBackward:
+		{
+			return (uint64)sizeof(FRigVMJumpOp);
+		}
+		case ERigVMOpCode::JumpAbsoluteIf:
+		case ERigVMOpCode::JumpForwardIf:
+		case ERigVMOpCode::JumpBackwardIf:
+		{
+			return (uint64)sizeof(FRigVMJumpIfOp);
+		}
+		case ERigVMOpCode::ChangeType:
+		{
+			return (uint64)sizeof(FRigVMChangeTypeOp);
+		}
+		case ERigVMOpCode::Exit:
+		{
+			return (uint64)sizeof(FRigVMBaseOp);
+		}
+		case ERigVMOpCode::Invalid:
+		{
+			ensure(false);
+			return 0;
+		}
+	}
+	return 0;
+}
+
+uint64 FRigVMByteCode::AddZeroOp(const FRigVMOperand& InArg)
+{
+	FRigVMUnaryOp Op(ERigVMOpCode::Zero, InArg);
+	return AddOp(Op);
+}
+
+uint64 FRigVMByteCode::AddFalseOp(const FRigVMOperand& InArg)
+{
+	FRigVMUnaryOp Op(ERigVMOpCode::BoolFalse, InArg);
+	return AddOp(Op);
+}
+
+uint64 FRigVMByteCode::AddTrueOp(const FRigVMOperand& InArg)
+{
+	FRigVMUnaryOp Op(ERigVMOpCode::BoolTrue, InArg);
+	return AddOp(Op);
+}
+
+uint64 FRigVMByteCode::AddCopyOp(const FRigVMOperand& InSource, const FRigVMOperand& InTarget)
+{
+	FRigVMCopyOp Op(InSource, InTarget);
+	return AddOp(Op);
+}
+
+uint64 FRigVMByteCode::AddIncrementOp(const FRigVMOperand& InArg)
+{
+	FRigVMUnaryOp Op(ERigVMOpCode::Increment, InArg);
+	return AddOp(Op);
+}
+
+uint64 FRigVMByteCode::AddDecrementOp(const FRigVMOperand& InArg)
+{
+	FRigVMUnaryOp Op(ERigVMOpCode::Decrement, InArg);
+	return AddOp(Op);
+}
+
+uint64 FRigVMByteCode::AddEqualsOp(const FRigVMOperand& InA, const FRigVMOperand& InB, const FRigVMOperand& InResult)
+{
+	FRigVMComparisonOp Op(ERigVMOpCode::Equals, InA, InB, InResult);
+	return AddOp(Op);
+}
+
+uint64 FRigVMByteCode::AddNotEqualsOp(const FRigVMOperand& InA, const FRigVMOperand& InB, const FRigVMOperand& InResult)
+{
+	FRigVMComparisonOp Op(ERigVMOpCode::NotEquals, InA, InB, InResult);
+	return AddOp(Op);
+}
+
+uint64 FRigVMByteCode::AddJumpOp(ERigVMOpCode InOpCode, uint16 InInstructionIndex)
+{
+	FRigVMJumpOp Op(InOpCode, InInstructionIndex);
+	return AddOp(Op);
+}
+
+uint64 FRigVMByteCode::AddJumpIfOp(ERigVMOpCode InOpCode, uint16 InInstructionIndex, const FRigVMOperand& InConditionArg, bool bJumpWhenConditionIs)
+{
+	FRigVMJumpIfOp Op(InOpCode, InConditionArg, InInstructionIndex, bJumpWhenConditionIs);
+	return AddOp(Op);
+}
+
+uint64 FRigVMByteCode::AddChangeTypeOp(FRigVMOperand InArg, ERigVMRegisterType InType, uint16 InElementSize, uint16 InElementCount, uint16 InSliceCount)
+{
+	FRigVMChangeTypeOp Op(InArg, InType, InElementSize, InElementCount, InSliceCount);
+	return AddOp(Op);
+}
+
+uint64 FRigVMByteCode::AddExecuteOp(uint16 InFunctionIndex, const TArrayView<FRigVMOperand>& InOperands)
+{
+	FRigVMExecuteOp Op(InFunctionIndex, (uint8)InOperands.Num());
+	uint64 OpByteIndex = AddOp(Op);
+	uint64 OperandsByteIndex = (uint64)ByteCode.AddUninitialized(sizeof(FRigVMOperand) * InOperands.Num());
+	FMemory::Memcpy(ByteCode.GetData() + OperandsByteIndex, InOperands.GetData(), sizeof(FRigVMOperand) * InOperands.Num());
+	return OpByteIndex;
+}
+
+uint64 FRigVMByteCode::AddExitOp()
+{
+	FRigVMBaseOp Op(ERigVMOpCode::Exit);
+	return AddOp(Op);
+}
+
diff --git a/Engine/Source/Runtime/RigVM/Private/RigVMCore/RigVMExecuteContext.cpp b/Engine/Source/Runtime/RigVM/Private/RigVMCore/RigVMExecuteContext.cpp
new file mode 100644
--- /dev/null
+++ b/Engine/Source/Runtime/RigVM/Private/RigVMCore/RigVMExecuteContext.cpp
@@ -0,0 +1,3 @@
+// Copyright 1998-2019 Epic Games, Inc. All Rights Reserved.
+
+#include "RigVMCore/RigVMExecuteContext.h"
diff --git a/Engine/Source/Runtime/RigVM/Private/RigVMCore/RigVMMemory.cpp b/Engine/Source/Runtime/RigVM/Private/RigVMCore/RigVMMemory.cpp
new file mode 100644
--- /dev/null
+++ b/Engine/Source/Runtime/RigVM/Private/RigVMCore/RigVMMemory.cpp
@@ -0,0 +1,1036 @@
+// Copyright 1998-2019 Epic Games, Inc. All Rights Reserved.
+
+#include "RigVMCore/RigVMMemory.h"
+#include "UObject/AnimObjectVersion.h"
+#include "UObject/PropertyPortFlags.h"
+#include "UObject/Package.h"
+
+bool FRigVMRegister::Serialize(FArchive& Ar)
+{
+	Ar.UsingCustomVersion(FAnimObjectVersion::GUID);
+
+	if (Ar.CustomVer(FAnimObjectVersion::GUID) < FAnimObjectVersion::StoreMarkerNamesOnSkeleton)
+	{
+		return false;
+	}
+
+	Ar << Type;
+	Ar << ByteIndex;
+	Ar << ElementSize;
+	Ar << ElementCount;
+	Ar << SliceIndex;
+	Ar << SliceCount;
+	Ar << AlignmentBytes;
+	Ar << TrailingBytes;
+	Ar << Name;
+	Ar << ScriptStructIndex;
+
+	return true;
+}
+
+bool FRigVMRegisterOffset::Serialize(FArchive& Ar)
+{
+	Ar.UsingCustomVersion(FAnimObjectVersion::GUID);
+
+	if (Ar.CustomVer(FAnimObjectVersion::GUID) < FAnimObjectVersion::StoreMarkerNamesOnSkeleton)
+	{
+		return false;
+	}
+
+	Ar << Segments;
+	Ar << Type;
+	Ar << CPPType;
+	Ar << ScriptStructPath;
+	Ar << ElementSize;
+
+	return true;
+}
+
+FRigVMRegisterOffset::FRigVMRegisterOffset(UScriptStruct* InScriptStruct, const FString& InSegmentPath, int32 InInitialOffset, uint16 InElementSize)
+	: Segments()
+	, Type(ERigVMRegisterType::Plain)
+	, CPPType()
+	, ScriptStruct(nullptr)
+	, ScriptStructPath()
+	, ElementSize(InElementSize)
+{
+	struct FRigVMRegisterOffsetBuilder
+	{
+		static void WalkStruct(UStruct* InStruct, const FString& InPath, FRigVMRegisterOffset& Offset)
+		{
+			FString Left, Right;
+			if (!InPath.Split(TEXT("."), &Left, &Right))
+			{
+				Left = InPath;
+			}
+
+			UProperty* Property = InStruct->FindPropertyByName(*Left);
+			check(Property)
+
+			int32 SegmentIndex = Property->GetOffset_ReplaceWith_ContainerPtrToValuePtr();
+			if (Offset.Segments.Num() > 0)
+			{
+				if (Offset.Segments.Last() >= 0)
+				{
+					Offset.Segments[Offset.Segments.Num() - 1] += SegmentIndex;
+				}
+				else
+				{
+					Offset.Segments.Add(SegmentIndex);
+				}
+			}
+			else
+			{
+				Offset.Segments.Add(SegmentIndex);
+			}
+
+			if (!Right.IsEmpty())
+			{
+				if (UStructProperty* StructProperty = Cast<UStructProperty>(Property))
+				{
+					WalkStruct(StructProperty->Struct, Right, Offset);
+				}
+				else if (UArrayProperty* ArrayProperty = Cast<UArrayProperty>(Property))
+				{
+					WalkArray(ArrayProperty, Right, Offset);
+				}
+			}
+			else
+			{
+				Offset.CPPType = *Property->GetCPPType();
+				Offset.ElementSize = Property->ElementSize;
+
+				if (UArrayProperty* ArrayProperty = Cast<UArrayProperty>(Property))
+				{
+					Offset.Segments.Add(-1);
+					Property = ArrayProperty->Inner;
+				}
+
+				if (UStructProperty* StructProperty = Cast<UStructProperty>(Property))
+				{
+					Offset.ScriptStruct = StructProperty->Struct;
+					Offset.Type = ERigVMRegisterType::Struct;
+				}
+				else
+				{
+					Offset.Type = ERigVMRegisterType::Plain;
+				}
+			}
+		}
+
+		static void WalkArray(UArrayProperty* InArrayProperty, const FString& InPath, FRigVMRegisterOffset& Offset)
+		{
+			FString Left, Right;
+			if (!InPath.Split(TEXT("."), &Left, &Right))
+			{
+				Left = InPath;
+			}
+
+			int32 ArrayIndex = FCString::Atoi(*Left);
+			int32 SegmentIndex = -1 - InArrayProperty->Inner->ElementSize * ArrayIndex;
+
+			if (Offset.Segments.Num() > 0)
+			{
+				if (Offset.Segments.Last() == 0)
+				{
+					Offset.Segments[Offset.Segments.Num() - 1] = SegmentIndex;
+				}
+				else
+				{
+					Offset.Segments.Add(SegmentIndex);
+				}
+			}
+			else
+			{
+				Offset.Segments.Add(SegmentIndex);
+			}
+
+			if (!Right.IsEmpty())
+			{
+				if (UStructProperty* StructProperty = Cast<UStructProperty>(InArrayProperty->Inner))
+				{
+					WalkStruct(StructProperty->Struct, Right, Offset);
+				}
+				else if (UArrayProperty* ArrayProperty = Cast<UArrayProperty>(InArrayProperty->Inner))
+				{
+					WalkArray(ArrayProperty, Right, Offset);
+				}
+			}
+			else
+			{
+				Offset.CPPType = *InArrayProperty->Inner->GetCPPType();
+				Offset.ElementSize = InArrayProperty->Inner->ElementSize;
+
+				if (UArrayProperty* ArrayProperty = Cast<UArrayProperty>(InArrayProperty->Inner))
+				{
+					Offset.Segments.Add(-1);
+					InArrayProperty = ArrayProperty;
+				}
+
+				if (UStructProperty* StructProperty = Cast<UStructProperty>(InArrayProperty->Inner))
+				{
+					Offset.ScriptStruct = StructProperty->Struct;
+					Offset.Type = ERigVMRegisterType::Struct;
+				}
+				else
+				{
+					Offset.Type = ERigVMRegisterType::Plain;
+				}
+			}
+
+		}
+	};
+
+	Segments.Add(InInitialOffset);
+
+	if (!InSegmentPath.IsEmpty() || InScriptStruct != nullptr)
+	{
+		ensure(!InSegmentPath.IsEmpty());
+		check(InScriptStruct)
+		FString SegmentPath = InSegmentPath;
+		SegmentPath = SegmentPath.Replace(TEXT("["), TEXT("."));
+		SegmentPath = SegmentPath.Replace(TEXT("]"), TEXT("."));
+		FRigVMRegisterOffsetBuilder::WalkStruct(InScriptStruct, SegmentPath, *this);
+		if (ScriptStruct != nullptr)
+		{
+			ScriptStructPath = *ScriptStruct->GetPathName();
+		}
+		if (Type == ERigVMRegisterType::Plain)
+		{
+			if (CPPType == TEXT("FName"))
+			{
+				Type = ERigVMRegisterType::Name;
+			}
+			else if (CPPType == TEXT("FString"))
+			{
+				Type = ERigVMRegisterType::String;
+			}
+		}
+	}
+
+	ensure(ElementSize > 0);
+}
+
+uint8* FRigVMRegisterOffset::GetData(uint8* InContainer) const
+{
+	uint8* Data = InContainer;
+	for(int32 SegmentIndex : Segments)
+	{
+		if (SegmentIndex < 0)
+		{
+			int32 ArrayOffset = (-SegmentIndex) - 1;
+			TArray<uint8>* ArrayPtr = (TArray<uint8>*)Data;
+			Data = ArrayPtr->GetData() + ArrayOffset;
+		}
+		else
+		{
+			Data = Data + SegmentIndex;
+		}
+	}
+	return Data;
+}
+
+bool FRigVMRegisterOffset::operator == (const FRigVMRegisterOffset& InOther) const
+{
+	if (Segments.Num() != InOther.Segments.Num())
+	{
+		return false;
+	}
+	if (ElementSize != InOther.ElementSize)
+	{
+		return false;
+	}
+	if (GetScriptStruct() != InOther.GetScriptStruct())
+	{
+		return false;
+	}
+	for (int32 Index = 0; Index < Segments.Num(); Index++)
+	{
+		if (Segments[Index] != InOther.Segments[Index])
+		{
+			return false;
+		}
+	}
+	return true;
+}
+
+UScriptStruct* FRigVMRegisterOffset::GetScriptStruct() const
+{
+	if (ScriptStruct == nullptr)
+	{
+		if (ScriptStructPath != NAME_None)
+		{
+			FRigVMRegisterOffset* MutableThis = (FRigVMRegisterOffset*)this;
+			MutableThis->ScriptStruct = FindObject<UScriptStruct>(ANY_PACKAGE, *ScriptStructPath.ToString());
+		}
+	}
+	return ScriptStruct;
+}
+
+uint16 FRigVMRegisterOffset::GetElementSize() const
+{
+	return ElementSize;
+}
+
+FRigVMMemoryContainer::FRigVMMemoryContainer(bool bInUseNames)
+	: bUseNameMap(bInUseNames)
+	, MemoryType(ERigVMMemoryType::Work)
+{
+}
+
+FRigVMMemoryContainer::FRigVMMemoryContainer(const FRigVMMemoryContainer& Other)
+{
+	*this = Other;
+}
+
+FRigVMMemoryContainer::~FRigVMMemoryContainer()
+{
+	Reset();
+}
+
+FRigVMMemoryContainer& FRigVMMemoryContainer::operator= (const FRigVMMemoryContainer &InOther)
+{
+	Reset();
+
+	bUseNameMap = InOther.bUseNameMap;
+	Data.Append(InOther.Data);
+	Registers.Append(InOther.Registers);
+	ScriptStructs.Append(InOther.ScriptStructs);
+
+	UpdateRegisters();
+
+	for (int32 Index = 0; Index < Registers.Num(); Index++)
+	{
+		Construct(Index);
+		Copy(Index, Index, &InOther);
+	}
+
+	return *this;
+}
+
+bool FRigVMMemoryContainer::Serialize(FArchive& Ar)
+{
+	Ar.UsingCustomVersion(FAnimObjectVersion::GUID);
+
+ 	if (Ar.CustomVer(FAnimObjectVersion::GUID) < FAnimObjectVersion::StoreMarkerNamesOnSkeleton)
+ 	{
+ 		return false;
+ 	}
+
+	if (Ar.IsLoading())
+	{
+		for (int32 RegisterIndex = 0; RegisterIndex < Registers.Num(); RegisterIndex++)
+		{
+			Registers[RegisterIndex].MoveToFirstSlice();
+			Destroy(RegisterIndex);
+		}
+	}
+	else
+	{
+		for (int32 RegisterIndex = 0; RegisterIndex < Registers.Num(); RegisterIndex++)
+		{
+			Registers[RegisterIndex].MoveToFirstSlice();
+		}
+	}
+
+	Ar << bUseNameMap;
+	Ar << MemoryType;
+	Ar << Registers;
+	Ar << RegisterOffsets;
+
+	if (Ar.IsLoading())
+	{
+		ScriptStructs.Reset();
+		TArray<FString> ScriptStructPaths;
+		Ar << ScriptStructPaths;
+		
+		for (const FString& ScriptStructPath : ScriptStructPaths)
+		{
+			UScriptStruct* ScriptStruct = FindObject<UScriptStruct>(nullptr, *ScriptStructPath);
+			ensure(ScriptStruct != nullptr);
+			ScriptStructs.Add(ScriptStruct);
+		}
+
+		uint64 TotalBytes = 0;
+		Ar << TotalBytes;
+		
+		Data.SetNumZeroed(TotalBytes);
+
+		for (int32 RegisterIndex = 0; RegisterIndex < Registers.Num(); RegisterIndex++)
+		{
+			Registers[RegisterIndex].MoveToFirstSlice();
+			Construct(RegisterIndex);
+		}
+
+		for (const FRigVMRegister& Register : Registers)
+		{
+			switch (Register.Type)
+			{
+				case ERigVMRegisterType::Plain:
+				{
+					TArray<uint8> View;
+					Ar << View;
+					ensure(View.Num() == Register.GetAllocatedBytes());
+					FMemory::Memcpy(&Data[Register.GetFirstAllocatedByte()], View.GetData(), View.Num());
+					break;
+				}
+				case ERigVMRegisterType::Name:
+				{
+					TArray<FName> View;
+					Ar << View;
+					ensure(View.Num() == Register.GetTotalElementCount());
+					for (uint16 ElementIndex = 0; ElementIndex < Register.GetTotalElementCount(); ElementIndex++)
+					{
+						*((FName*)&Data[Register.GetWorkByteIndex() + Register.ElementSize * ElementIndex]) = View[ElementIndex];
+					}
+					break;
+				}
+				case ERigVMRegisterType::String:
+				{
+					TArray<FString> View;
+					Ar << View;
+					ensure(View.Num() == Register.GetTotalElementCount());
+					for (uint16 ElementIndex = 0; ElementIndex < Register.GetTotalElementCount(); ElementIndex++)
+					{
+						*((FString*)&Data[Register.GetWorkByteIndex() + Register.ElementSize * ElementIndex]) = View[ElementIndex];
+					}
+					break;
+				}
+				case ERigVMRegisterType::Struct:
+				{
+					TArray<FString> View;
+					Ar << View;
+					ensure(View.Num() == Register.GetTotalElementCount());
+
+					uint8* DataPtr = &Data[Register.GetWorkByteIndex()];
+					UScriptStruct* ScriptStruct = ScriptStructs[Register.ScriptStructIndex];
+					for (uint16 ElementIndex = 0; ElementIndex < Register.GetTotalElementCount(); ElementIndex++)
+					{
+						ScriptStruct->ImportText(*View[ElementIndex], DataPtr, nullptr, PPF_None, nullptr, ScriptStruct->GetName());
+						DataPtr += Register.ElementSize;
+					}
+					break;
+				}
+			}
+		}
+
+		UpdateRegisters();
+	}
+	else
+	{
+		TArray<FString> ScriptStructPaths;
+		for (UScriptStruct* ScriptStruct : ScriptStructs)
+		{
+			ScriptStructPaths.Add(ScriptStruct->GetPathName());
+		}
+		Ar << ScriptStructPaths;
+
+		uint64 TotalBytes = Data.Num();
+		Ar << TotalBytes;
+
+		for (FRigVMRegister& Register : Registers)
+		{
+			Register.MoveToFirstSlice();
+
+			switch(Register.Type)
+			{
+				case ERigVMRegisterType::Plain:
+				{
+					TArray<uint8> View;
+					View.Append(&Data[Register.GetFirstAllocatedByte()], Register.GetAllocatedBytes());
+					Ar << View;
+					break;
+				}
+				case ERigVMRegisterType::Name:
+				{
+					TArray<FName> View;
+					View.Append((FName*)&Data[Register.GetWorkByteIndex()], Register.GetTotalElementCount());
+					Ar << View;
+					break;
+				}
+				case ERigVMRegisterType::String:
+				{
+					TArray<FString> View;
+					View.Append((FString*)&Data[Register.GetWorkByteIndex()], Register.GetTotalElementCount());
+					Ar << View;
+					break;
+				}
+				case ERigVMRegisterType::Struct:
+				{
+					uint8* DataPtr = &Data[Register.GetWorkByteIndex()];
+					UScriptStruct* ScriptStruct = ScriptStructs[Register.ScriptStructIndex];
+
+					TArray<FString> View;
+					for (uint16 ElementIndex = 0; ElementIndex < Register.GetTotalElementCount(); ElementIndex++)
+					{
+						FString Value;
+						ScriptStruct->ExportText(Value, DataPtr, nullptr, nullptr, PPF_None, nullptr);
+						View.Add(Value);
+						DataPtr += Register.ElementSize;
+					}
+
+					Ar << View;
+					break;
+				}
+			}
+		}
+	}
+
+	return true;
+}
+
+void FRigVMMemoryContainer::Reset()
+{
+	for (int32 Index = 0; Index < Registers.Num(); Index++)
+	{
+		Destroy(Index);
+	}
+
+	Data.Reset();
+	Registers.Reset();
+	ScriptStructs.Reset();
+	NameMap.Reset();
+}
+
+bool FRigVMMemoryContainer::Copy(
+	int32 InSourceRegisterIndex,
+	int32 InTargetRegisterIndex,
+	const FRigVMMemoryContainer* InSourceMemory,
+	int32 InSourceRegisterOffset,
+	int32 InTargetRegisterOffset)
+{
+	if (InSourceMemory == nullptr)
+	{
+		InSourceMemory = this;
+	}
+
+	ensure(InSourceMemory->Registers.IsValidIndex(InSourceRegisterIndex));
+	ensure(Registers.IsValidIndex(InTargetRegisterIndex));
+
+	if (InSourceRegisterIndex == InTargetRegisterIndex && InSourceRegisterOffset == InTargetRegisterOffset && this == InSourceMemory)
+	{
+		return false;
+	}
+
+	const FRigVMRegister& Source = InSourceMemory->Registers[InSourceRegisterIndex];
+	const FRigVMRegister& Target = Registers[InTargetRegisterIndex];
+
+	const uint8* SourcePtr = InSourceMemory->GetData(InSourceRegisterIndex, InSourceRegisterOffset);
+	uint8* TargetPtr = GetData(InTargetRegisterIndex, InTargetRegisterOffset);
+	uint16 NumBytes = Target.GetNumBytesPerSlice();
+
+	ERigVMRegisterType TargetType = Target.Type;
+	if (InTargetRegisterOffset != INDEX_NONE)
+	{
+		TargetType = RegisterOffsets[InTargetRegisterOffset].GetType();
+		NumBytes = RegisterOffsets[InTargetRegisterOffset].GetElementSize();
+	}
+
+	switch (TargetType)
+	{
+		case ERigVMRegisterType::Plain:
+		{
+			FMemory::Memcpy(TargetPtr, SourcePtr, NumBytes);
+			break;
+		}
+		case ERigVMRegisterType::Struct:
+		{
+			UScriptStruct* ScriptStruct = GetScriptStruct(InTargetRegisterIndex, InTargetRegisterOffset);
+			int32 NumStructs = NumBytes / ScriptStruct->GetStructureSize();
+			ScriptStruct->CopyScriptStruct(TargetPtr, SourcePtr, NumStructs);
+			break;
+		}
+		case ERigVMRegisterType::Name:
+		{
+			int32 NumNames = NumBytes / sizeof(FName);
+			TArrayView<FName> TargetNames((FName*)TargetPtr, NumNames);
+			TArrayView<FName> SourceNames((FName*)SourcePtr, NumNames);
+			for (int32 Index = 0; Index < NumNames; Index++)
+			{
+				TargetNames[Index] = SourceNames[Index];
+			}
+			break;
+		}
+		case ERigVMRegisterType::String:
+		{
+			int32 NumStrings = NumBytes / sizeof(FString);
+			TArrayView<FString> TargetStrings((FString*)TargetPtr, NumStrings);
+			TArrayView<FString> SourceStrings((FString*)SourcePtr, NumStrings);
+			for (int32 Index = 0; Index < NumStrings; Index++)
+			{
+				TargetStrings[Index] = SourceStrings[Index];
+			}
+			break;
+		}
+		case ERigVMRegisterType::Invalid:
+		{
+			return false;
+		}
+	}
+
+	return true;
+}
+
+bool FRigVMMemoryContainer::Copy(
+	const FName& InSourceName,
+	const FName& InTargetName,
+	const FRigVMMemoryContainer* InSourceMemory,
+	int32 InSourceRegisterOffset,
+	int32 InTargetRegisterOffset)
+{
+	ensure(bUseNameMap);
+
+	int32 SourceRegisterIndex = GetIndex(InSourceName);
+	int32 TargetRegisterIndex = GetIndex(InTargetName);
+
+	if(SourceRegisterIndex == INDEX_NONE || TargetRegisterIndex == INDEX_NONE)
+	{
+		return false;
+	}
+
+	return Copy(SourceRegisterIndex, TargetRegisterIndex, InSourceMemory, InSourceRegisterOffset, InTargetRegisterOffset);
+}
+
+bool FRigVMMemoryContainer::Copy(
+	const FRigVMOperand& InSourceOperand,
+	const FRigVMOperand& InTargetOperand,
+	const FRigVMMemoryContainer* InSourceMemory)
+{
+	return Copy(InSourceOperand.GetRegisterIndex(), InTargetOperand.GetRegisterIndex(), InSourceMemory, InSourceOperand.GetRegisterOffset(), InTargetOperand.GetRegisterOffset());
+}
+
+int32 FRigVMMemoryContainer::Allocate(const FName& InNewName, int32 InElementSize, int32 InElementCount, int32 InSliceCount, const uint8* InDataPtr, bool bUpdateRegisters)
+{
+	FName Name = InNewName;
+	if (bUseNameMap && InNewName == NAME_None)
+	{
+		const TCHAR* RegisterPrefix = TEXT("Register");
+		int32 RegisterSuffix = 0;
+		do
+		{
+			Name = FName(*FString::Printf(TEXT("%s_%d"), RegisterPrefix, RegisterSuffix++));
+		} while (!IsNameAvailable(Name));
+	}
+
+	ensure(InElementSize > 0 && InElementCount > 0 && InSliceCount > 0);
+
+	if (bUseNameMap)
+	{
+		if (!IsNameAvailable(Name))
+		{
+			return INDEX_NONE;
+		}
+	}
+
+	FRigVMRegister NewRegister;
+	NewRegister.ByteIndex = Data.Num();
+	if (bUseNameMap)
+	{
+		NewRegister.Name = Name;
+	}
+	NewRegister.ElementSize = (uint16)InElementSize;
+	NewRegister.ElementCount = (uint16)InElementCount;
+	NewRegister.SliceCount = (uint16)InSliceCount;
+	NewRegister.Type = ERigVMRegisterType::Plain;
+
+	Data.AddZeroed(NewRegister.GetAllocatedBytes());
+
+	if (InDataPtr != nullptr)
+	{
+		NewRegister.MoveToFirstSlice();
+		for (uint16 SliceIndex = 0; SliceIndex < NewRegister.SliceCount; SliceIndex++)
+		{
+			FMemory::Memcpy(&Data[NewRegister.GetWorkByteIndex()], InDataPtr, NewRegister.GetNumBytesPerSlice());
+			NewRegister.MoveToNextSlice();
+		}
+		NewRegister.MoveToFirstSlice();
+	}
+
+	int32 RegisterIndex = Registers.Num();
+	Registers.Add(NewRegister);
+
+	if (bUpdateRegisters)
+	{
+		UpdateRegisters();
+	}
+	return RegisterIndex;
+}
+
+int32 FRigVMMemoryContainer::Allocate(int32 InElementSize, int32 InElementCount, int32 InSliceCount, const uint8* InDataPtr, bool bUpdateRegisters)
+{
+	return Allocate(NAME_None, InElementSize, InElementCount, InSliceCount, InDataPtr, bUpdateRegisters);
+}
+
+bool FRigVMMemoryContainer::Construct(int32 InRegisterIndex, int32 InElementIndex)
+{
+	ensure(Registers.IsValidIndex(InRegisterIndex));
+
+	const FRigVMRegister& Register = Registers[InRegisterIndex];
+	switch (Register.Type)
+	{
+		case ERigVMRegisterType::Struct:
+		{
+			uint8* DataPtr = &Data[InElementIndex == INDEX_NONE ? Register.GetWorkByteIndex() : Register.GetWorkByteIndex() + InElementIndex * Register.ElementSize];
+			int32 Count = InElementIndex == INDEX_NONE ? Register.GetTotalElementCount() : 1;
+
+			UScriptStruct* ScriptStruct = GetScriptStruct(InRegisterIndex);
+			ScriptStruct->InitializeStruct(DataPtr, Count);
+			break;
+		}
+		case ERigVMRegisterType::String:
+		{
+			FString* DataPtr = (FString*)&Data[InElementIndex == INDEX_NONE ? Register.GetWorkByteIndex() : Register.GetWorkByteIndex() + InElementIndex * Register.ElementSize];
+			int32 Count = InElementIndex == INDEX_NONE ? Register.GetTotalElementCount() : 1;
+
+			FMemory::Memzero(DataPtr, Count * Register.ElementSize);
+			for (int32 Index = 0; Index < Count; Index++)
+			{
+				DataPtr[Index] = FString();
+			}
+			break;
+		}
+		case ERigVMRegisterType::Name:
+		{
+			FName* DataPtr = (FName*)&Data[InElementIndex == INDEX_NONE ? Register.GetWorkByteIndex() : Register.GetWorkByteIndex() + InElementIndex * Register.ElementSize];
+			int32 Count = InElementIndex == INDEX_NONE ? Register.GetTotalElementCount() : 1;
+
+			FMemory::Memzero(DataPtr, Count * Register.ElementSize);
+			for (int32 Index = 0; Index < Count; Index++)
+			{
+				DataPtr[Index] = FName();
+			}
+			break;
+		}
+		default:
+		{
+			return false;
+		}
+	}
+
+	return true;
+}
+
+bool FRigVMMemoryContainer::Destroy(int32 InRegisterIndex, int32 InElementIndex)
+{
+	ensure(Registers.IsValidIndex(InRegisterIndex));
+
+	FRigVMRegister& Register = Registers[InRegisterIndex];
+
+	if (InElementIndex == INDEX_NONE)
+	{
+		Register.MoveToFirstSlice();
+	}
+
+	switch (Register.Type)
+	{
+		case ERigVMRegisterType::Struct:
+		{
+			uint8* DataPtr = &Data[InElementIndex == INDEX_NONE ? Register.GetWorkByteIndex() : Register.GetWorkByteIndex() + InElementIndex * Register.ElementSize];
+			int32 Count = InElementIndex == INDEX_NONE ? Register.GetTotalElementCount() : 1;
+
+			UScriptStruct* ScriptStruct = GetScriptStruct(InRegisterIndex);
+			ScriptStruct->DestroyStruct(DataPtr, Count);
+			break;
+		}
+		case ERigVMRegisterType::String:
+		{
+			FString* DataPtr = (FString*)&Data[InElementIndex == INDEX_NONE ? Register.GetWorkByteIndex() : Register.GetWorkByteIndex() + InElementIndex * Register.ElementSize];
+			int32 Count = InElementIndex == INDEX_NONE ? Register.GetTotalElementCount() : 1;
+
+			for (int32 Index = 0; Index < Count; Index++)
+			{
+				DataPtr[Index] = FString();
+			}
+			break;
+		}
+		case ERigVMRegisterType::Name:
+		{
+			FName* DataPtr = (FName*)&Data[InElementIndex == INDEX_NONE ? Register.GetWorkByteIndex() : Register.GetWorkByteIndex() + InElementIndex * Register.ElementSize];
+			int32 Count = InElementIndex == INDEX_NONE ? Register.GetTotalElementCount() : 1;
+
+			for (int32 Index = 0; Index < Count; Index++)
+			{
+				DataPtr[Index] = FName();
+			}
+			break;
+		}
+		default:
+		{
+			return false;
+		}
+	}
+
+	return true;
+}
+
+bool FRigVMMemoryContainer::Remove(int32 InRegisterIndex)
+{
+	if (InRegisterIndex < 0 || InRegisterIndex >= Registers.Num())
+	{
+		return false;
+	}
+
+	Destroy(InRegisterIndex);
+
+	FRigVMRegister RegisterToRemove = Registers[InRegisterIndex];
+	Data.RemoveAt(RegisterToRemove.ByteIndex, RegisterToRemove.GetAllocatedBytes());
+	Registers.RemoveAt(InRegisterIndex);
+
+	for (int32 Index = InRegisterIndex; Index < Registers.Num(); Index++)
+	{
+		Registers[Index].ByteIndex -= RegisterToRemove.GetAllocatedBytes();
+	}
+
+	UpdateRegisters();
+	return true;
+}
+
+bool FRigVMMemoryContainer::Remove(const FName& InRegisterName)
+{
+	ensure(bUseNameMap);
+	return Remove(GetIndex(InRegisterName));
+}
+
+FName FRigVMMemoryContainer::Rename(int32 InRegisterIndex, const FName& InNewName)
+{
+	if (Registers[InRegisterIndex].Name == InNewName)
+	{
+		return Registers[InRegisterIndex].Name;
+	}
+
+	if (!IsNameAvailable(InNewName))
+	{
+		return Registers[InRegisterIndex].Name;
+	}
+
+	Registers[InRegisterIndex].Name = InNewName;
+	UpdateRegisters();
+
+	return InNewName;
+}
+
+FName FRigVMMemoryContainer::Rename(const FName& InOldName, const FName& InNewName)
+{
+	ensure(bUseNameMap);
+
+	int32 RegisterIndex = GetIndex(InOldName);
+	if (RegisterIndex == INDEX_NONE)
+	{
+		return NAME_None;
+	}
+
+	return Rename(RegisterIndex, InNewName);
+}
+
+bool FRigVMMemoryContainer::Resize(int32 InRegisterIndex, int32 InNewElementCount, int32 InNewSliceCount)
+{
+	ensure(Registers.IsValidIndex(InRegisterIndex));
+	ensure(Registers[InRegisterIndex].TrailingBytes == 0);
+
+	if (InNewElementCount <= 0 || InNewSliceCount < 0)
+	{
+		return Remove(InRegisterIndex);
+	}
+
+	uint32 NewTotalCount = (uint32)InNewElementCount * (uint32)InNewSliceCount;
+	if (Registers[InRegisterIndex].GetTotalElementCount() == NewTotalCount)
+	{
+		return false;
+	}
+
+	FRigVMRegister& Register = Registers[InRegisterIndex];
+	Register.MoveToFirstSlice();
+
+	if (Register.GetTotalElementCount() > NewTotalCount) // shrink
+	{
+		int32 ElementsToRemove = Register.GetTotalElementCount() - NewTotalCount;
+		int32 NumBytesToRemove = Register.ElementSize * ElementsToRemove;
+		int32 FirstByteToRemove = Register.GetWorkByteIndex() + Register.ElementSize * NewTotalCount;
+
+		for (uint32 ElementIndex = NewTotalCount; ElementIndex < Register.GetTotalElementCount(); ElementIndex++)
+		{
+			Destroy(InRegisterIndex, (uint32)ElementIndex);
+		}
+
+		Data.RemoveAt(FirstByteToRemove, NumBytesToRemove);
+		Register.ElementCount = (uint32)InNewElementCount;
+		Register.SliceCount = (uint32)InNewSliceCount;
+
+		for (int32 RegisterIndex = InRegisterIndex + 1; RegisterIndex < Registers.Num(); RegisterIndex++)
+		{
+			Registers[RegisterIndex].ByteIndex -= NumBytesToRemove;
+		}
+	}
+	else // grow
+	{
+		int32 OldElementCount = Register.GetTotalElementCount();
+		int32 ElementsToAdd = NewTotalCount - Register.GetTotalElementCount();
+		int32 NumBytesToAdd = Register.ElementSize * ElementsToAdd;
+		int32 FirstByteToAdd = Register.GetWorkByteIndex() + Register.ElementSize * Register.GetTotalElementCount();
+
+		Data.InsertZeroed(FirstByteToAdd, NumBytesToAdd);
+		Register.ElementCount = (uint32)InNewElementCount;
+		Register.SliceCount = (uint32)InNewSliceCount;
+
+		for (uint32 ElementIndex = OldElementCount; ElementIndex < NewTotalCount; ElementIndex++)
+		{
+			Construct(InRegisterIndex, (int32)ElementIndex);
+		}
+
+		for (int32 RegisterIndex = InRegisterIndex + 1; RegisterIndex < Registers.Num(); RegisterIndex++)
+		{
+			Registers[RegisterIndex].ByteIndex += NumBytesToAdd;
+		}
+	}
+
+	UpdateRegisters();
+	return true;
+}
+
+bool FRigVMMemoryContainer::Resize(const FName& InRegisterName, int32 InNewElementCount, int32 InNewSliceCount)
+{
+	ensure(bUseNameMap);
+
+	int32 RegisterIndex = GetIndex(InRegisterName);
+	if (RegisterIndex == INDEX_NONE)
+	{
+		return false;
+	}
+
+	return Resize(RegisterIndex, InNewElementCount, InNewSliceCount);
+}
+
+bool FRigVMMemoryContainer::ChangeRegisterType(int32 InRegisterIndex, ERigVMRegisterType InNewType, int32 InElementSize, const uint8* InDataPtr, int32 InNewElementCount, int32 InNewSliceCount)
+{
+	ensure(Registers.IsValidIndex(InRegisterIndex));
+	
+	FRigVMRegister& Register = Registers[InRegisterIndex];
+	ensure(Register.AlignmentBytes == 0);
+	ensure(InNewType == ERigVMRegisterType::Name || InNewType == ERigVMRegisterType::String || InNewType == ERigVMRegisterType::Plain);
+
+	Register.MoveToFirstSlice();
+
+	if (Register.Type == InNewType && Register.ElementSize == InElementSize && Register.ElementCount == InNewElementCount && Register.SliceCount == InNewSliceCount)
+	{
+		return false;
+	}
+
+	Destroy(InRegisterIndex);
+
+	uint16 OldAllocatedBytes = Register.GetAllocatedBytes();
+	uint16 NewAllocatedBytes = (uint16)InElementSize * (uint16)InNewElementCount * (uint16)InNewSliceCount;
+	ensure(NewAllocatedBytes <= OldAllocatedBytes);
+
+	Register.Type = InNewType;
+	Register.ElementSize = (uint16)InElementSize;
+	Register.ElementCount = (uint16)InNewElementCount;
+	Register.SliceCount = (uint16)InNewSliceCount;
+	Register.TrailingBytes = OldAllocatedBytes - NewAllocatedBytes;
+
+	Construct(InRegisterIndex);
+
+	if (InDataPtr != nullptr)
+	{
+		for (int32 SliceIndex = 0; SliceIndex < Register.SliceCount; SliceIndex++)
+		{
+			FMemory::Memcpy(GetData(InRegisterIndex), InDataPtr, Register.GetNumBytesPerSlice());
+			Register.MoveToNextSlice();
+		}
+		Register.MoveToFirstSlice();
+	}
+
+	return true;
+}
+
+int32 FRigVMMemoryContainer::GetOrAddRegisterOffset(int32 InRegisterIndex, UScriptStruct* InScriptStruct, const FString& InSegmentPath, int32 InInitialOffset, int32 InElementSize)
+{
+	if ((InScriptStruct == nullptr || InSegmentPath.IsEmpty()) && InInitialOffset == 0)
+	{
+		return INDEX_NONE;
+	}
+
+	ensure(Registers.IsValidIndex(InRegisterIndex));
+
+	if (InElementSize == 0)
+	{
+		InElementSize = Registers[InRegisterIndex].GetNumBytesPerSlice();
+	}
+
+	FRigVMRegisterOffset Offset(InScriptStruct, InSegmentPath, InInitialOffset, InElementSize);
+	int32 ExistingIndex = RegisterOffsets.Find(Offset);
+	if (ExistingIndex == INDEX_NONE)
+	{
+		return RegisterOffsets.Add(Offset);
+	}
+	return ExistingIndex;
+}
+
+void FRigVMMemoryContainer::UpdateRegisters()
+{
+	int32 AlignmentShift = 0;
+	for (int32 RegisterIndex = 0; RegisterIndex < Registers.Num(); RegisterIndex++)
+	{
+		FRigVMRegister& Register = Registers[RegisterIndex];
+		Register.ByteIndex += AlignmentShift;
+		Register.MoveToFirstSlice();
+
+		UScriptStruct* ScriptStruct = GetScriptStruct(RegisterIndex);
+		if (ScriptStruct != nullptr)
+		{
+			UScriptStruct::ICppStructOps* TheCppStructOps = ScriptStruct->GetCppStructOps();
+			if (TheCppStructOps != NULL)
+			{
+				if (!TheCppStructOps->HasZeroConstructor())
+				{
+					uint8* Pointer = GetData(RegisterIndex);
+
+					if (Register.AlignmentBytes > 0)
+					{
+						if (!IsAligned(Pointer, TheCppStructOps->GetAlignment()))
+						{
+							Data.RemoveAt(Register.GetFirstAllocatedByte(), Register.AlignmentBytes);
+							AlignmentShift -= Register.AlignmentBytes;
+							Register.ByteIndex -= Register.AlignmentBytes;
+							Register.AlignmentBytes = 0;
+							Pointer = GetData(RegisterIndex);
+						}
+					}
+
+					while (!IsAligned(Pointer, TheCppStructOps->GetAlignment()))
+					{
+						Data.InsertZeroed(Register.GetFirstAllocatedByte(), 1);
+						Register.AlignmentBytes++;
+						Register.ByteIndex++;
+						AlignmentShift++;
+						Pointer = GetData(RegisterIndex);
+					}
+				}
+			}
+		}
+	}
+
+	if (bUseNameMap)
+	{
+		NameMap.Reset();
+		for (int32 Index = 0; Index < Registers.Num(); Index++)
+		{
+			NameMap.Add(Registers[Index].Name, Index);
+		}
+	}
+}
+
+void FRigVMMemoryContainer::FillWithZeroes(int32 InRegisterIndex)
+{
+	ensure(Registers.IsValidIndex(InRegisterIndex));
+	FMemory::Memzero(GetData(InRegisterIndex), Registers[InRegisterIndex].GetNumBytesAllSlices());
+}
+
+int32 FRigVMMemoryContainer::FindOrAddScriptStruct(UScriptStruct* InScriptStruct)
+{
+	int32 StructIndex = INDEX_NONE;
+	if (ScriptStructs.Find(InScriptStruct, StructIndex))
+	{
+		return StructIndex;
+	}
+	return ScriptStructs.Add(InScriptStruct);
+}
diff --git a/Engine/Source/Runtime/RigVM/Public/RigVMCore/RigVM.h b/Engine/Source/Runtime/RigVM/Public/RigVMCore/RigVM.h
new file mode 100644
--- /dev/null
+++ b/Engine/Source/Runtime/RigVM/Public/RigVMCore/RigVM.h
@@ -0,0 +1,405 @@
+// Copyright 1998-2019 Epic Games, Inc. All Rights Reserved.
+
+#pragma once
+
+#include "CoreMinimal.h"
+#include "RigVMMemory.h"
+#include "RigVMRegistry.h"
+#include "RigVMByteCode.h"
+#include "RigVM.generated.h"
+
+// The type of parameter for a VM
+UENUM(BlueprintType)
+enum class ERigVMParameterType : uint8
+{
+	Input,
+	Output,
+	Invalid
+};
+
+/**
+ * The RigVMParameter define an input or output of the RigVM.
+ * Parameters are mapped to work state memory registers and can be
+ * used to set input parameters as well as retrieve output parameters.
+ */
+USTRUCT(BlueprintType)
+struct FRigVMParameter
+{
+	GENERATED_BODY()
+
+public:
+
+	FRigVMParameter()
+		: Type(ERigVMParameterType::Invalid)
+		, Name(NAME_None)
+		, RegisterIndex(INDEX_NONE)
+		, CPPType()
+		, ScriptStruct(nullptr)
+		, ScriptStructPath()
+	{
+	}
+
+	// returns true if the parameter is valid
+	bool IsValid() const { return Type != ERigVMParameterType::Invalid; }
+
+	// returns the type of this parameter
+	ERigVMParameterType GetType() const { return Type; }
+
+	// returns the name of this parameters
+	const FName& GetName() const { return Name; }
+
+	// returns the register index of this parameter in the work memory
+	int32 GetRegisterIndex() const { return RegisterIndex; }
+
+	// returns the cpp type of the parameter
+	FString GetCPPType() const { return CPPType; }
+	
+	// Returns the script struct used by this parameter (in case it is a struct)
+	UScriptStruct* GetScriptStruct() const;
+
+private:
+
+	FRigVMParameter(ERigVMParameterType InType, const FName& InName, int32 InRegisterIndex, const FString& InCPPType, UScriptStruct* InScriptStruct)
+		: Type(InType)
+		, Name(InName)
+		, RegisterIndex(InRegisterIndex)
+		, CPPType(InCPPType)
+		, ScriptStruct(InScriptStruct)
+		, ScriptStructPath(NAME_None)
+	{
+		if (ScriptStruct)
+		{
+			ScriptStructPath = *ScriptStruct->GetPathName();
+		}
+	}
+
+	UPROPERTY()
+	ERigVMParameterType Type;
+
+	UPROPERTY()
+	FName Name;
+
+	UPROPERTY()
+	int32 RegisterIndex;
+
+	UPROPERTY()
+	FString CPPType;
+
+	UPROPERTY(transient)
+	UScriptStruct* ScriptStruct;
+
+	UPROPERTY()
+	FName ScriptStructPath;
+
+	friend class URigVM;
+	friend class URigVMCompiler;
+};
+
+/**
+ * The RigVM is the main object for evaluating FRigVMByteCode instructions.
+ * It combines the byte code, a list of required function pointers for 
+ * execute instructions and required memory in one class.
+ */
+UCLASS(BlueprintType)
+class RIGVM_API URigVM : public UObject
+{
+	GENERATED_BODY()
+
+public:
+
+	URigVM();
+	virtual ~URigVM();
+
+	// resets the container and removes all memory
+	void Reset();
+
+	// Executes the VM.
+	// You can optionally provide external memory to the execution
+	// and provide optional additional operands.
+	bool Execute(FRigVMMemoryContainerPtrArray Memory, TArrayView<void*> AdditionalArguments);
+
+	// Executes the VM.
+	// You can optionally provide external memory to the execution
+	// and provide optional additional operands.
+	UFUNCTION(BlueprintCallable, Category = RigVM)
+	bool Execute();
+
+	// Add a function for execute instructions to this VM.
+	// Execute instructions can then refer to the function by index.
+	UFUNCTION()
+	int32 AddRigVMFunction(UScriptStruct* InRigVMStruct, const FName& InMethodName);
+
+	// The default mutable work memory
+	UPROPERTY()
+	FRigVMMemoryContainer WorkMemory;
+
+	// The default const literal memory
+	UPROPERTY()
+	FRigVMMemoryContainer LiteralMemory;
+
+	// The byte code of the VM
+	UPROPERTY()
+	FRigVMByteCode ByteCode;
+
+	// Returns the instructions of the VM
+	const FRigVMInstructionArray& GetInstructions();
+
+	// Returns the parameters of the VM
+	const TArray<FRigVMParameter>& GetParameters();
+
+	// Adds a new input / output to the VM
+	template<class T>
+	FORCEINLINE FRigVMParameter AddPlainParameter(ERigVMParameterType InParameterType, const FName& InName, const FString& InCPPType, const TArray<T>& DefaultValue)
+	{
+		return AddPlainParameter(InParameterType, InName, InCPPType, (uint8*)DefaultValue.GetData(), sizeof(T), DefaultValue.Num());
+	}
+
+	// Adds a new input / output to the VM
+	template<class T>
+	FORCEINLINE FRigVMParameter AddStructParameter(ERigVMParameterType InParameterType, const FName& InName, const TArray<T>& DefaultValue)
+	{
+		UScriptStruct* ScriptStruct = T::StaticStruct();
+		return AddStructParameter(InParameterType, InName, ScriptStruct, (uint8*)DefaultValue.GetData(), DefaultValue.Num());
+	}
+
+	// Adds a new input / output to the VM
+	template<class T>
+	FORCEINLINE FRigVMParameter AddPlainParameter(ERigVMParameterType InParameterType, const FName& InName, const FString& InCPPType, T DefaultValue)
+	{
+		TArray<T> DefaultValues;
+		DefaultValues.Add(DefaultValue);
+		return AddPlainParameter<T>(InParameterType, InName, InCPPType, DefaultValues);
+	}
+
+	// Adds a new input / output to the VM
+	template<class T>
+	FORCEINLINE FRigVMParameter AddStructParameter(ERigVMParameterType InParameterType, const FName& InName, const T& DefaultValue)
+	{
+		TArray<T> DefaultValues;
+		DefaultValues.Add(DefaultValue);
+		return AddStructParameter<T>(InParameterType, InName, DefaultValues);
+	}
+
+	// Retrieve the array size of the parameter
+	int32 GetParameterArraySize(const FRigVMParameter& InParameter) const
+	{
+		return (int32)WorkMemory[InParameter.GetRegisterIndex()].GetTotalElementCount();
+	}
+
+	// Retrieve the array size of the parameter
+	int32 GetParameterArraySize(int32 InParameterIndex) const
+	{
+		return GetParameterArraySize(Parameters[InParameterIndex]);
+	}
+
+	// Retrieve the array size of the parameter
+	UFUNCTION(BlueprintCallable, Category = RigVM)
+	int32 GetParameterArraySize(const FName& InParameterName) const
+	{
+		int32 ParameterIndex = ParametersNameMap.FindChecked(InParameterName);
+		return GetParameterArraySize(ParameterIndex);
+	}
+
+	// Retrieve the value of a parameter
+	template<class T>
+	const T& GetParameterValue(const FRigVMParameter& InParameter, int32 InArrayIndex = 0)
+	{
+		return WorkMemory.GetArray<T>(InParameter.GetRegisterIndex())[InArrayIndex];
+	}
+
+	// Retrieve the value of a parameter given its index
+	template<class T>
+	const T& GetParameterValue(int32 InParameterIndex, int32 InArrayIndex = 0)
+	{
+		return GetParameterValue<T>(Parameters[InParameterIndex], InArrayIndex);
+	}
+
+	// Retrieve the value of a parameter given its name
+	template<class T>
+	const T& GetParameterValue(const FName& InParameterName, int32 InArrayIndex = 0)
+	{
+		int32 ParameterIndex = ParametersNameMap.FindChecked(InParameterName);
+		return GetParameterValue<T>(ParameterIndex, InArrayIndex);
+	}
+
+	// Set the value of a parameter
+	template<class T>
+	void SetParameterValue(const FRigVMParameter& InParameter, const T& InNewValue, int32 InArrayIndex = 0)
+	{
+		WorkMemory.GetArray<T>(InParameter.GetRegisterIndex())[InArrayIndex] = InNewValue;
+	}
+
+	// Set the value of a parameter given its index
+	template<class T>
+	void SetParameterValue(int32 ParameterIndex, const T& InNewValue, int32 InArrayIndex = 0)
+	{
+		return SetParameterValue<T>(Parameters[ParameterIndex], InNewValue, InArrayIndex);
+	}
+
+	// Set the value of a parameter given its name
+	template<class T>
+	void SetParameterValue(const FName& InParameterName, const T& InNewValue, int32 InArrayIndex = 0)
+	{
+		int32 ParameterIndex = ParametersNameMap.FindChecked(InParameterName);
+		return SetParameterValue<T>(ParameterIndex, InNewValue, InArrayIndex);
+	}
+
+	UFUNCTION(BlueprintCallable, Category = RigVM)
+	bool GetParameterValueBool(const FName& InParameterName, int32 InArrayIndex = 0)
+	{
+		return GetParameterValue<bool>(InParameterName, InArrayIndex);
+	}
+
+	UFUNCTION(BlueprintCallable, Category = RigVM)
+	float GetParameterValueFloat(const FName& InParameterName, int32 InArrayIndex = 0)
+	{
+		return GetParameterValue<float>(InParameterName, InArrayIndex);
+	}
+
+	UFUNCTION(BlueprintCallable, Category = RigVM)
+	int32 GetParameterValueInt(const FName& InParameterName, int32 InArrayIndex = 0)
+	{
+		return GetParameterValue<int32>(InParameterName, InArrayIndex);
+	}
+
+	UFUNCTION(BlueprintCallable, Category = RigVM)
+	FName GetParameterValueName(const FName& InParameterName, int32 InArrayIndex = 0)
+	{
+		return GetParameterValue<FName>(InParameterName, InArrayIndex);
+	}
+
+	UFUNCTION(BlueprintCallable, Category = RigVM)
+	FString GetParameterValueString(const FName& InParameterName, int32 InArrayIndex = 0)
+	{
+		return GetParameterValue<FString>(InParameterName, InArrayIndex);
+	}
+
+	UFUNCTION(BlueprintCallable, Category = RigVM)
+	FVector2D GetParameterValueVector2D(const FName& InParameterName, int32 InArrayIndex = 0)
+	{
+		return GetParameterValue<FVector2D>(InParameterName, InArrayIndex);
+	}
+
+	UFUNCTION(BlueprintCallable, Category = RigVM)
+	FVector GetParameterValueVector(const FName& InParameterName, int32 InArrayIndex = 0)
+	{
+		return GetParameterValue<FVector>(InParameterName, InArrayIndex);
+	}
+
+	UFUNCTION(BlueprintCallable, Category = RigVM)
+	FQuat GetParameterValueQuat(const FName& InParameterName, int32 InArrayIndex = 0)
+	{
+		return GetParameterValue<FQuat>(InParameterName, InArrayIndex);
+	}
+
+	UFUNCTION(BlueprintCallable, Category = RigVM)
+	FTransform GetParameterValueTransform(const FName& InParameterName, int32 InArrayIndex = 0)
+	{
+		return GetParameterValue<FTransform>(InParameterName, InArrayIndex);
+	}
+
+	UFUNCTION(BlueprintCallable, Category = RigVM)
+	void SetParameterValueBool(const FName& InParameterName, bool InValue, int32 InArrayIndex = 0)
+	{
+		SetParameterValue<bool>(InParameterName, InValue, InArrayIndex);
+	}
+
+	UFUNCTION(BlueprintCallable, Category = RigVM)
+	void SetParameterValueFloat(const FName& InParameterName, float InValue, int32 InArrayIndex = 0)
+	{
+		SetParameterValue<float>(InParameterName, InValue, InArrayIndex);
+	}
+
+	UFUNCTION(BlueprintCallable, Category = RigVM)
+	void SetParameterValueInt(const FName& InParameterName, int32 InValue, int32 InArrayIndex = 0)
+	{
+		SetParameterValue<int32>(InParameterName, InValue, InArrayIndex);
+	}
+
+	UFUNCTION(BlueprintCallable, Category = RigVM)
+	void SetParameterValueName(const FName& InParameterName, const FName& InValue, int32 InArrayIndex = 0)
+	{
+		SetParameterValue<FName>(InParameterName, InValue, InArrayIndex);
+	}
+
+	UFUNCTION(BlueprintCallable, Category = RigVM)
+	void SetParameterValueString(const FName& InParameterName, const FString& InValue, int32 InArrayIndex = 0)
+	{
+		SetParameterValue<FString>(InParameterName, InValue, InArrayIndex);
+	}
+
+	UFUNCTION(BlueprintCallable, Category = RigVM)
+	void SetParameterValueVector2D(const FName& InParameterName, const FVector2D& InValue, int32 InArrayIndex = 0)
+	{
+		SetParameterValue<FVector2D>(InParameterName, InValue, InArrayIndex);
+	}
+
+	UFUNCTION(BlueprintCallable, Category = RigVM)
+	void SetParameterValueVector(const FName& InParameterName, const FVector& InValue, int32 InArrayIndex = 0)
+	{
+		SetParameterValue<FVector>(InParameterName, InValue, InArrayIndex);
+	}
+
+	UFUNCTION(BlueprintCallable, Category = RigVM)
+	void SetParameterValueQuat(const FName& InParameterName, const FQuat& InValue, int32 InArrayIndex = 0)
+	{
+		SetParameterValue<FQuat>(InParameterName, InValue, InArrayIndex);
+	}
+
+	UFUNCTION(BlueprintCallable, Category = RigVM)
+	void SetParameterValueTransform(const FName& InParameterName, const FTransform& InValue, int32 InArrayIndex = 0)
+	{
+		SetParameterValue<FTransform>(InParameterName, InValue, InArrayIndex);
+	}
+
+private:
+
+	FORCEINLINE FRigVMParameter AddPlainParameter(ERigVMParameterType InParameterType, const FName& InName, const FString& InCPPType, const uint8* InPtr, uint16 InSize, int32 InArraySize)
+	{
+		ensure(InParameterType != ERigVMParameterType::Invalid);
+		int32 ArraySize = InArraySize <= 0 ? 1 : InArraySize;
+		int32 RegisterIndex = WorkMemory.AddPlainArray(WorkMemory.SupportsNames() ? InName : NAME_None, InSize, ArraySize, InPtr, 1);
+		if (RegisterIndex == INDEX_NONE)
+		{
+			return FRigVMParameter();
+		}
+		FName Name = WorkMemory.SupportsNames() ? WorkMemory[RegisterIndex].Name : InName;
+		FRigVMParameter Parameter(InParameterType, Name, RegisterIndex, InCPPType, nullptr);
+		ParametersNameMap.Add(Parameter.Name, Parameters.Add(Parameter));
+		return Parameter;
+	}
+
+	FORCEINLINE FRigVMParameter AddStructParameter(ERigVMParameterType InParameterType, const FName& InName, UScriptStruct* InScriptStruct, const uint8* InPtr, int32 InArraySize)
+	{
+		ensure(InParameterType != ERigVMParameterType::Invalid);
+		int32 ArraySize = InArraySize <= 0 ? 1 : InArraySize;
+		int32 RegisterIndex = WorkMemory.AddStructArray(WorkMemory.SupportsNames() ? InName : NAME_None, InScriptStruct, ArraySize, InPtr, 1);
+		if (RegisterIndex == INDEX_NONE)
+		{
+			return FRigVMParameter();
+		}
+		FName Name = WorkMemory.SupportsNames() ? WorkMemory[RegisterIndex].Name : InName;
+		FRigVMParameter Parameter(InParameterType, Name, RegisterIndex, InScriptStruct->GetName(), InScriptStruct);
+		ParametersNameMap.Add(Parameter.Name, Parameters.Add(Parameter));
+		return Parameter;
+	}
+
+	void ResolveFunctionsIfRequired();
+	void RefreshInstructionsIfRequired();
+
+	UPROPERTY(transient)
+	FRigVMInstructionArray Instructions;
+
+	UPROPERTY()
+	TArray<FString> FunctionNames;
+
+	TArray<FRigVMFunctionPtr> Functions;
+
+	UPROPERTY()
+	TArray<FRigVMParameter> Parameters;
+
+	UPROPERTY()
+	TMap<FName, int32> ParametersNameMap;
+
+	friend class URigVMCompiler;
+};
diff --git a/Engine/Source/Runtime/RigVM/Public/RigVMCore/RigVMByteCode.h b/Engine/Source/Runtime/RigVM/Public/RigVMCore/RigVMByteCode.h
new file mode 100644
--- /dev/null
+++ b/Engine/Source/Runtime/RigVM/Public/RigVMCore/RigVMByteCode.h
@@ -0,0 +1,482 @@
+// Copyright 1998-2019 Epic Games, Inc. All Rights Reserved.
+
+#pragma once
+
+#include "CoreMinimal.h"
+#include "RigVMRegistry.h"
+#include "RigVMByteCode.generated.h"
+
+struct FRigVMByteCode;
+
+// The code for a single operation within the RigVM
+UENUM()
+enum class ERigVMOpCode : uint8
+{
+	Execute_0_Operands, // execute a rig function with 0 operands
+	Execute_1_Operands, // execute a rig function with 1 operands
+	Execute_2_Operands, // execute a rig function with 2 operands
+	Execute_3_Operands, // execute a rig function with 3 operands
+	Execute_4_Operands, // execute a rig function with 4 operands
+	Execute_5_Operands, // execute a rig function with 5 operands
+	Execute_6_Operands, // execute a rig function with 6 operands
+	Execute_7_Operands, // execute a rig function with 7 operands
+	Execute_8_Operands, // execute a rig function with 8 operands
+	Execute_9_Operands, // execute a rig function with 9 operands
+	Execute_10_Operands, // execute a rig function with 10 operands
+	Execute_11_Operands, // execute a rig function with 11 operands
+	Execute_12_Operands, // execute a rig function with 12 operands
+	Execute_13_Operands, // execute a rig function with 13 operands
+	Execute_14_Operands, // execute a rig function with 14 operands
+	Execute_15_Operands, // execute a rig function with 15 operands
+	Execute_16_Operands, // execute a rig function with 16 operands
+	Execute_17_Operands, // execute a rig function with 17 operands
+	Execute_18_Operands, // execute a rig function with 18 operands
+	Execute_19_Operands, // execute a rig function with 19 operands
+	Execute_20_Operands, // execute a rig function with 20 operands
+	Execute_21_Operands, // execute a rig function with 21 operands
+	Execute_22_Operands, // execute a rig function with 22 operands
+	Execute_23_Operands, // execute a rig function with 23 operands
+	Execute_24_Operands, // execute a rig function with 24 operands
+	Execute_25_Operands, // execute a rig function with 25 operands
+	Execute_26_Operands, // execute a rig function with 26 operands
+	Execute_27_Operands, // execute a rig function with 27 operands
+	Execute_28_Operands, // execute a rig function with 28 operands
+	Execute_29_Operands, // execute a rig function with 29 operands
+	Execute_30_Operands, // execute a rig function with 30 operands
+	Execute_31_Operands, // execute a rig function with 31 operands
+	Execute_32_Operands, // execute a rig function with 32 operands
+	Execute_33_Operands, // execute a rig function with 33 operands
+	Execute_34_Operands, // execute a rig function with 34 operands
+	Execute_35_Operands, // execute a rig function with 35 operands
+	Execute_36_Operands, // execute a rig function with 36 operands
+	Execute_37_Operands, // execute a rig function with 37 operands
+	Execute_38_Operands, // execute a rig function with 38 operands
+	Execute_39_Operands, // execute a rig function with 39 operands
+	Execute_40_Operands, // execute a rig function with 40 operands
+	Execute_41_Operands, // execute a rig function with 41 operands
+	Execute_42_Operands, // execute a rig function with 42 operands
+	Execute_43_Operands, // execute a rig function with 43 operands
+	Execute_44_Operands, // execute a rig function with 44 operands
+	Execute_45_Operands, // execute a rig function with 45 operands
+	Execute_46_Operands, // execute a rig function with 46 operands
+	Execute_47_Operands, // execute a rig function with 47 operands
+	Execute_48_Operands, // execute a rig function with 48 operands
+	Execute_49_Operands, // execute a rig function with 49 operands
+	Execute_50_Operands, // execute a rig function with 50 operands
+	Execute_51_Operands, // execute a rig function with 51 operands
+	Execute_52_Operands, // execute a rig function with 52 operands
+	Execute_53_Operands, // execute a rig function with 53 operands
+	Execute_54_Operands, // execute a rig function with 54 operands
+	Execute_55_Operands, // execute a rig function with 55 operands
+	Execute_56_Operands, // execute a rig function with 56 operands
+	Execute_57_Operands, // execute a rig function with 57 operands
+	Execute_58_Operands, // execute a rig function with 58 operands
+	Execute_59_Operands, // execute a rig function with 59 operands
+	Execute_60_Operands, // execute a rig function with 60 operands
+	Execute_61_Operands, // execute a rig function with 61 operands
+	Execute_62_Operands, // execute a rig function with 62 operands
+	Execute_63_Operands, // execute a rig function with 63 operands
+	Execute_64_Operands, // execute a rig function with 64 operands
+	Zero, // zero the memory of a given register
+	BoolFalse, // set a given register to false
+	BoolTrue, // set a given register to true
+	Copy, // copy the content of one register to another
+	Increment, // increment a int32 register
+	Decrement, // decrement a int32 register
+	Equals, // fill a bool register with the result of (A == B)
+	NotEquals, // fill a bool register with the result of (A != B)
+	JumpAbsolute, // jump to an absolute instruction index
+	JumpForward, // jump forwards given a relative instruction index offset
+	JumpBackward, // jump backwards given a relative instruction index offset
+	JumpAbsoluteIf, // jump to an absolute instruction index based on a condition register
+	JumpForwardIf, // jump forwards given a relative instruction index offset based on a condition register
+	JumpBackwardIf, // jump backwards given a relative instruction index offset based on a condition register
+	ChangeType, // change the type of a register
+	Exit, // exit the execution loop
+	Invalid
+};
+
+// Base class for all VM operations
+struct RIGVM_API FRigVMBaseOp
+{
+	FRigVMBaseOp(ERigVMOpCode InOpCode = ERigVMOpCode::Invalid)
+	: OpCode(InOpCode)
+	{
+	}
+
+	ERigVMOpCode OpCode;
+};
+
+
+// execute a function
+struct RIGVM_API FRigVMExecuteOp : public FRigVMBaseOp
+{
+	FRigVMExecuteOp()
+	: FRigVMBaseOp()
+	, FunctionIndex(INDEX_NONE)
+	{
+	}
+
+	FRigVMExecuteOp(uint16 InFunctionIndex, uint8 InArgumentCount)
+	: FRigVMBaseOp((ERigVMOpCode)(uint8(ERigVMOpCode::Execute_0_Operands) + InArgumentCount))
+	, FunctionIndex(InFunctionIndex)
+	{
+	}
+
+	uint16 FunctionIndex;
+
+	FORCEINLINE uint8 GetOperandCount() const { return uint8(OpCode) - uint8(ERigVMOpCode::Execute_0_Operands); }
+};
+
+// operator used for zero, false, true, increment, decrement
+struct RIGVM_API FRigVMUnaryOp : public FRigVMBaseOp
+{
+	FRigVMUnaryOp()
+		: FRigVMBaseOp(ERigVMOpCode::Invalid)
+		, Arg()
+	{
+	}
+
+	FRigVMUnaryOp(ERigVMOpCode InOpCode, FRigVMOperand InArg)
+		: FRigVMBaseOp(InOpCode)
+		, Arg(InArg)
+	{
+		ensure(
+			uint8(InOpCode) == uint8(ERigVMOpCode::Zero) ||
+			uint8(InOpCode) == uint8(ERigVMOpCode::BoolFalse) ||
+			uint8(InOpCode) == uint8(ERigVMOpCode::BoolTrue) ||
+			uint8(InOpCode) == uint8(ERigVMOpCode::Increment) ||
+			uint8(InOpCode) == uint8(ERigVMOpCode::Decrement) ||
+			uint8(InOpCode) == uint8(ERigVMOpCode::JumpAbsoluteIf) ||
+			uint8(InOpCode) == uint8(ERigVMOpCode::JumpForwardIf) ||
+			uint8(InOpCode) == uint8(ERigVMOpCode::JumpBackwardIf) ||
+			uint8(InOpCode) == uint8(ERigVMOpCode::ChangeType)
+		);
+	}
+
+	FRigVMOperand Arg;
+};
+
+// copy the content of one register to another
+struct RIGVM_API FRigVMCopyOp : public FRigVMBaseOp
+{
+	FRigVMCopyOp()
+	: FRigVMBaseOp(ERigVMOpCode::Copy)
+	, Source()
+	, Target()
+	{
+	}
+
+	FRigVMCopyOp(
+		FRigVMOperand InSource,
+		FRigVMOperand InTarget
+	)
+		: FRigVMBaseOp(ERigVMOpCode::Copy)
+		, Source(InSource)
+		, Target(InTarget)
+	{
+	}
+
+	FRigVMOperand Source;
+	FRigVMOperand Target;
+};
+
+// used for equals and not equals comparisons
+struct RIGVM_API FRigVMComparisonOp : public FRigVMBaseOp
+{
+	FRigVMComparisonOp()
+		: FRigVMBaseOp(ERigVMOpCode::Invalid)
+		, A()
+		, B()
+		, Result()
+	{
+	}
+
+	FRigVMComparisonOp(
+		ERigVMOpCode InOpCode,
+		FRigVMOperand InA,
+		FRigVMOperand InB,
+		FRigVMOperand InResult
+	)
+		: FRigVMBaseOp(InOpCode)
+		, A(InA)
+		, B(InB)
+		, Result(InResult)
+	{
+		ensure(
+			uint8(InOpCode) == uint8(ERigVMOpCode::Equals) ||
+			uint8(InOpCode) == uint8(ERigVMOpCode::NotEquals)
+			);
+	}
+
+	FRigVMOperand A;
+	FRigVMOperand B;
+	FRigVMOperand Result;
+};
+
+// jump to a new instruction index.
+// the instruction can be absolute, relative forward or relative backward
+// based on the opcode 
+struct RIGVM_API FRigVMJumpOp : public FRigVMBaseOp
+{
+	FRigVMJumpOp()
+	: FRigVMBaseOp(ERigVMOpCode::Invalid)
+	, InstructionIndex(INDEX_NONE)
+	{
+	}
+
+	FRigVMJumpOp(ERigVMOpCode InOpCode, int32 InInstructionIndex)
+	: FRigVMBaseOp(InOpCode)
+	, InstructionIndex(InInstructionIndex)
+	{
+		ensure(uint8(InOpCode) >= uint8(ERigVMOpCode::JumpAbsolute));
+		ensure(uint8(InOpCode) <= uint8(ERigVMOpCode::JumpBackward));
+	}
+
+	int32 InstructionIndex;
+};
+
+// jump to a new instruction index based on a condition.
+// the instruction can be absolute, relative forward or relative backward
+// based on the opcode 
+struct RIGVM_API FRigVMJumpIfOp : public FRigVMUnaryOp
+{
+	FRigVMJumpIfOp()
+		: FRigVMUnaryOp()
+		, InstructionIndex(INDEX_NONE)
+		, Condition(true)
+	{
+	}
+
+	FRigVMJumpIfOp(ERigVMOpCode InOpCode, FRigVMOperand InConditionArg, int32 InInstructionIndex, bool InCondition = false)
+		: FRigVMUnaryOp(InOpCode, InConditionArg)
+		, InstructionIndex(InInstructionIndex)
+		, Condition(InCondition)
+	{
+		ensure(uint8(InOpCode) >= uint8(ERigVMOpCode::JumpAbsoluteIf));
+		ensure(uint8(InOpCode) <= uint8(ERigVMOpCode::JumpBackwardIf));
+	}
+
+	int32 InstructionIndex;
+	bool Condition;
+};
+
+// change the type of a register
+struct RIGVM_API FRigVMChangeTypeOp : public FRigVMUnaryOp
+{
+	FRigVMChangeTypeOp()
+		: FRigVMUnaryOp()
+		, Type(ERigVMRegisterType::Invalid)
+		, ElementSize(0)
+		, ElementCount(0)
+		, SliceCount(0)
+	{
+	}
+
+	FRigVMChangeTypeOp(FRigVMOperand InArg, ERigVMRegisterType InType, uint16 InElementSize, uint16 InElementCount, uint16 InSliceCount)
+		: FRigVMUnaryOp(ERigVMOpCode::ChangeType, InArg)
+		, Type(InType)
+		, ElementSize(InElementSize)
+		, ElementCount(InElementCount)
+		, SliceCount(InSliceCount)
+	{
+	}
+
+	ERigVMRegisterType Type;
+	uint16 ElementSize;
+	uint16 ElementCount;
+	uint16 SliceCount;
+};
+
+/**
+ * The FRigVMInstruction represents
+ * a single instruction within the VM.
+ */
+USTRUCT()
+struct RIGVM_API FRigVMInstruction
+{
+	GENERATED_USTRUCT_BODY()
+
+	FRigVMInstruction(ERigVMOpCode InOpCode = ERigVMOpCode::Invalid, uint64 InByteCodeIndex = UINT64_MAX)
+		: OpCode(InOpCode)
+		, ByteCodeIndex(InByteCodeIndex)
+	{
+	}
+
+	UPROPERTY()
+	ERigVMOpCode OpCode;
+
+	UPROPERTY()
+	uint64 ByteCodeIndex;
+};
+
+/**
+ * The FRigVMInstructionArray represents all current instructions
+ * within a RigVM and can be used to iterate over all operators and retrieve
+ * each instruction's data.
+ */
+USTRUCT()
+struct RIGVM_API FRigVMInstructionArray
+{
+	GENERATED_USTRUCT_BODY()
+
+public:
+
+	FRigVMInstructionArray();
+
+	// Resets the data structure and removes all storage.
+	void Reset();
+
+	// Returns true if a given instruction index is valid.
+	FORCEINLINE bool IsValidIndex(int32 InIndex) const { return Instructions.IsValidIndex(InIndex); }
+
+	// Returns the number of instructions.
+	FORCEINLINE int32 Num() const { return Instructions.Num(); }
+
+	// const accessor for an instruction given its index
+	FORCEINLINE const FRigVMInstruction& operator[](int32 InIndex) const { return Instructions[InIndex]; }
+
+private:
+
+	// hide utility constructor
+	FRigVMInstructionArray(const FRigVMByteCode& InByteCode);
+
+	UPROPERTY()
+	TArray<FRigVMInstruction> Instructions;
+
+	friend struct FRigVMByteCode;
+};
+
+/**
+ * The FRigVMByteCode is a container to store a list of instructions with
+ * their corresponding data. The byte code is then used within a VM to 
+ * execute. To iterate over the instructions within the byte code you can 
+ * use GetInstructions() to retrieve a FRigVMInstructionArray.
+ */
+USTRUCT()
+struct RIGVM_API FRigVMByteCode
+{
+	GENERATED_USTRUCT_BODY()
+
+public:
+
+	FRigVMByteCode();
+
+	// resets the container and removes all memory
+	void Reset();
+
+	// returns the number of instructions in this container
+	uint64 Num() const;
+
+	// adds an execute operator given its function index operands
+	uint64 AddExecuteOp(uint16 InFunctionIndex, const TArrayView<FRigVMOperand>& InOperands);
+
+	// adds a zero operator to zero the memory of a given argument
+	uint64 AddZeroOp(const FRigVMOperand& InArg);
+
+	// adds a false operator to set a given argument to false
+	uint64 AddFalseOp(const FRigVMOperand& InArg);
+
+	// adds a true operator to set a given argument to true
+	uint64 AddTrueOp(const FRigVMOperand& InArg);
+
+	// adds a copy operator to copy the content of a source argument to a target argument
+	uint64 AddCopyOp(const FRigVMOperand& InSource, const FRigVMOperand& InTarget);
+
+	// adds an increment operator to increment a int32 argument
+	uint64 AddIncrementOp(const FRigVMOperand& InArg);
+
+	// adds an decrement operator to decrement a int32 argument
+	uint64 AddDecrementOp(const FRigVMOperand& InArg);
+
+	// adds an equals operator to store the comparison result of A and B into a Result argument
+	uint64 AddEqualsOp(const FRigVMOperand& InA, const FRigVMOperand& InB, const FRigVMOperand& InResult);
+
+	// adds an not-equals operator to store the comparison result of A and B into a Result argument
+	uint64 AddNotEqualsOp(const FRigVMOperand& InA, const FRigVMOperand& InB, const FRigVMOperand& InResult);
+
+	// adds an absolute, forward or backward jump operator
+	uint64 AddJumpOp(ERigVMOpCode InOpCode, uint16 InInstructionIndex);
+
+	// adds an absolute, forward or backward jump operator based on a condition argument
+	uint64 AddJumpIfOp(ERigVMOpCode InOpCode, uint16 InInstructionIndex, const FRigVMOperand& InConditionArg, bool bJumpWhenConditionIs = false);
+
+	// adds a change-type operator to reuse a register for a smaller or same size type
+	uint64 AddChangeTypeOp(FRigVMOperand InArg, ERigVMRegisterType InType, uint16 InElementSize, uint16 InElementCount, uint16 InSliceCount = 1);
+
+	// adds an exit operator to exit the execution loop
+	uint64 AddExitOp();
+
+	// returns an instruction array for iterating over all operators
+	FORCEINLINE FRigVMInstructionArray GetInstructions() const
+	{
+		return FRigVMInstructionArray(*this);
+	}
+
+	// returns the opcode at a given byte index
+	FORCEINLINE ERigVMOpCode GetOpCodeAt(uint64 InByteCodeIndex) const
+	{
+		ensure(InByteCodeIndex >= 0 && InByteCodeIndex < ByteCode.Num());
+		return (ERigVMOpCode)ByteCode[InByteCodeIndex];
+	}
+
+	// returns the size of the operator in bytes at a given byte index
+	uint64 GetOpNumBytesAt(uint64 InByteCodeIndex, bool bIncludeOperands = true) const;
+
+	// returns an operator at a given byte code index
+	template<class OpType>
+	FORCEINLINE const OpType& GetOpAt(uint64 InByteCodeIndex) const
+	{
+		ensure(InByteCodeIndex >= 0 && InByteCodeIndex <= ByteCode.Num() - sizeof(OpType));
+		return *(const OpType*)(ByteCode.GetData() + InByteCodeIndex);
+	}
+
+	// returns an operator for a given instruction
+	template<class OpType>
+	FORCEINLINE const OpType& GetOpAt(const FRigVMInstruction& InInstruction) const
+	{
+		return GetOpAt<OpType>(InInstruction.ByteCodeIndex);
+	}
+
+	// returns a list of operands at a given byte code index
+	FORCEINLINE TArrayView<FRigVMOperand> GetOperandsAt(uint64 InByteCodeIndex, uint16 InArgumentCount) const
+	{
+		ensure(InByteCodeIndex >= 0 && InByteCodeIndex <= ByteCode.Num() - sizeof(FRigVMOperand) * InArgumentCount);
+		return TArrayView<FRigVMOperand>((FRigVMOperand*)(ByteCode.GetData() + InByteCodeIndex), InArgumentCount);
+	}
+
+	// returns the operands for an execute operator / instruction at a given byte code index
+	FORCEINLINE TArrayView<FRigVMOperand> GetOperandsForExecuteOp(uint64 InByteCodeIndex) const
+	{
+		const FRigVMExecuteOp& ExecuteOp = GetOpAt<FRigVMExecuteOp>(InByteCodeIndex);
+		return GetOperandsAt(InByteCodeIndex + sizeof(FRigVMExecuteOp), ExecuteOp.GetOperandCount());
+	}
+
+	// returns the operands for a given execute instruction
+	FORCEINLINE TArrayView<FRigVMOperand> GetOperandsForExecuteOp(const FRigVMInstruction& InInstruction) const
+	{
+		const FRigVMExecuteOp& ExecuteOp = GetOpAt<FRigVMExecuteOp>(InInstruction);
+		return GetOperandsAt(InInstruction.ByteCodeIndex + sizeof(FRigVMExecuteOp), ExecuteOp.GetOperandCount());
+	}
+
+	// returns the raw data of the byte code
+	FORCEINLINE const TArrayView<uint8> GetByteCode() const
+	{
+		const uint8* Data = ByteCode.GetData();
+		return TArrayView<uint8>((uint8*)Data, ByteCode.Num());
+	}
+
+private:
+
+	template<class OpType>
+	FORCEINLINE uint64 AddOp(const OpType& InOp)
+	{
+		uint64 ByteIndex = (uint64)ByteCode.AddUninitialized(sizeof(OpType));
+		FMemory::Memcpy(ByteCode.GetData() + ByteIndex, &InOp, sizeof(OpType));
+		return ByteIndex;
+	}
+
+	// memory for all functions
+	UPROPERTY()
+	TArray<uint8> ByteCode;
+};
diff --git a/Engine/Source/Runtime/RigVM/Public/RigVMCore/RigVMExecuteContext.h b/Engine/Source/Runtime/RigVM/Public/RigVMCore/RigVMExecuteContext.h
new file mode 100644
--- /dev/null
+++ b/Engine/Source/Runtime/RigVM/Public/RigVMCore/RigVMExecuteContext.h
@@ -0,0 +1,18 @@
+// Copyright 1998-2019 Epic Games, Inc. All Rights Reserved.
+
+#pragma once
+
+#include "CoreMinimal.h"
+#include "RigVMExecuteContext.generated.h"
+
+/**
+ * The execute context is used for mutable nodes to
+ * indicate execution order.
+ */
+USTRUCT()
+struct RIGVM_API FRigVMExecuteContext
+{
+	GENERATED_BODY()
+
+	TArrayView<void*> OpaqueArguments;
+};
