Copying //UE4/Dev-Anim to Dev-Main (//UE4/Dev-Main) @ 8782600

diff --git a/Engine/Source/Runtime/RigVM/Public/RigVMCore/RigVMMemory.h b/Engine/Source/Runtime/RigVM/Public/RigVMCore/RigVMMemory.h
new file mode 100644
--- /dev/null
+++ b/Engine/Source/Runtime/RigVM/Public/RigVMCore/RigVMMemory.h
@@ -0,0 +1,1018 @@
+// Copyright 1998-2019 Epic Games, Inc. All Rights Reserved.
+
+#pragma once
+
+#include "CoreMinimal.h"
+#include "UObject/ObjectMacros.h"
+#include "RigVMMemory.generated.h"
+
+/**
+ * The type of memory used. Typically we differentiate between
+ * Work (Mutable) and Literal (Constant) memory.
+ */
+UENUM()
+enum class ERigVMMemoryType: uint8
+{
+	Work, // Mutable state
+	Literal, // Const / fixed state
+	Invalid
+};
+
+/**
+ * The FRigVMOperand represents an argument used for an operator
+ * within the virtual machine. Operands provide information about
+ * which memory needs to be referred to, which register within the
+ * memory all the way to the actual byte address in memory.
+ * The FRigVMOperand is a light weight address for a register in
+ * a FRigVMMemoryContainer.
+ */
+struct RIGVM_API FRigVMOperand
+{
+public:
+
+	FRigVMOperand()
+		: MemoryType(ERigVMMemoryType::Work)
+		, RegisterIndex(UINT16_MAX)
+		, RegisterOffset(UINT16_MAX)
+	{
+	}
+
+	FRigVMOperand(ERigVMMemoryType InMemoryType, int32 InRegisterIndex, int32 InRegisterOffset = INDEX_NONE)
+		: MemoryType(InMemoryType)
+		, RegisterIndex(InRegisterIndex < 0 ? UINT16_MAX : (uint16)InRegisterIndex)
+		, RegisterOffset(InRegisterOffset < 0 ? UINT16_MAX : (uint16)InRegisterOffset)
+	{
+	}
+
+	// returns the memory type of this argument
+	FORCEINLINE ERigVMMemoryType GetMemoryType() const { return MemoryType; }
+
+	// returns the index of the container of this argument
+	FORCEINLINE int32 GetContainerIndex() const { return (int32)MemoryType; }
+
+	// returns the index of the register of this argument
+	FORCEINLINE int32 GetRegisterIndex() const { return RegisterIndex == UINT16_MAX ? INDEX_NONE : (int32)RegisterIndex; }
+
+	// returns the index of the register of this argument
+	FORCEINLINE int32 GetRegisterOffset() const { return RegisterOffset == UINT16_MAX ? INDEX_NONE : (int32)RegisterOffset; }
+
+private:
+
+	ERigVMMemoryType MemoryType;
+	uint16 RegisterIndex;
+	uint16 RegisterOffset;
+};
+
+typedef TArrayView<FRigVMOperand> FRigVMOperandArray;
+
+/**
+ * The type of register within the memory.
+ */
+UENUM()
+enum class ERigVMRegisterType : uint8
+{
+	Plain, // bool, int32, float, FVector etc.
+	String, // FString
+	Name, // FName
+	Struct, // Any USTRUCT
+	Invalid
+};
+
+// The register represents an address with the VM's memory. Within a register
+// we can store arbitrary data, so it provides a series of properties to
+// describe the memory location.
+// Registers also support the notion of slices. A slice is a complete copy of
+// the memory - so for example if your register stores 4 Vectors, then a slice 
+// would contain 48 bytes (4 * 3 * 4). The register can however store multiple
+// slices / copies of that if needed. Slices can be used to provide 
+// per-invocation memory to functions within the same register.
+// An integrator for example that needs to store a simulated position
+// might want access to a separate memory per loop iteration.
+USTRUCT()
+struct RIGVM_API FRigVMRegister
+{
+	GENERATED_BODY()
+
+	FRigVMRegister()
+		: Type(ERigVMRegisterType::Invalid)
+		, ByteIndex(INDEX_NONE)
+		, ElementSize(0)
+		, ElementCount(0)
+		, SliceIndex(0)
+		, SliceCount(1)
+		, AlignmentBytes(0)
+		, TrailingBytes(0)
+		, Name(NAME_None)
+		, ScriptStructIndex(INDEX_NONE)
+	{
+	}
+
+	// The type of register (plain, name, string, etc.)
+	UPROPERTY()
+	ERigVMRegisterType Type;
+
+	// The index of the first work byte
+	UPROPERTY()
+	uint32 ByteIndex;
+
+	// The size of each store element
+	UPROPERTY()
+	uint16 ElementSize;
+
+	// The number of elements in this register
+	UPROPERTY()
+	uint16 ElementCount;
+
+	// The currently active slice index
+	UPROPERTY()
+	uint16 SliceIndex;
+
+	// The number of slices (complete copies)
+	UPROPERTY()
+	uint16 SliceCount;
+
+	// The number of leading bytes for alignment
+	UPROPERTY()
+	uint8 AlignmentBytes;
+
+	// The number of trailing bytes.
+	// These originate after shrinking a register.
+	UPROPERTY()
+	uint16 TrailingBytes;
+
+	// The name of the register (can be None)
+	UPROPERTY()
+	FName Name;
+
+	// For struct registers this is the index of the
+	// struct used - otherwise INDEX_NONE
+	UPROPERTY()
+	int32 ScriptStructIndex;
+
+	bool Serialize(FArchive& Ar);
+	FORCEINLINE friend FArchive& operator<<(FArchive& Ar, FRigVMRegister& P)
+	{
+		P.Serialize(Ar);
+		return Ar;
+	}
+
+	// returns the current address of the register within the data byte array.
+	// this can change over time - as the register is moving through slices.
+	// use GetFirstAllocatedByte to get the fixed first byte.
+	FORCEINLINE uint64 GetWorkByteIndex() const { return ByteIndex; }
+
+	// returns the first allocated byte in the data byte array
+	FORCEINLINE uint64 GetFirstAllocatedByte() const { return ByteIndex - (uint64)AlignmentBytes - (uint64)(SliceIndex * GetNumBytesPerSlice()); }
+
+	// Returns the leading alignment bytes
+	FORCEINLINE uint8 GetAlignmentBytes() const { return AlignmentBytes; }
+
+	// Returns true if the register stores more than one element
+	FORCEINLINE bool IsArray() const { return ElementCount > 1; }
+
+	// Returns true if the register stores plain / shallow memory
+	FORCEINLINE bool IsPlain() const { return ScriptStructIndex == INDEX_NONE; }
+
+	// Returns the number of allocated bytes (including alignment + trailing bytes)
+	FORCEINLINE uint16 GetAllocatedBytes() const { return ElementCount * ElementSize * SliceCount + (uint16)AlignmentBytes + TrailingBytes; }
+
+	// Returns the number of bytes for a complete slice
+	FORCEINLINE uint16 GetNumBytesPerSlice() const { return ElementCount * ElementSize; }
+
+	// Returns the number of bytes for all slices
+	FORCEINLINE uint16 GetNumBytesAllSlices() const { return ElementCount * ElementSize * SliceCount; }
+
+	// Returns the total number of elements (elementcount * slicecount) in the register
+	FORCEINLINE uint32 GetTotalElementCount() const { return (uint32)ElementCount * (uint32)SliceCount; }
+	
+	// Moves the register address to the first slice
+	FORCEINLINE void MoveToFirstSlice() { ByteIndex -= SliceIndex * GetNumBytesPerSlice(); SliceIndex = 0; }
+
+	// Moves the register address to the next slice. This also wraps around - so when the register reaches 
+	// the end it will move to the first slice again.
+	FORCEINLINE void MoveToNextSlice()
+	{
+		if(SliceCount <= 1)
+		{
+			return;
+		}
+		
+		if(SliceIndex == SliceCount - 1)
+		{
+			MoveToFirstSlice();
+		}
+		else
+		{
+			ByteIndex += GetNumBytesPerSlice();
+			SliceIndex++;
+		}
+	}
+};
+
+typedef TArrayView<FRigVMRegister> FRigVMRegisterArray;
+
+// The register offset represents a memory offset within a register's memory.
+// This can be used to represent memory addresses of array elements within
+// a struct, for example.
+USTRUCT()
+struct RIGVM_API FRigVMRegisterOffset
+{
+	GENERATED_BODY()
+
+public:
+
+	// default constructor
+	FRigVMRegisterOffset()
+		: Segments()
+		, Type(ERigVMRegisterType::Invalid)
+		, CPPType(NAME_None)
+		, ScriptStruct(nullptr)
+		, ScriptStructPath(NAME_None)
+		, ElementSize(0)
+	{
+	}
+
+	bool Serialize(FArchive& Ar);
+	FORCEINLINE friend FArchive& operator<<(FArchive& Ar, FRigVMRegisterOffset& P)
+	{
+		P.Serialize(Ar);
+		return Ar;
+	}
+
+	// constructs a path given a struct and a segment path
+	FRigVMRegisterOffset(UScriptStruct* InScriptStruct, const FString& InSegmentPath, int32 InInitialOffset = 0, uint16 InElementSize = 0);
+
+	// returns the data pointer within a container
+	uint8* GetData(uint8* InContainer) const;
+
+	// returns the segments of this path
+	const TArray<int32>& GetSegments() const { return Segments; }
+
+	bool operator == (const FRigVMRegisterOffset& InOther) const;
+
+	FORCEINLINE ERigVMRegisterType GetType() const { return Type; }
+	FORCEINLINE FName GetCPPType() const { return CPPType; }
+	uint16 GetElementSize() const;
+	UScriptStruct* GetScriptStruct() const;
+
+private:
+
+	UPROPERTY()
+	TArray<int32> Segments;
+
+	UPROPERTY()
+	ERigVMRegisterType Type;
+
+	UPROPERTY()
+	FName CPPType;
+
+	UPROPERTY(transient)
+	UScriptStruct* ScriptStruct;
+
+	UPROPERTY()
+	FName ScriptStructPath;
+
+	UPROPERTY()
+	uint16 ElementSize;
+
+	friend struct FRigVMRegisterOffsetBuilder;
+};
+
+/**
+ * The FRigVMMemoryContainer provides a heterogeneous memory container to store arbitrary
+ * data. Each element stored can be referred to using a FRigVMRegister.
+ * Elements can be accessed by index (index of the register), FRigVMOperand or by name.
+ * Name access is optional and is specified upon construction of the container.
+ * The memory container provides a series of templated functions to add and get data.
+ * Plain types (shallow types without the need for construction) can be added and retrieved
+ * using the methods suffixed with 'Plain'.
+ *
+ * For example:
+ * 		int32 Index = Container.AddPlain<float>(4.f);
+ *      float& ValueRef = Container.GetRef<float>(Index);
+ *
+ * This can also be done with arrays:
+ *      TArray<float> MyArray = {3.f, 4.f, 5.f};
+ * 		int32 Index = Container.AddPlainArray<float>(MyArray);
+ *      TArrayView<float> ArrayView = Container.GetArray<float>(Index);
+ *
+ * More complex data structures such as USTRUCT types need to be stored with the
+ * 'Struct' suffixed methods - since they require a constructor / destructor to be called.
+ *
+ * For example:
+ *
+ *		FMyStruct MyStruct;
+ *		int32 Index = Container.AddStruct<FMyStruct>(MyStruct);
+ *      FMyStruct& ValueRef = Container.GetRef<FMyStruct>(Index);
+ */
+USTRUCT()
+struct RIGVM_API FRigVMMemoryContainer
+{
+	GENERATED_BODY()
+	
+public:
+
+	FRigVMMemoryContainer(bool bInUseNames = true);
+	FRigVMMemoryContainer(const FRigVMMemoryContainer& Other);
+	~FRigVMMemoryContainer();
+
+	FRigVMMemoryContainer& operator= (const FRigVMMemoryContainer &InOther);
+
+	// returns the memory type of this container
+	FORCEINLINE ERigVMMemoryType GetMemoryType() const { return MemoryType;  }
+
+	// sets the memory type. should only be used when the container is empty
+	FORCEINLINE void SetMemoryType(ERigVMMemoryType InMemoryType) { MemoryType = InMemoryType; }
+
+	// returns true if this container supports name based lookup
+	FORCEINLINE bool SupportsNames() const { return bUseNameMap;  }
+
+	// returns the number of registers in this container
+	FORCEINLINE int32 Num() const { return Registers.Num(); }
+
+	// resets the container and removes all storage.
+	void Reset();
+
+	// const accessor for a register based on index
+	FORCEINLINE const FRigVMRegister& operator[](int32 InIndex) const { return Registers[InIndex]; }
+
+	// accessor for a register based on index
+	FORCEINLINE FRigVMRegister& operator[](int32 InIndex) { return Registers[InIndex]; }
+
+	// const accessor for a register based on an argument
+	FORCEINLINE const FRigVMRegister& operator[](const FRigVMOperand& InArg) const { return Registers[InArg.GetRegisterIndex()]; }
+
+	// accessor for a register based on an argument
+	FORCEINLINE FRigVMRegister& operator[](const FRigVMOperand& InArg) { return Registers[InArg.GetRegisterIndex()]; }
+
+	// const accessor for a register based on a a name. note: only works if SupportsNames() == true
+	FORCEINLINE const FRigVMRegister& operator[](const FName& InName) const { return Registers[GetIndex(InName)]; }
+
+	// accessor for a register based on a a name. note: only works if SupportsNames() == true
+	FORCEINLINE FRigVMRegister& operator[](const FName& InName) { return Registers[GetIndex(InName)]; }
+
+	FORCEINLINE TArray<FRigVMRegister>::RangedForIteratorType      begin() { return Registers.begin(); }
+	FORCEINLINE TArray<FRigVMRegister>::RangedForConstIteratorType begin() const { return Registers.begin(); }
+	FORCEINLINE TArray<FRigVMRegister>::RangedForIteratorType      end() { return Registers.end(); }
+	FORCEINLINE TArray<FRigVMRegister>::RangedForConstIteratorType end() const { return Registers.end(); }
+
+	bool Serialize(FArchive& Ar);
+	FORCEINLINE friend FArchive& operator<<(FArchive& Ar, FRigVMMemoryContainer& P)
+	{
+		P.Serialize(Ar);
+		return Ar;
+	}
+
+	// Returns an argument for a given register.
+	// This is typically used to store a light weight address for use within a VM.
+	FORCEINLINE FRigVMOperand GetOperand(int32 InRegisterIndex, int32 InRegisterOffset)
+	{
+		ensure(Registers.IsValidIndex(InRegisterIndex));
+		return FRigVMOperand(MemoryType, InRegisterIndex, InRegisterOffset);
+	}
+
+	// Returns an argument for a given register.
+	// This is typically used to store a light weight address for use within a VM.
+	FORCEINLINE FRigVMOperand GetOperand(int32 InRegisterIndex, const FString& InSegmentPath = FString(), int32 InArrayElement = INDEX_NONE)
+	{
+		ensure(Registers.IsValidIndex(InRegisterIndex));
+		UScriptStruct* ScriptStruct = nullptr;
+		if (!InSegmentPath.IsEmpty())
+		{
+			ScriptStruct = GetScriptStruct(InRegisterIndex);
+		}
+
+		int32 InitialOffset = 0;
+		int32 ElementSize = 0;
+		if (InArrayElement != INDEX_NONE)
+		{
+			InitialOffset = InArrayElement * Registers[InRegisterIndex].ElementSize;
+			ElementSize = Registers[InRegisterIndex].ElementSize;
+		}
+
+		return GetOperand(InRegisterIndex, GetOrAddRegisterOffset(InRegisterIndex, ScriptStruct, InSegmentPath, InitialOffset, ElementSize));
+	}
+
+	// Returns the current const data pointer for a given register.
+	// Note: This refers to the active slice - and can change over time.
+	FORCEINLINE const uint8* GetData(const FRigVMRegister& Register, int32 InRegisterOffset = INDEX_NONE) const
+	{
+		ensure(Register.ElementCount > 0);
+		uint8* Ptr = (uint8*)&Data[Register.GetWorkByteIndex()];
+		if (InRegisterOffset != INDEX_NONE)
+		{
+			Ptr = RegisterOffsets[InRegisterOffset].GetData(Ptr);
+		}
+		return Ptr;
+	}
+
+	// Returns the current const data pointer for a given register index.
+	// Note: This refers to the active slice - and can change over time.
+	FORCEINLINE const uint8* GetData(int32 InRegisterIndex, int32 InRegisterOffset = INDEX_NONE) const
+	{
+		ensure(Registers.IsValidIndex(InRegisterIndex));
+		const FRigVMRegister& Register = Registers[InRegisterIndex];
+		return GetData(Register, InRegisterOffset);
+	}
+
+	// Returns the current data pointer for a given register.
+	// Note: This refers to the active slice - and can change over time.
+	FORCEINLINE uint8* GetData(FRigVMRegister& Register, int32 InRegisterOffset = INDEX_NONE, bool bMoveToNextSlice = false)
+	{
+		ensure(Register.ElementCount > 0);
+		if(bMoveToNextSlice)
+		{
+			Register.MoveToNextSlice();
+		}
+		uint8* Ptr = &Data[Register.GetWorkByteIndex()];
+		if (InRegisterOffset != INDEX_NONE)
+		{
+			Ptr = RegisterOffsets[InRegisterOffset].GetData(Ptr);
+		}
+		return Ptr;
+	}
+
+	// Returns the current data pointer for a given register index.
+	// Note: This refers to the active slice - and can change over time.
+	FORCEINLINE uint8* GetData(int32 InRegisterIndex, int32 InRegisterOffset = INDEX_NONE, bool bMoveToNextSlice = false)
+	{
+		ensure(Registers.IsValidIndex(InRegisterIndex));
+		FRigVMRegister& Register = Registers[InRegisterIndex];
+		return GetData(Register, InRegisterOffset, bMoveToNextSlice);
+	}
+
+	// Returns the current const typed pointer for a given register.
+	// Note: This refers to the active slice - and can change over time.
+	template<class T>
+	FORCEINLINE const T* Get(const FRigVMRegister& InRegister, int32 InRegisterOffset = INDEX_NONE) const
+	{
+		ensure(InRegister.ElementCount > 0);
+		return (const T*)GetData(InRegister, InRegisterOffset);
+	}
+
+	// Returns the current const typed pointer for a given register index.
+	// Note: This refers to the active slice - and can change over time.
+	template<class T>
+	FORCEINLINE const T* Get(int32 InRegisterIndex, int32 InRegisterOffset = INDEX_NONE) const
+	{
+		ensure(Registers.IsValidIndex(InRegisterIndex));
+		const FRigVMRegister& Register = Registers[InRegisterIndex];
+		return Get<T>(Register, InRegisterOffset);
+	}
+
+	// Returns the current const typed pointer for a given argument
+	// Note: This refers to the active slice - and can change over time.
+	template<class T>
+	FORCEINLINE const T* Get(const FRigVMOperand& InOperand) const
+	{
+		return Get<T>(InOperand.GetRegisterIndex(), InOperand.GetRegisterOffset());
+	}
+
+	// Returns the current const typed reference for a given register
+	// Note: This refers to the active slice - and can change over time.
+	template<class T>
+	FORCEINLINE const T& GetRef(const FRigVMRegister& Register, int32 InRegisterOffset = INDEX_NONE) const
+	{
+		return *Get<T>(Register, InRegisterOffset);
+	}
+
+	// Returns the current const typed reference for a given register index
+	// Note: This refers to the active slice - and can change over time.
+	template<class T>
+	FORCEINLINE const T& GetRef(int32 InRegisterIndex, int32 InRegisterOffset = INDEX_NONE) const
+	{
+		return *Get<T>(InRegisterIndex, InRegisterOffset);
+	}
+
+	// Returns the current const typed reference for a given argument
+	// Note: This refers to the active slice - and can change over time.
+	template<class T>
+	FORCEINLINE const T& GetRef(const FRigVMOperand& InOperand) const
+	{
+		return GetRef<T>(InOperand.GetRegisterIndex(), InOperand.GetRegisterOffset());
+	}
+
+	// Returns the current typed pointer for a given register
+	// Note: This refers to the active slice - and can change over time.
+	template<class T>
+	FORCEINLINE T* Get(FRigVMRegister& InRegister, int32 InRegisterOffset = INDEX_NONE, bool bMoveToNextSlice = false)
+	{
+		ensure(InRegister.ElementCount > 0);
+		if(bMoveToNextSlice)
+		{
+			InRegister.MoveToNextSlice();
+		}
+		return (T*)GetData(InRegister, InRegisterOffset);
+	}
+
+	// Returns the current typed pointer for a given register index
+	// Note: This refers to the active slice - and can change over time.
+	template<class T>
+	FORCEINLINE T* Get(int32 InRegisterIndex, int32 InRegisterOffset = INDEX_NONE, bool bMoveToNextSlice = false)
+	{
+		ensure(Registers.IsValidIndex(InRegisterIndex));
+		FRigVMRegister& Register = Registers[InRegisterIndex];
+		return Get<T>(Register, InRegisterOffset, bMoveToNextSlice);
+	}
+
+	// Returns the current typed pointer for a given argument
+	// Note: This refers to the active slice - and can change over time.
+	template<class T>
+	FORCEINLINE T* Get(const FRigVMOperand& InOperand, bool bMoveToNextSlice = false)
+	{
+		return Get<T>(InOperand.GetRegisterIndex(), InOperand.GetRegisterOffset(), bMoveToNextSlice);
+	}
+
+	// Returns the current typed reference for a given register
+	// Note: This refers to the active slice - and can change over time.
+	template<class T>
+	FORCEINLINE T& GetRef(FRigVMRegister& Register, int32 InRegisterOffset = INDEX_NONE, bool bMoveToNextSlice = false)
+	{
+		return *Get<T>(Register, InRegisterOffset, bMoveToNextSlice);
+	}
+
+	// Returns the current typed reference for a given register index
+	// Note: This refers to the active slice - and can change over time.
+	template<class T>
+	FORCEINLINE T& GetRef(int32 InRegisterIndex, int32 InRegisterOffset = INDEX_NONE, bool bMoveToNextSlice = false)
+	{
+		return *Get<T>(InRegisterIndex, InRegisterOffset, bMoveToNextSlice);
+	}
+
+	// Returns the current typed reference for a given argument
+	// Note: This refers to the active slice - and can change over time.
+	template<class T>
+	FORCEINLINE T& GetRef(const FRigVMOperand& InOperand, bool bMoveToNextSlice = false)
+	{
+		return GetRef<T>(InOperand.GetRegisterIndex(), InOperand.GetRegisterOffset(), bMoveToNextSlice);
+	}
+
+	// Returns an array view for all elements of the current slice for a given register.
+	template<class T>
+	FORCEINLINE TArrayView<T> GetArray(FRigVMRegister& InRegister, int32 InRegisterOffset = INDEX_NONE, bool bMoveToNextSlice = false)
+	{
+		ensure(InRegister.ElementCount > 0);
+		if (InRegisterOffset == INDEX_NONE)
+		{
+			return TArrayView<T>((T*)GetData(InRegister, INDEX_NONE, bMoveToNextSlice), InRegister.ElementCount);
+		}
+		TArray<T>* StoredArray = (TArray<T>*)GetData(InRegister, InRegisterOffset, bMoveToNextSlice);
+		return TArrayView<T>(StoredArray->GetData(), StoredArray->Num());
+	}
+
+	// Returns an array view for all elements of the current slice for a given register index.
+	template<class T>
+	FORCEINLINE TArrayView<T> GetArray(int32 InRegisterIndex, int32 InRegisterOffset = INDEX_NONE, bool bMoveToNextSlice = false)
+	{
+		ensure(Registers.IsValidIndex(InRegisterIndex));
+		FRigVMRegister& Register = Registers[InRegisterIndex];
+		return GetArray<T>(Register, InRegisterOffset, bMoveToNextSlice);
+	}
+	
+	// Returns an array view for all elements of the current slice for a given argument.
+	template<class T>
+	FORCEINLINE TArrayView<T> GetArray(const FRigVMOperand& InOperand, bool bMoveToNextSlice = false)
+	{
+		return GetArray<T>(InOperand.GetRegisterIndex(), InOperand.GetRegisterOffset(), bMoveToNextSlice);
+	}
+
+	// Returns the script struct used for a given register (can be nullptr for non-struct-registers).
+	FORCEINLINE UScriptStruct* GetScriptStruct(const FRigVMRegister& Register) const
+	{
+		if (Register.ScriptStructIndex != INDEX_NONE)
+		{
+			ensure(ScriptStructs.IsValidIndex(Register.ScriptStructIndex));
+			return ScriptStructs[Register.ScriptStructIndex];
+		}
+		return nullptr;
+	}
+
+	// Returns the script struct used for a given register index (can be nullptr for non-struct-registers).
+	FORCEINLINE UScriptStruct* GetScriptStruct(int32 InRegisterIndex, int32 InRegisterOffset = INDEX_NONE) const
+	{
+		if (InRegisterOffset == INDEX_NONE)
+		{
+			ensure(Registers.IsValidIndex(InRegisterIndex));
+			const FRigVMRegister& Register = Registers[InRegisterIndex];
+			return GetScriptStruct(Register);
+		}
+		ensure(RegisterOffsets.IsValidIndex(InRegisterOffset));
+		const FRigVMRegisterOffset& Path = RegisterOffsets[InRegisterOffset];
+		return Path.GetScriptStruct();
+	}
+
+	// Copies the content of a source register to a target register.
+	// The source register can optionally be referencing a specific source memory container.
+	bool Copy(
+		int32 InSourceRegisterIndex,
+		int32 InTargetRegisterIndex,
+		const FRigVMMemoryContainer* InSourceMemory = nullptr,
+		int32 InSourceRegisterOffset = INDEX_NONE,
+		int32 InTargetRegisterOffset = INDEX_NONE);
+
+	// Copies the content of a source register to a target register.
+	// The source register can optionally be referencing a specific source memory container.
+	// Note: This only works if SupportsNames() == true
+	bool Copy(
+		const FName& InSourceName,
+		const FName& InTargetName,
+		const FRigVMMemoryContainer* InSourceMemory = nullptr,
+		int32 InSourceRegisterOffset = INDEX_NONE,
+		int32 InTargetRegisterOffset = INDEX_NONE);
+
+	// Copies the content of a source register to a target register.
+	// The source register can optionally be referencing a specific source memory container.
+	bool Copy(
+		const FRigVMOperand& InSourceOperand,
+		const FRigVMOperand& InTargetOperand,
+		const FRigVMMemoryContainer* InSourceMemory = nullptr);
+
+	// Returns the index of a register based on the register name.
+	// Note: This only works if SupportsNames() == true
+	FORCEINLINE int32 GetIndex(const FName& InName) const
+	{
+		if (!bUseNameMap)
+		{
+			return INDEX_NONE;
+		}
+
+		if (NameMap.Num() != Registers.Num())
+		{
+			for (int32 Index = 0; Index < Registers.Num(); Index++)
+			{
+				if (Registers[Index].Name == InName)
+				{
+					return Index;
+				}
+			}
+		}
+		else
+		{
+			const int32* Index = NameMap.Find(InName);
+			if (Index != nullptr)
+			{
+				return *Index;
+			}
+		}
+
+		return INDEX_NONE;
+	}
+
+	// Returns true if a given name is available for a new register.
+	FORCEINLINE bool IsNameAvailable(const FName& InPotentialNewName) const
+	{
+		if (!bUseNameMap)
+		{
+			return false;
+		}
+		return GetIndex(InPotentialNewName) == INDEX_NONE;
+	}
+
+	// Adds a new named register for a typed plain array from an array view (used by compiler)
+	template<class T>
+	FORCEINLINE int32 AddPlainArray(const FName& InNewName, TArrayView<T> InArrayView, int32 InSliceCount = 1)
+	{
+		return AddPlainArray(InNewName, sizeof(T), InArrayView.Num(), (const uint8*)InArrayView.GetData(), InSliceCount);
+	}
+
+	// Adds a new unnamed register for a typed plain array from an array view.
+	template<class T>
+	FORCEINLINE int32 AddPlainArray(TArrayView<T> InArrayView, int32 InSliceCount = 1)
+	{
+		return AddPlainArray<T>(NAME_None, InArrayView, InSliceCount);
+	}
+
+	// Adds a new named register for a typed plain value from a value reference.
+	template<class T>
+	FORCEINLINE int32 AddPlain(const FName& InNewName, const T& InValue, int32 InSliceCount = 1)
+	{
+		TArrayView<T> ArrayView((T*)&InValue, 1);
+		return AddPlainArray<T>(InNewName, ArrayView, InSliceCount);
+	}
+
+	// Adds a new unnamed register for a typed plain value from a value reference.
+	template<class T>
+	FORCEINLINE int32 AddPlain(const T& InValue, int32 InSliceCount = 1)
+	{
+		return AddPlain<T>(NAME_None, InValue, InSliceCount);
+	}
+
+	// Adds a new named register for a typed struct array given a value reference.
+	template<class T>
+	FORCEINLINE int32 AddStructArray(const FName& InNewName, TArrayView<T> InArrayView, int32 InSliceCount = 1)
+	{
+		UScriptStruct* Struct = T::StaticStruct();
+		if (Struct == nullptr)
+		{
+			return INDEX_NONE;
+		}
+		return AddStructArray(InNewName, Struct, InArrayView.Num(), (uint8*)InArrayView.GetData(), InSliceCount);
+	}
+
+	// Adds a new named register for a typed struct array given a TArray
+	template<class T>
+	FORCEINLINE int32 AddStructArray(TArrayView<T> InArrayView, int32 InSliceCount = 1)
+	{
+		return AddStructArray<T>(NAME_None, InArrayView, InSliceCount);
+	}
+
+	// Adds a new named register for a typed struct given a value reference
+	template<class T>
+	FORCEINLINE int32 AddStruct(const FName& InNewName, const T& InValue, int32 InSliceCount = 1)
+	{
+		TArrayView<T> ArrayView((T*)&InValue, 1);
+		return AddStructArray<T>(InNewName, ArrayView, InSliceCount);
+	}
+
+	// Adds a new unnamed register for a typed struct given a value reference
+	template<class T>
+	FORCEINLINE int32 AddStruct(const T& InValue, int32 InSliceCount = 1)
+	{
+		return AddStruct<T>(NAME_None, InValue, InSliceCount);
+	}
+
+	// Remove a register given its index
+	bool Remove(int32 InRegisterIndex);
+
+	// Remove a register given its name
+	// Note: This only works if SupportsNames() == true
+	bool Remove(const FName& InRegisterName);
+
+	// Remove a register given its index
+	// Note: This only works if SupportsNames() == true
+	FName Rename(int32 InRegisterIndex, const FName& InNewName);
+
+	// Remove a register given its old name
+	// Note: This only works if SupportsNames() == true
+	FName Rename(const FName& InOldName, const FName& InNewName);
+
+	// Resize a register given its index
+	bool Resize(int32 InRegisterIndex, int32 InNewElementCount, int32 InNewSliceCount = 1);
+
+	// Resize a register given its name
+	// Note: This only works if SupportsNames() == true
+	bool Resize(const FName& InRegisterName, int32 InNewElementCount, int32 InNewSliceCount = 1);
+
+	// Changes the type of a register provided a new type and a data pointer.
+	// This can be used to store different content in a register for re-using the memory,
+	// but this only works for new elements smaller or of the same size (since the register can't grow).
+	bool ChangeRegisterType(int32 InRegisterIndex, ERigVMRegisterType InNewType, int32 InElementSize, const uint8* InDataPtr = nullptr, int32 InNewElementCount = 1, int32 InNewSliceCount = 1);
+
+	// Changes the type of a register provided a new type and a data pointer.
+	// This can be used to store different content in a register for re-using the memory,
+	// but this only works for new elements smaller or of the same size (since the register can't grow).
+	template<class T>
+	FORCEINLINE bool ChangeRegisterType(int32 InRegisterIndex, const T* InDataPtr = nullptr, int32 InNewElementCount = 1, int32 InNewSliceCount = 1)
+	{
+		return ChangeRegisterType(InRegisterIndex, ERigVMRegisterType::Plain, sizeof(T), (uint8*)InDataPtr, InNewElementCount, InNewSliceCount);
+	}
+
+	// Adds a register path and returns its index
+	int32 GetOrAddRegisterOffset(int32 InRegisterIndex, UScriptStruct* InScriptStruct, const FString& InSegmentPath, int32 InInitialOffset = 0, int32 InElementSize = 0);
+
+private:
+
+	// Adds a new named register for a plain array from a data pointer (used by compiler)
+	FORCEINLINE int32 AddPlainArray(const FName& InNewName, int32 InElementSize, int32 InCount, const uint8* InDataPtr = nullptr, int32 InSliceCount = 1)
+	{
+		return Allocate(InNewName, InElementSize, InCount, InSliceCount, InDataPtr);
+	}
+
+	// Adds a new named register for a struct given a scriptstruct and a data pointer (used by compiler)
+	FORCEINLINE int32 AddStructArray(const FName& InNewName, UScriptStruct* InScriptStruct, int32 InCount, const uint8* InDataPtr = nullptr, int32 InSliceCount = 1)
+	{
+		int32 Register = Allocate(InNewName, InScriptStruct->GetStructureSize(), InCount, InSliceCount, nullptr, false);
+		if (Register == INDEX_NONE)
+		{
+			return INDEX_NONE;
+		}
+
+		Registers[Register].Type = ERigVMRegisterType::Struct;
+		Registers[Register].ScriptStructIndex = FindOrAddScriptStruct(InScriptStruct);
+
+		UpdateRegisters();
+
+		// construct the content
+		Construct(Register);
+
+		// copy values from the provided data
+		if (InDataPtr != nullptr)
+		{
+			Registers[Register].MoveToFirstSlice();
+			for (uint16 SliceIndex = 0; SliceIndex < Registers[Register].SliceCount; SliceIndex++)
+			{
+				InScriptStruct->CopyScriptStruct(GetData(Register), InDataPtr, InCount);
+				Registers[Register].MoveToNextSlice();
+			}
+			Registers[Register].MoveToFirstSlice();
+		}
+
+		return Register;
+	}
+
+	// Updates internal data for topological changes
+	void UpdateRegisters();
+
+	// Allocates a new named register
+	int32 Allocate(const FName& InNewName, int32 InElementSize, int32 InElementCount, int32 InSliceCount, const uint8* InDataPtr = nullptr, bool bUpdateRegisters = true);
+
+	// Allocates a new unnamed register
+	int32 Allocate(int32 InElementSize, int32 InElementCount, int32 InSliceCount, const uint8* InDataPtr = nullptr, bool bUpdateRegisters = true);
+
+	// Performs optional construction of data within a struct register
+	bool Construct(int32 InRegisterIndex, int32 InElementIndex = INDEX_NONE);
+
+	// Performs optional destruction of data within a struct register
+	bool Destroy(int32 InRegisterIndex, int32 InElementIndex = INDEX_NONE);
+
+	// Fills a register with zero memory
+	void FillWithZeroes(int32 InRegisterIndex);
+
+	// Ensures to add a script struct to the internal map if needed
+	int32 FindOrAddScriptStruct(UScriptStruct* InScriptStruct);
+
+	UPROPERTY()
+	bool bUseNameMap;
+
+	UPROPERTY()
+	ERigVMMemoryType MemoryType;
+
+	UPROPERTY()
+	TArray<FRigVMRegister> Registers;
+
+	UPROPERTY()
+	TArray<FRigVMRegisterOffset> RegisterOffsets;
+
+	UPROPERTY(transient)
+	TArray<uint8> Data;
+
+	UPROPERTY(transient)
+	TArray<UScriptStruct*> ScriptStructs;
+
+	UPROPERTY(transient)
+	TMap<FName, int32> NameMap;
+
+	friend class URigVM;
+	friend class URigVMCompiler;
+};
+
+template<>
+FORCEINLINE bool FRigVMMemoryContainer::ChangeRegisterType(int32 InRegisterIndex, const FName* InDataPtr, int32 InNewElementCount, int32 InNewSliceCount)
+{
+	if (!ChangeRegisterType(InRegisterIndex, ERigVMRegisterType::Name, sizeof(FName), nullptr, InNewElementCount, InNewSliceCount))
+	{
+		return false;
+	}
+
+	if (InDataPtr)
+	{
+		Registers[InRegisterIndex].MoveToFirstSlice();
+		for (uint16 SliceIndex = 0; SliceIndex < Registers[InRegisterIndex].SliceCount; SliceIndex++)
+		{
+			FName* DataPtr = (FName*)GetData(InRegisterIndex);
+			for (int32 Index = 0; Index < InNewElementCount; Index++)
+			{
+				DataPtr[Index] = InDataPtr[Index];
+			}
+			Registers[InRegisterIndex].MoveToNextSlice();
+		}
+		Registers[InRegisterIndex].MoveToFirstSlice();
+	}
+
+	return true;
+}
+
+template<>
+FORCEINLINE bool FRigVMMemoryContainer::ChangeRegisterType(int32 InRegisterIndex, const FString* InDataPtr, int32 InNewElementCount, int32 InNewSliceCount)
+{
+	if (!ChangeRegisterType(InRegisterIndex, ERigVMRegisterType::String, sizeof(FString), nullptr, InNewElementCount, InNewSliceCount))
+	{
+		return false;
+	}
+
+	if (InDataPtr)
+	{
+		Registers[InRegisterIndex].MoveToFirstSlice();
+		for (uint16 SliceIndex = 0; SliceIndex < Registers[InRegisterIndex].SliceCount; SliceIndex++)
+		{
+			FString* DataPtr = (FString*)GetData(InRegisterIndex);
+			for (int32 Index = 0; Index < InNewElementCount; Index++)
+			{
+				DataPtr[Index] = InDataPtr[Index];
+			}
+			Registers[InRegisterIndex].MoveToNextSlice();
+		}
+		Registers[InRegisterIndex].MoveToFirstSlice();
+	}
+
+	return true;
+}
+
+template<>
+FORCEINLINE int32 FRigVMMemoryContainer::AddPlainArray(const FName& InNewName, TArrayView<FRotator> InArrayView, int32 InSliceCount)
+{
+	return AddStructArray(InNewName, TBaseStructure<FRotator>::Get(), InArrayView.Num(), (uint8*)InArrayView.GetData(), InSliceCount);
+}
+
+template<>
+FORCEINLINE int32 FRigVMMemoryContainer::AddPlainArray(const FName& InNewName, TArrayView<FQuat> InArrayView, int32 InSliceCount)
+{
+	return AddStructArray(InNewName, TBaseStructure<FQuat>::Get(), InArrayView.Num(), (uint8*)InArrayView.GetData(), InSliceCount);
+}
+
+template<>
+FORCEINLINE int32 FRigVMMemoryContainer::AddPlainArray(const FName& InNewName, TArrayView<FTransform> InArrayView, int32 InSliceCount)
+{
+	return AddStructArray(InNewName, TBaseStructure<FTransform>::Get(), InArrayView.Num(), (uint8*)InArrayView.GetData(), InSliceCount);
+}
+
+template<>
+FORCEINLINE int32 FRigVMMemoryContainer::AddPlainArray(const FName& InNewName, TArrayView<FLinearColor> InArrayView, int32 InSliceCount)
+{
+	return AddStructArray(InNewName, TBaseStructure<FLinearColor>::Get(), InArrayView.Num(), (uint8*)InArrayView.GetData(), InSliceCount);
+}
+
+template<>
+FORCEINLINE int32 FRigVMMemoryContainer::AddPlainArray(const FName& InNewName, TArrayView<FColor> InArrayView, int32 InSliceCount)
+{
+	return AddStructArray(InNewName, TBaseStructure<FColor>::Get(), InArrayView.Num(), (uint8*)InArrayView.GetData(), InSliceCount);
+}
+
+template<>
+FORCEINLINE int32 FRigVMMemoryContainer::AddPlainArray(const FName& InNewName, TArrayView<FPlane> InArrayView, int32 InSliceCount)
+{
+	return AddStructArray(InNewName, TBaseStructure<FPlane>::Get(), InArrayView.Num(), (uint8*)InArrayView.GetData(), InSliceCount);
+}
+
+template<>
+FORCEINLINE int32 FRigVMMemoryContainer::AddPlainArray(const FName& InNewName, TArrayView<FVector> InArrayView, int32 InSliceCount)
+{
+	return AddStructArray(InNewName, TBaseStructure<FVector>::Get(), InArrayView.Num(), (uint8*)InArrayView.GetData(), InSliceCount);
+}
+
+template<>
+FORCEINLINE int32 FRigVMMemoryContainer::AddPlainArray(const FName& InNewName, TArrayView<FVector2D> InArrayView, int32 InSliceCount)
+{
+	return AddStructArray(InNewName, TBaseStructure<FVector2D>::Get(), InArrayView.Num(), (uint8*)InArrayView.GetData(), InSliceCount);
+}
+
+template<>
+FORCEINLINE int32 FRigVMMemoryContainer::AddPlainArray(const FName& InNewName, TArrayView<FVector4> InArrayView, int32 InSliceCount)
+{
+	return AddStructArray(InNewName, TBaseStructure<FVector4>::Get(), InArrayView.Num(), (uint8*)InArrayView.GetData(), InSliceCount);
+}
+
+template<>
+FORCEINLINE int32 FRigVMMemoryContainer::AddPlainArray(const FName& InNewName, TArrayView<FName> InArrayView, int32 InSliceCount)
+{
+	int32 Register = Allocate(InNewName, sizeof(FName), InArrayView.Num(), InSliceCount, nullptr);
+	Registers[Register].Type = ERigVMRegisterType::Name;
+
+	Construct(Register);
+
+	if (InArrayView.Num() > 0)
+	{
+		Registers[Register].MoveToFirstSlice();
+		for (uint16 SliceIndex = 0; SliceIndex < Registers[Register].SliceCount; SliceIndex++)
+		{
+			FName* DataPtr = (FName*)GetData(Register);
+			for (int32 Index = 0; Index < InArrayView.Num(); Index++)
+			{
+				DataPtr[Index] = InArrayView[Index];
+			}
+			Registers[Register].MoveToNextSlice();
+		}
+		Registers[Register].MoveToFirstSlice();
+	}
+
+	return Register;
+}
+
+template<>
+FORCEINLINE int32 FRigVMMemoryContainer::AddPlainArray(const FName& InNewName, TArrayView<FString> InArrayView, int32 InSliceCount)
+{
+	int32 Register = Allocate(InNewName, sizeof(FString), InArrayView.Num(), InSliceCount, nullptr);
+	Registers[Register].Type = ERigVMRegisterType::String;
+
+	Construct(Register);
+
+	if (InArrayView.Num() > 0)
+	{
+		Registers[Register].MoveToFirstSlice();
+		for (uint16 SliceIndex = 0; SliceIndex < Registers[Register].SliceCount; SliceIndex++)
+		{
+			FString* DataPtr = (FString*)GetData(Register);
+			for (int32 Index = 0; Index < InArrayView.Num(); Index++)
+			{
+				DataPtr[Index] = InArrayView[Index];
+			}
+			Registers[Register].MoveToNextSlice();
+		}
+		Registers[Register].MoveToFirstSlice();
+	}
+
+	return Register;
+}
+
+typedef FRigVMMemoryContainer* FRigVMMemoryContainerPtr;
+typedef TArrayView<FRigVMMemoryContainer*> FRigVMMemoryContainerPtrArray;
\ No newline at end of file
