RigVM - Rename/move file(s) - introduced new RigVM module

diff --git a/Engine/Plugins/Experimental/ControlRig/Source/ControlRig/ControlRig.Build.cs b/Engine/Plugins/Experimental/ControlRig/Source/ControlRig/ControlRig.Build.cs
--- a/Engine/Plugins/Experimental/ControlRig/Source/ControlRig/ControlRig.Build.cs
+++ b/Engine/Plugins/Experimental/ControlRig/Source/ControlRig/ControlRig.Build.cs
@@ -21,7 +21,8 @@ namespace UnrealBuildTool.Rules
                     "MovieScene",
                     "MovieSceneTracks",
                     "PropertyPath",
-					"TimeManagement"
+					"TimeManagement",
+                    "RigVM"
                 }
             );
 
diff --git a/Engine/Source/Runtime/AnimationCore/Private/MultiplexByteCode.cpp b/Engine/Source/Runtime/AnimationCore/Private/MultiplexByteCode.cpp
deleted file mode 100644
diff --git a/Engine/Source/Runtime/AnimationCore/Private/MultiplexRegistry.cpp b/Engine/Source/Runtime/AnimationCore/Private/MultiplexRegistry.cpp
deleted file mode 100644
diff --git a/Engine/Source/Runtime/AnimationCore/Private/MultiplexStorage.cpp b/Engine/Source/Runtime/AnimationCore/Private/MultiplexStorage.cpp
deleted file mode 100644
diff --git a/Engine/Source/Runtime/AnimationCore/Private/MultiplexVM.cpp b/Engine/Source/Runtime/AnimationCore/Private/MultiplexVM.cpp
deleted file mode 100644
diff --git a/Engine/Source/Runtime/AnimationCore/Public/MultiplexByteCode.h b/Engine/Source/Runtime/AnimationCore/Public/MultiplexByteCode.h
deleted file mode 100644
diff --git a/Engine/Source/Runtime/AnimationCore/Public/MultiplexRegistry.h b/Engine/Source/Runtime/AnimationCore/Public/MultiplexRegistry.h
deleted file mode 100644
diff --git a/Engine/Source/Runtime/AnimationCore/Public/MultiplexStorage.h b/Engine/Source/Runtime/AnimationCore/Public/MultiplexStorage.h
deleted file mode 100644
diff --git a/Engine/Source/Runtime/AnimationCore/Public/MultiplexVM.h b/Engine/Source/Runtime/AnimationCore/Public/MultiplexVM.h
deleted file mode 100644
diff --git a/Engine/Source/Runtime/Engine/Engine.Build.cs b/Engine/Source/Runtime/Engine/Engine.Build.cs
--- a/Engine/Source/Runtime/Engine/Engine.Build.cs
+++ b/Engine/Source/Runtime/Engine/Engine.Build.cs
@@ -13,7 +13,7 @@ public class Engine : ModuleRules
 
 		SharedPCHHeaderFile = "Public/EngineSharedPCH.h";
 
-		PublicIncludePathModuleNames.AddRange(new string[] { "Renderer", "PacketHandler", "AudioMixer", "AnimationCore" });
+		PublicIncludePathModuleNames.AddRange(new string[] { "Renderer", "PacketHandler", "AudioMixer", "AnimationCore", "RigVM" });
 
 		PrivateIncludePaths.AddRange(
 			new string[] {
@@ -86,7 +86,8 @@ public class Engine : ModuleRules
                 "NetworkReplayStreaming",
 				"PhysicsCore",
 				"PhysicsSQ",
-				"ChaosSolvers"
+				"ChaosSolvers",
+				"RigVM"
 			}
 		);
 
diff --git a/Engine/Source/Runtime/RigVM/Private/MultiplexByteCode.cpp b/Engine/Source/Runtime/RigVM/Private/MultiplexByteCode.cpp
new file mode 100644
--- /dev/null
+++ b/Engine/Source/Runtime/RigVM/Private/MultiplexByteCode.cpp
@@ -0,0 +1,161 @@
+// Copyright 1998-2019 Epic Games, Inc. All Rights Reserved.
+
+#include "MultiplexByteCode.h"
+
+FMultiplexByteCodeTable::FMultiplexByteCodeTable()
+{
+}
+
+FMultiplexByteCodeTable::FMultiplexByteCodeTable(const FMultiplexByteCode& InByteCode)
+{
+	uint64 Address = 0;
+	while (Address < InByteCode.Num())
+	{
+		EMultiplexOpCode OpCode = InByteCode.GetOpCodeAt(Address);
+		Entries.Add(FMultiplexByteCodeTableEntry(OpCode, Address));
+		Address += InByteCode.GetOpNumBytesAt(Address);
+	}
+}
+
+void FMultiplexByteCodeTable::Reset()
+{
+	Entries.Reset();
+}
+
+FMultiplexByteCode::FMultiplexByteCode()
+{
+}
+
+void FMultiplexByteCode::Reset()
+{
+	ByteCode.Reset();
+}
+
+uint64 FMultiplexByteCode::Num() const
+{
+	return (uint64)ByteCode.Num();
+}
+
+uint64 FMultiplexByteCode::GetOpNumBytesAt(uint64 InByteCodeIndex, bool bIncludeArguments) const
+{
+	EMultiplexOpCode OpCode = GetOpCodeAt(InByteCodeIndex);
+	switch (OpCode)
+	{
+		case EMultiplexOpCode::Copy:
+		{
+			return (uint64)sizeof(FMultiplexCopyOp);
+		}
+		case EMultiplexOpCode::Increment:
+		{
+			return (uint64)sizeof(FMultiplexIncrementOp);
+		}
+		case EMultiplexOpCode::Decrement:
+		{
+			return (uint64)sizeof(FMultiplexDecrementOp);
+		}
+		case EMultiplexOpCode::Equals:
+		{
+			return (uint64)sizeof(FMultiplexEqualsOp);
+		}
+		case EMultiplexOpCode::NotEquals:
+		{
+			return (uint64)sizeof(FMultiplexNotEqualsOp);
+		}
+		case EMultiplexOpCode::Jump:
+		{
+			return (uint64)sizeof(FMultiplexJumpOp);
+		}
+		case EMultiplexOpCode::JumpIfTrue:
+		{
+			return (uint64)sizeof(FMultiplexJumpIfTrueOp);
+		}
+		case EMultiplexOpCode::JumpIfFalse:
+		{
+			return (uint64)sizeof(FMultiplexJumpIfFalseOp);
+		}
+		case EMultiplexOpCode::Execute:
+		{
+			uint64 NumBytes = (uint64)sizeof(FMultiplexExecuteOp);
+			if(bIncludeArguments)
+			{
+				const FMultiplexExecuteOp& ExecuteOp = GetOpAt<FMultiplexExecuteOp>(InByteCodeIndex);
+				NumBytes += (uint64)ExecuteOp.ArgumentCount * (uint64)sizeof(FMultiplexArgument);
+			}
+			return NumBytes;
+		}
+		case EMultiplexOpCode::Exit:
+		{
+			return (uint64)sizeof(FMultiplexExitOp);
+		}
+		case EMultiplexOpCode::Invalid:
+		{
+			ensure(false);
+			return 0;
+		}
+	}
+	return 0;
+}
+
+uint64 FMultiplexByteCode::AddCopyOp(const FMultiplexArgument& InSource, const FMultiplexArgument& InTarget, int32 InSourceOffset, int32 InTargetOffset, int32 InNumBytes)
+{
+	FMultiplexCopyOp Op(InSource, InTarget, InSourceOffset, InTargetOffset, InNumBytes);
+	return AddOp(Op);
+}
+
+uint64 FMultiplexByteCode::AddIncrementOp(const FMultiplexArgument& InArg)
+{
+	FMultiplexIncrementOp Op(InArg);
+	return AddOp(Op);
+}
+
+uint64 FMultiplexByteCode::AddDecrementOp(const FMultiplexArgument& InArg)
+{
+	FMultiplexDecrementOp Op(InArg);
+	return AddOp(Op);
+}
+
+uint64 FMultiplexByteCode::AddEqualsOp(const FMultiplexArgument& InA, const FMultiplexArgument& InB, const FMultiplexArgument& InResult)
+{
+	FMultiplexEqualsOp Op(InA, InB, InResult);
+	return AddOp(Op);
+}
+
+uint64 FMultiplexByteCode::AddNotEqualsOp(const FMultiplexArgument& InA, const FMultiplexArgument& InB, const FMultiplexArgument& InResult)
+{
+	FMultiplexNotEqualsOp Op(InA, InB, InResult);
+	return AddOp(Op);
+}
+
+uint64 FMultiplexByteCode::AddJumpOp(uint64 InByteCodeIndex)
+{
+	FMultiplexJumpOp Op(InByteCodeIndex);
+	return AddOp(Op);
+}
+
+uint64 FMultiplexByteCode::AddJumpIfTrueOp(uint64 InByteCodeIndex, const FMultiplexArgument& InCondition)
+{
+	FMultiplexJumpIfTrueOp Op(InByteCodeIndex, InCondition);
+	return AddOp(Op);
+}
+
+uint64 FMultiplexByteCode::AddJumpIfFalseOp(uint64 InByteCodeIndex, const FMultiplexArgument& InCondition)
+{
+	FMultiplexJumpIfFalseOp Op(InByteCodeIndex, InCondition);
+	return AddOp(Op);
+}
+
+uint64 FMultiplexByteCode::AddExecuteOp(uint16 InFunctionIndex, const TArrayView<FMultiplexArgument>& InArguments)
+{
+	FMultiplexExecuteOp Op(InFunctionIndex, (uint16)InArguments.Num());
+	uint64 OpAddress = AddOp(Op);
+	uint64 ArgumentsAddress = (uint64)ByteCode.AddUninitialized(sizeof(FMultiplexArgument) * InArguments.Num());
+	FMemory::Memcpy(ByteCode.GetData() + ArgumentsAddress, InArguments.GetData(), sizeof(FMultiplexArgument) * InArguments.Num());
+	return OpAddress;
+}
+
+uint64 FMultiplexByteCode::AddExitOp()
+{
+	FMultiplexExitOp Op;
+	return AddOp(Op);
+}
+
diff --git a/Engine/Source/Runtime/RigVM/Private/MultiplexRegistry.cpp b/Engine/Source/Runtime/RigVM/Private/MultiplexRegistry.cpp
new file mode 100644
--- /dev/null
+++ b/Engine/Source/Runtime/RigVM/Private/MultiplexRegistry.cpp
@@ -0,0 +1,37 @@
+// Copyright 1998-2019 Epic Games, Inc. All Rights Reserved.
+
+#include "MultiplexRegistry.h"
+#include "UObject/UObjectIterator.h"
+
+FMultiplexRegistry FMultiplexRegistry::s_MultiplexRegistry;
+
+FMultiplexRegistry& FMultiplexRegistry::Get()
+{
+	return s_MultiplexRegistry;
+}
+
+void FMultiplexRegistry::Refresh()
+{
+}
+
+void FMultiplexRegistry::Register(const TCHAR* InName, FMultiplexFunctionPtr InFunctionPtr)
+{
+	if (Find(InName) != nullptr)
+	{
+		return;
+	}
+	Functions.Add(FMultiplexFunction(InName, InFunctionPtr));
+}
+
+FMultiplexFunctionPtr FMultiplexRegistry::Find(const TCHAR* InName) const
+{
+	for (const FMultiplexFunction& Function : Functions)
+	{
+		if (FCString::Strcmp(Function.Name, InName) == 0)
+		{
+			return Function.FunctionPtr;
+		}
+	}
+
+	return nullptr;
+}
diff --git a/Engine/Source/Runtime/RigVM/Private/MultiplexStorage.cpp b/Engine/Source/Runtime/RigVM/Private/MultiplexStorage.cpp
new file mode 100644
--- /dev/null
+++ b/Engine/Source/Runtime/RigVM/Private/MultiplexStorage.cpp
@@ -0,0 +1,523 @@
+// Copyright 1998-2019 Epic Games, Inc. All Rights Reserved.
+
+#include "MultiplexStorage.h"
+
+FMultiplexStorage::FMultiplexStorage(bool bInUseNames)
+	:bUseNameMap(bInUseNames)
+{
+}
+
+FMultiplexStorage::FMultiplexStorage(const FMultiplexStorage& Other)
+{
+	*this = Other;
+}
+
+FMultiplexStorage::~FMultiplexStorage()
+{
+	Reset();
+}
+
+FMultiplexStorage& FMultiplexStorage::operator= (const FMultiplexStorage &InOther)
+{
+	Reset();
+
+	bUseNameMap = InOther.bUseNameMap;
+	Data.Append(InOther.Data);
+	Addresses.Append(InOther.Addresses);
+	ScriptStructs.Append(InOther.ScriptStructs);
+
+	UpdateAddresses();
+
+	for (int32 Index = 0; Index < Addresses.Num(); Index++)
+	{
+		Construct(Index);
+		Copy(Index, Index, &InOther);
+	}
+
+	return *this;
+}
+
+void FMultiplexStorage::Reset()
+{
+	for (int32 Index = 0; Index < Addresses.Num(); Index++)
+	{
+		Destroy(Index);
+	}
+
+	Data.Reset();
+	Addresses.Reset();
+	ScriptStructs.Reset();
+	NameMap.Reset();
+}
+
+bool FMultiplexStorage::Copy(
+	int32 InSourceAddressIndex,
+	int32 InTargetAddressIndex,
+	const FMultiplexStorage* InSourceStorage,
+	int32 InSourceByteOffset,
+	int32 InTargetByteOffset,
+	int32 InNumBytes)
+{
+	if (InSourceStorage == nullptr)
+	{
+		InSourceStorage = this;
+	}
+
+	ensure(InSourceStorage->Addresses.IsValidIndex(InSourceAddressIndex));
+	ensure(Addresses.IsValidIndex(InTargetAddressIndex));
+
+	if (InSourceAddressIndex == InTargetAddressIndex && InSourceByteOffset == InTargetByteOffset && this == InSourceStorage)
+	{
+		return false;
+	}
+
+	const FMultiplexAddress& Source = InSourceStorage->Addresses[InSourceAddressIndex];
+	const FMultiplexAddress& Target = Addresses[InTargetAddressIndex];
+
+	int32 SourceStartByte = Source.FirstByte();
+	int32 SourceNumBytes = Source.NumBytes(false);
+	if(InSourceByteOffset != INDEX_NONE)
+	{
+		ensure(InNumBytes != INDEX_NONE);
+		ensure(InSourceByteOffset >= 0 && InSourceByteOffset < Source.NumBytes(false));
+		ensure(InNumBytes > 0 && InSourceByteOffset + InNumBytes <= Source.NumBytes(false));
+
+		SourceStartByte += InSourceByteOffset;
+		SourceNumBytes = InNumBytes;
+	}
+
+	int32 TargetStartByte = Target.FirstByte();
+	int32 TargetNumBytes = Target.NumBytes(false);
+	if(InTargetByteOffset != INDEX_NONE)
+	{
+		ensure(InNumBytes != INDEX_NONE);
+		ensure(InTargetByteOffset >= 0 && InTargetByteOffset < Target.NumBytes(false));
+		ensure(InNumBytes > 0 && InTargetByteOffset + InNumBytes <= Target.NumBytes(false));
+
+		TargetStartByte += InTargetByteOffset;
+		TargetNumBytes = InNumBytes;
+	}
+
+	if (SourceNumBytes != TargetNumBytes)
+	{
+		return false;
+	}
+
+	switch (Target.Type)
+	{
+		case EMultiplexAddressType::Plain:
+		{
+			FMemory::Memcpy(&Data[TargetStartByte], &InSourceStorage->Data[SourceStartByte], TargetNumBytes);
+			break;
+		}
+		case EMultiplexAddressType::Struct:
+		{
+			UScriptStruct* ScriptStruct = GetScriptStruct(InTargetAddressIndex);
+			int32 NumStructs = TargetNumBytes / ScriptStruct->GetStructureSize();
+			ScriptStruct->CopyScriptStruct(&Data[TargetStartByte], &InSourceStorage->Data[SourceStartByte], NumStructs);
+			break;
+		}
+		case EMultiplexAddressType::Name:
+		{
+			int32 NumNames = TargetNumBytes / sizeof(FName);
+			TArrayView<FName> TargetNames((FName*)&Data[TargetStartByte], NumNames);
+			TArrayView<FName> SourceNames((FName*)&InSourceStorage->Data[SourceStartByte], NumNames);
+			for (int32 Index = 0; Index < NumNames; Index++)
+			{
+				TargetNames[Index] = SourceNames[Index];
+			}
+			break;
+		}
+		case EMultiplexAddressType::String:
+		{
+			int32 NumStrings = TargetNumBytes / sizeof(FString);
+			TArrayView<FString> TargetStrings((FString*)&Data[TargetStartByte], NumStrings);
+			TArrayView<FString> SourceStrings((FString*)&InSourceStorage->Data[SourceStartByte], NumStrings);
+			for (int32 Index = 0; Index < NumStrings; Index++)
+			{
+				TargetStrings[Index] = SourceStrings[Index];
+			}
+			break;
+		}
+		case EMultiplexAddressType::Invalid:
+		{
+			return false;
+		}
+	}
+
+	return true;
+}
+
+bool FMultiplexStorage::Copy(
+	const FName& InSourceName,
+	const FName& InTargetName,
+	const FMultiplexStorage* InSourceStorage,
+	int32 InSourceByteOffset,
+	int32 InTargetByteOffset,
+	int32 InNumBytes)
+{
+	ensure(bUseNameMap);
+
+	int32 SourceAddressIndex = GetIndex(InSourceName);
+	int32 TargetAddressIndex = GetIndex(InTargetName);
+
+	if(SourceAddressIndex == INDEX_NONE || TargetAddressIndex == INDEX_NONE)
+	{
+		return false;
+	}
+
+	return Copy(SourceAddressIndex, TargetAddressIndex, InSourceStorage, InSourceByteOffset, InTargetByteOffset, InNumBytes);
+}
+
+int32 FMultiplexStorage::Allocate(const FName& InNewName, int32 InElementSize, int32 InCount, const void* InDataPtr, bool bUpdateAddresses)
+{
+	FName Name = InNewName;
+	if (bUseNameMap && InNewName == NAME_None)
+	{
+		const TCHAR* AddressPrefix = TEXT("Address");
+		int32 AddressSuffix = 0;
+		do
+		{
+			Name = FName(*FString::Printf(TEXT("%s_%d"), AddressPrefix, AddressSuffix++));
+		} while (!IsNameAvailable(Name));
+	}
+
+	ensure(InElementSize > 0 && InCount > 0);
+
+	if (bUseNameMap)
+	{
+		if (!IsNameAvailable(Name))
+		{
+			return INDEX_NONE;
+		}
+	}
+
+	FMultiplexAddress NewAddress;
+	NewAddress.ByteIndex = Data.Num();
+	if (bUseNameMap)
+	{
+		NewAddress.Name = Name;
+	}
+	NewAddress.ElementSize = InElementSize;
+	NewAddress.ElementCount = InCount;
+	NewAddress.Type = EMultiplexAddressType::Plain;
+
+	Data.AddZeroed(NewAddress.NumBytes());
+
+	if (InDataPtr != nullptr)
+	{
+		FMemory::Memcpy(&Data[NewAddress.FirstByte()], InDataPtr, NewAddress.NumBytes());
+	}
+
+	int32 AddressIndex = Addresses.Num();
+	Addresses.Add(NewAddress);
+
+	if (bUpdateAddresses)
+	{
+		UpdateAddresses();
+	}
+	return AddressIndex;
+}
+
+int32 FMultiplexStorage::Allocate(int32 InElementSize, int32 InCount, const void* InDataPtr, bool bUpdateAddresses)
+{
+	return Allocate(NAME_None, InElementSize, InCount, InDataPtr, bUpdateAddresses);
+}
+
+bool FMultiplexStorage::Construct(int32 InAddressIndex, int32 InElementIndex)
+{
+	ensure(Addresses.IsValidIndex(InAddressIndex));
+
+	const FMultiplexAddress& Address = Addresses[InAddressIndex];
+	switch (Address.Type)
+	{
+		case EMultiplexAddressType::Struct:
+		{
+			void* DataPtr = (void*)&Data[InElementIndex == INDEX_NONE ? Address.FirstByte() : Address.FirstByte() + InElementIndex * Address.ElementSize];
+			int32 Count = InElementIndex == INDEX_NONE ? Address.ElementCount : 1;
+
+			UScriptStruct* ScriptStruct = GetScriptStruct(InAddressIndex);
+			ScriptStruct->InitializeStruct(DataPtr, Count);
+			break;
+		}
+		case EMultiplexAddressType::String:
+		{
+			FString* DataPtr = (FString*)&Data[InElementIndex == INDEX_NONE ? Address.FirstByte() : Address.FirstByte() + InElementIndex * Address.ElementSize];
+			int32 Count = InElementIndex == INDEX_NONE ? Address.ElementCount : 1;
+
+			FMemory::Memzero(DataPtr, Count * Address.ElementSize);
+			for (int32 Index = 0; Index < Count; Index++)
+			{
+				DataPtr[Index] = FString();
+			}
+			break;
+		}
+		case EMultiplexAddressType::Name:
+		{
+			FName* DataPtr = (FName*)&Data[InElementIndex == INDEX_NONE ? Address.FirstByte() : Address.FirstByte() + InElementIndex * Address.ElementSize];
+			int32 Count = InElementIndex == INDEX_NONE ? Address.ElementCount : 1;
+
+			FMemory::Memzero(DataPtr, Count * Address.ElementSize);
+			for (int32 Index = 0; Index < Count; Index++)
+			{
+				DataPtr[Index] = FName();
+			}
+			break;
+		}
+		default:
+		{
+			return false;
+		}
+	}
+
+	return true;
+}
+
+bool FMultiplexStorage::Destroy(int32 InAddressIndex, int32 InElementIndex)
+{
+	ensure(Addresses.IsValidIndex(InAddressIndex));
+
+	const FMultiplexAddress& Address = Addresses[InAddressIndex];
+	switch (Address.Type)
+	{
+		case EMultiplexAddressType::Struct:
+		{
+			void* DataPtr = (void*)&Data[InElementIndex == INDEX_NONE ? Address.FirstByte() : Address.FirstByte() + InElementIndex * Address.ElementSize];
+			int32 Count = InElementIndex == INDEX_NONE ? Address.ElementCount : 1;
+
+			UScriptStruct* ScriptStruct = GetScriptStruct(InAddressIndex);
+			ScriptStruct->DestroyStruct(DataPtr, Count);
+			break;
+		}
+		case EMultiplexAddressType::String:
+		{
+			FString* DataPtr = (FString*)&Data[InElementIndex == INDEX_NONE ? Address.FirstByte() : Address.FirstByte() + InElementIndex * Address.ElementSize];
+			int32 Count = InElementIndex == INDEX_NONE ? Address.ElementCount : 1;
+
+			for (int32 Index = 0; Index < Count; Index++)
+			{
+				DataPtr[Index].Empty();
+			}
+			break;
+		}
+		case EMultiplexAddressType::Name:
+		{
+			FName* DataPtr = (FName*)&Data[InElementIndex == INDEX_NONE ? Address.FirstByte() : Address.FirstByte() + InElementIndex * Address.ElementSize];
+			int32 Count = InElementIndex == INDEX_NONE ? Address.ElementCount : 1;
+
+			for (int32 Index = 0; Index < Count; Index++)
+			{
+				DataPtr[Index] = FName();
+			}
+			break;
+		}
+		default:
+		{
+			return false;
+		}
+	}
+
+	return true;
+}
+
+bool FMultiplexStorage::Remove(int32 InAddressIndex)
+{
+	if (InAddressIndex < 0 || InAddressIndex >= Addresses.Num())
+	{
+		return false;
+	}
+
+	Destroy(InAddressIndex);
+
+	FMultiplexAddress AddressToRemove = Addresses[InAddressIndex];
+	Data.RemoveAt(AddressToRemove.ByteIndex, AddressToRemove.NumBytes());
+	Addresses.RemoveAt(InAddressIndex);
+
+	for (int32 Index = InAddressIndex; Index < Addresses.Num(); Index++)
+	{
+		Addresses[Index].ByteIndex -= AddressToRemove.NumBytes();
+	}
+
+	UpdateAddresses();
+	return true;
+}
+
+bool FMultiplexStorage::Remove(const FName& InAddressName)
+{
+	ensure(bUseNameMap);
+	return Remove(GetIndex(InAddressName));
+}
+
+FName FMultiplexStorage::Rename(int32 InAddressIndex, const FName& InNewName)
+{
+	if (Addresses[InAddressIndex].Name == InNewName)
+	{
+		return Addresses[InAddressIndex].Name;
+	}
+
+	if (!IsNameAvailable(InNewName))
+	{
+		return Addresses[InAddressIndex].Name;
+	}
+
+	Addresses[InAddressIndex].Name = InNewName;
+	UpdateAddresses();
+
+	return InNewName;
+}
+
+FName FMultiplexStorage::Rename(const FName& InOldName, const FName& InNewName)
+{
+	ensure(bUseNameMap);
+
+	int32 AddressIndex = GetIndex(InOldName);
+	if (AddressIndex == INDEX_NONE)
+	{
+		return NAME_None;
+	}
+
+	return Rename(AddressIndex, InNewName);
+}
+
+bool FMultiplexStorage::Resize(int32 InAddressIndex, int32 InNewElementCount)
+{
+	if (InNewElementCount <= 0)
+	{
+		return Remove(InAddressIndex);
+	}
+
+	if (Addresses[InAddressIndex].ElementCount == InNewElementCount)
+	{
+		return false;
+	}
+
+	FMultiplexAddress& Address = Addresses[InAddressIndex];
+
+	if (Address.ElementCount > InNewElementCount) // shrink
+	{
+		int32 ElementsToRemove = Address.ElementCount - InNewElementCount;
+		int32 NumBytesToRemove = Address.ElementSize * ElementsToRemove;
+		int32 FirstByteToRemove = Address.FirstByte() + Address.ElementSize * InNewElementCount;
+
+		for (int32 ElementIndex = InNewElementCount; ElementIndex < Address.ElementCount; ElementIndex++)
+		{
+			Destroy(InAddressIndex, ElementIndex);
+		}
+
+		Data.RemoveAt(FirstByteToRemove, NumBytesToRemove);
+		Address.ElementCount = InNewElementCount;
+
+		for (int32 AddressIndex = InAddressIndex + 1; AddressIndex < Addresses.Num(); AddressIndex++)
+		{
+			Addresses[AddressIndex].ByteIndex -= NumBytesToRemove;
+		}
+	}
+	else // grow
+	{
+		int32 OldElementCount = Address.ElementCount;
+		int32 ElementsToAdd = InNewElementCount - Address.ElementCount;
+		int32 NumBytesToAdd = Address.ElementSize * ElementsToAdd;
+		int32 FirstByteToAdd = Address.FirstByte() + Address.ElementSize * Address.ElementCount;
+
+		Data.InsertZeroed(FirstByteToAdd, NumBytesToAdd);
+		Address.ElementCount = InNewElementCount;
+
+		for (int32 ElementIndex = OldElementCount; ElementIndex < InNewElementCount; ElementIndex++)
+		{
+			Construct(InAddressIndex, ElementIndex);
+		}
+
+		for (int32 AddressIndex = InAddressIndex + 1; AddressIndex < Addresses.Num(); AddressIndex++)
+		{
+			Addresses[AddressIndex].ByteIndex += NumBytesToAdd;
+		}
+	}
+
+	UpdateAddresses();
+	return true;
+}
+
+bool FMultiplexStorage::Resize(const FName& InAddressName, int32 InNewElementCount)
+{
+	ensure(bUseNameMap);
+
+	int32 AddressIndex = GetIndex(InAddressName);
+	if (AddressIndex == INDEX_NONE)
+	{
+		return false;
+	}
+
+	return Resize(AddressIndex, InNewElementCount);
+}
+
+void FMultiplexStorage::UpdateAddresses()
+{
+	int32 AlignmentShift = 0;
+	for (int32 AddressIndex = 0; AddressIndex < Addresses.Num(); AddressIndex++)
+	{
+		FMultiplexAddress& Address = Addresses[AddressIndex];
+		Address.ByteIndex += AlignmentShift;
+
+		UScriptStruct* ScriptStruct = GetScriptStruct(AddressIndex);
+		if (ScriptStruct != nullptr)
+		{
+			UScriptStruct::ICppStructOps* TheCppStructOps = ScriptStruct->GetCppStructOps();
+			if (TheCppStructOps != NULL)
+			{
+				if (!TheCppStructOps->HasZeroConstructor())
+				{
+					void* Pointer = GetData(AddressIndex);
+
+					if (Address.AlignmentBytes > 0)
+					{
+						if (!IsAligned(Pointer, TheCppStructOps->GetAlignment()))
+						{
+							Data.RemoveAt(Address.ByteIndex, Address.AlignmentBytes);
+							AlignmentShift -= Address.AlignmentBytes;
+							Address.AlignmentBytes = 0;
+							Pointer = GetData(AddressIndex);
+						}
+					}
+
+					while (!IsAligned(Pointer, TheCppStructOps->GetAlignment()))
+					{
+						Data.InsertZeroed(Address.ByteIndex, 1);
+						Address.AlignmentBytes++;
+						AlignmentShift++;
+						Pointer = GetData(AddressIndex);
+					}
+				}
+			}
+		}
+	}
+
+	for (int32 AddressIndex = 0; AddressIndex < Addresses.Num(); AddressIndex++)
+	{
+		Addresses[AddressIndex].Pointer = GetData(AddressIndex);
+	}
+
+	if (bUseNameMap)
+	{
+		NameMap.Reset();
+		for (int32 Index = 0; Index < Addresses.Num(); Index++)
+		{
+			NameMap.Add(Addresses[Index].Name, Index);
+		}
+	}
+}
+
+void FMultiplexStorage::FillWithZeroes(int32 InAddressIndex)
+{
+	ensure(Addresses.IsValidIndex(InAddressIndex));
+	FMemory::Memzero(GetData(InAddressIndex), Addresses[InAddressIndex].NumBytes());
+}
+
+int32 FMultiplexStorage::FindOrAddScriptStruct(UScriptStruct* InScriptStruct)
+{
+	int32 StructIndex = INDEX_NONE;
+	if (ScriptStructs.Find(InScriptStruct, StructIndex))
+	{
+		return StructIndex;
+	}
+	return ScriptStructs.Add(InScriptStruct);
+}
diff --git a/Engine/Source/Runtime/RigVM/Private/MultiplexVM.cpp b/Engine/Source/Runtime/RigVM/Private/MultiplexVM.cpp
new file mode 100644
--- /dev/null
+++ b/Engine/Source/Runtime/RigVM/Private/MultiplexVM.cpp
@@ -0,0 +1,152 @@
+// Copyright 1998-2019 Epic Games, Inc. All Rights Reserved.
+
+#include "MultiplexVM.h"
+
+void UMultiplexVM::Reset()
+{
+	Literals.Reset();
+	WorkState.Reset();
+	FunctionNames.Reset();
+	Functions.Reset();
+	ByteCode.Reset();
+	Instructions.Reset();
+}
+
+int32 UMultiplexVM::AddMultiplexFunction(UScriptStruct* InMultiplexStruct, const FName& InFunctionName)
+{
+	check(InMultiplexStruct);
+	FString FunctionKey = FString::Printf(TEXT("F%s::%s"), *InMultiplexStruct->GetName(), *InFunctionName.ToString());
+	int32 FunctionIndex = FunctionNames.Find(FunctionKey);
+	if (FunctionIndex != INDEX_NONE)
+	{
+		return FunctionIndex;
+	}
+
+	FMultiplexFunctionPtr Function = FMultiplexRegistry::Get().Find(*FunctionKey);
+	if (Function == nullptr)
+	{
+		return INDEX_NONE;
+	}
+
+	FunctionNames.Add(FunctionKey);
+	return Functions.Add(Function);
+}
+
+void UMultiplexVM::ResolveFunctionsIfRequired()
+{
+	if (Functions.Num() != FunctionNames.Num())
+	{
+		Functions.Reset();
+		Functions.SetNumZeroed(FunctionNames.Num());
+
+		for (int32 FunctionIndex = 0; FunctionIndex < FunctionNames.Num(); FunctionIndex)
+		{
+			Functions[FunctionIndex] = FMultiplexRegistry::Get().Find(*FunctionNames[FunctionIndex]);
+		}
+	}
+}
+
+void UMultiplexVM::RefreshInstructionsIfRequired()
+{
+	if (Instructions.Num() == 0)
+	{
+		Instructions = ByteCode.GetTable();
+	}
+}
+
+bool UMultiplexVM::Execute(FMultiplexStorage** Storage, TArrayView<void*> AdditionalArgs)
+{
+	ResolveFunctionsIfRequired();
+	RefreshInstructionsIfRequired();
+
+	if (Instructions.Num() == 0)
+	{
+		return true;
+	}
+
+	FMultiplexStorage* LocalStorage[] = { &WorkState, &Literals }; 
+	if (Storage == nullptr)
+	{
+		Storage = LocalStorage;
+	}
+
+	int32 InstructionIndex = 0;
+	while (Instructions[InstructionIndex].OpCode != EMultiplexOpCode::Exit)
+	{
+		switch (Instructions[InstructionIndex].OpCode)
+		{
+			case EMultiplexOpCode::Copy:
+			{
+				const FMultiplexCopyOp& Op = ByteCode.GetOpAt<FMultiplexCopyOp>(Instructions[InstructionIndex]);
+				Storage[0]->Copy(Op.Source.Index(), Op.Target.Index(), Storage[Op.Source.StorageType()], Op.SourceOffset, Op.TargetOffset, Op.NumBytes);
+				InstructionIndex++;
+				break;
+			}
+			case EMultiplexOpCode::Increment:
+			{
+				const FMultiplexIncrementOp& Op = ByteCode.GetOpAt<FMultiplexIncrementOp>(Instructions[InstructionIndex]);
+				Storage[0]->GetRef<int32>(Op.Arg.Index())++;
+				InstructionIndex++;
+				break;
+			}
+			case EMultiplexOpCode::Decrement:
+			{
+				const FMultiplexDecrementOp& Op = ByteCode.GetOpAt<FMultiplexDecrementOp>(Instructions[InstructionIndex]);
+				Storage[0]->GetRef<int32>(Op.Arg.Index())--;
+				InstructionIndex++;
+				break;
+			}
+			case EMultiplexOpCode::Equals:
+			{
+				const FMultiplexEqualsOp& Op = ByteCode.GetOpAt<FMultiplexEqualsOp>(Instructions[InstructionIndex]);
+				ensureMsgf(false, TEXT("TO BE IMPLEMENTED"));
+				break;
+			}
+			case EMultiplexOpCode::NotEquals:
+			{
+				const FMultiplexNotEqualsOp& Op = ByteCode.GetOpAt<FMultiplexNotEqualsOp>(Instructions[InstructionIndex]);
+				ensureMsgf(false, TEXT("TO BE IMPLEMENTED"));
+				break;
+			}
+			case EMultiplexOpCode::Jump:
+			{
+				const FMultiplexJumpOp& Op = ByteCode.GetOpAt<FMultiplexJumpOp>(Instructions[InstructionIndex]);
+				ensureMsgf(false, TEXT("TO BE IMPLEMENTED"));
+				break;
+			}
+			case EMultiplexOpCode::JumpIfTrue:
+			{
+				const FMultiplexJumpIfTrueOp& Op = ByteCode.GetOpAt<FMultiplexJumpIfTrueOp>(Instructions[InstructionIndex]);
+				ensureMsgf(false, TEXT("TO BE IMPLEMENTED"));
+				break;
+			}
+			case EMultiplexOpCode::JumpIfFalse:
+			{
+				const FMultiplexJumpIfFalseOp& Op = ByteCode.GetOpAt<FMultiplexJumpIfFalseOp>(Instructions[InstructionIndex]);
+				ensureMsgf(false, TEXT("TO BE IMPLEMENTED"));
+				break;
+			}
+			case EMultiplexOpCode::Execute:
+			{
+				const FMultiplexExecuteOp& Op = ByteCode.GetOpAt<FMultiplexExecuteOp>(Instructions[InstructionIndex]);
+				TArrayView<FMultiplexArgument> Args = ByteCode.GetArgumentsForExecuteOp(Instructions[InstructionIndex]);
+				(*Functions[Op.FunctionIndex])(Args, Storage, AdditionalArgs);
+				InstructionIndex++;
+				break;
+			}
+			case EMultiplexOpCode::Exit:
+			case EMultiplexOpCode::Invalid:
+			{
+				ensure(false);
+				return false;
+			}
+		}
+
+		if (!Instructions.IsValidIndex(InstructionIndex))
+		{
+			break;
+		}
+	}
+
+	return true;
+}
diff --git a/Engine/Source/Runtime/RigVM/Private/RigVMModule.cpp b/Engine/Source/Runtime/RigVM/Private/RigVMModule.cpp
new file mode 100644
--- /dev/null
+++ b/Engine/Source/Runtime/RigVM/Private/RigVMModule.cpp
@@ -0,0 +1,12 @@
+// Copyright 1998-2019 Epic Games, Inc. All Rights Reserved.
+
+/*=============================================================================
+	RigVM.h: Render core module implementation.
+=============================================================================*/
+
+#include "RigVMModule.h"
+#include "Modules/ModuleManager.h"
+
+IMPLEMENT_MODULE(FDefaultModuleImpl, RigVM);
+
+DEFINE_LOG_CATEGORY(LogRigVM);
diff --git a/Engine/Source/Runtime/RigVM/Public/MultiplexByteCode.h b/Engine/Source/Runtime/RigVM/Public/MultiplexByteCode.h
new file mode 100644
--- /dev/null
+++ b/Engine/Source/Runtime/RigVM/Public/MultiplexByteCode.h
@@ -0,0 +1,375 @@
+// Copyright 1998-2019 Epic Games, Inc. All Rights Reserved.
+
+#pragma once
+
+#include "CoreMinimal.h"
+#include "MultiplexRegistry.h"
+#include "MultiplexByteCode.generated.h"
+
+struct FMultiplexByteCode;
+
+UENUM()
+enum class EMultiplexOpCode : uint8
+{
+	Copy,
+	Increment,
+	Decrement,
+	Equals,
+	NotEquals,
+	Jump,
+	JumpIfTrue,
+	JumpIfFalse,
+	Execute,
+	Exit,
+	Invalid
+};
+
+struct RIGVM_API FMultiplexBaseOp
+{
+	FMultiplexBaseOp(EMultiplexOpCode InOpCode = EMultiplexOpCode::Invalid)
+	: OpCode(InOpCode)
+	{
+	}
+
+	EMultiplexOpCode OpCode;
+};
+
+struct RIGVM_API FMultiplexCopyOp : public FMultiplexBaseOp
+{
+	FMultiplexCopyOp()
+	: FMultiplexBaseOp(EMultiplexOpCode::Copy)
+	, Source()
+	, Target()
+	, SourceOffset(INDEX_NONE)
+	, TargetOffset(INDEX_NONE)
+	, NumBytes(INDEX_NONE)
+	{
+	}
+
+	FMultiplexCopyOp(
+		FMultiplexArgument InSource,
+		FMultiplexArgument InTarget,
+		int32 InSourceOffset = INDEX_NONE,
+		int32 InTargetOffset = INDEX_NONE,
+		int32 InNumBytes = INDEX_NONE
+	)
+		: FMultiplexBaseOp(EMultiplexOpCode::Copy)
+		, Source(InSource)
+		, Target(InTarget)
+		, SourceOffset(InSourceOffset)
+		, TargetOffset(InTargetOffset)
+		, NumBytes(InNumBytes)
+	{
+	}
+
+	FMultiplexArgument Source;
+	FMultiplexArgument Target;
+	int32 SourceOffset;
+	int32 TargetOffset;
+	int32 NumBytes;
+};
+
+struct RIGVM_API FMultiplexIncrementOp : public FMultiplexBaseOp
+{
+	FMultiplexIncrementOp()
+	: FMultiplexBaseOp(EMultiplexOpCode::Increment)
+	, Arg()
+	{
+	}
+
+	FMultiplexIncrementOp(FMultiplexArgument InArg)
+	: FMultiplexBaseOp(EMultiplexOpCode::Increment)
+	, Arg(InArg)
+	{
+	}
+
+	FMultiplexArgument Arg;
+};
+
+struct RIGVM_API FMultiplexDecrementOp : public FMultiplexBaseOp
+{
+	FMultiplexDecrementOp()
+	: FMultiplexBaseOp(EMultiplexOpCode::Decrement)
+	, Arg()
+	{
+	}
+
+	FMultiplexDecrementOp(FMultiplexArgument InArg)
+	: FMultiplexBaseOp(EMultiplexOpCode::Decrement)
+	, Arg(InArg)
+	{
+	}
+
+	FMultiplexArgument Arg;
+};
+
+struct RIGVM_API FMultiplexEqualsOp : public FMultiplexBaseOp
+{
+	FMultiplexEqualsOp()
+	: FMultiplexBaseOp(EMultiplexOpCode::Equals)
+	, A()
+	, B()
+	, Result()
+	{
+	}
+
+	FMultiplexEqualsOp(
+		FMultiplexArgument InA,
+		FMultiplexArgument InB,
+		FMultiplexArgument InResult
+	)
+	: FMultiplexBaseOp(EMultiplexOpCode::Equals)
+	, A(InA)
+	, B(InB)
+	, Result(InResult)
+	{
+	}
+
+	FMultiplexArgument A;
+	FMultiplexArgument B;
+	FMultiplexArgument Result;
+};
+
+struct RIGVM_API FMultiplexNotEqualsOp : public FMultiplexBaseOp
+{
+	FMultiplexNotEqualsOp()
+	: FMultiplexBaseOp(EMultiplexOpCode::NotEquals)
+	, A()
+	, B()
+	, Result()
+	{
+	}
+
+	FMultiplexNotEqualsOp(
+		FMultiplexArgument InA,
+		FMultiplexArgument InB,
+		FMultiplexArgument InResult
+	)
+	: FMultiplexBaseOp(EMultiplexOpCode::NotEquals)
+	, A(InA)
+	, B(InB)
+	, Result(InResult)
+	{
+	}
+
+	FMultiplexArgument A;
+	FMultiplexArgument B;
+	FMultiplexArgument Result;
+};
+
+struct RIGVM_API FMultiplexJumpOp : public FMultiplexBaseOp
+{
+	FMultiplexJumpOp()
+	: FMultiplexBaseOp(EMultiplexOpCode::Jump)
+	, InstructionIndex(INDEX_NONE)
+	{
+	}
+
+	FMultiplexJumpOp(int32 InInstructionIndex)
+	: FMultiplexBaseOp(EMultiplexOpCode::Jump)
+	, InstructionIndex(InInstructionIndex)
+	{
+	}
+
+	int32 InstructionIndex;
+};
+
+struct RIGVM_API FMultiplexJumpIfTrueOp : public FMultiplexBaseOp
+{
+	FMultiplexJumpIfTrueOp()
+	: FMultiplexBaseOp(EMultiplexOpCode::JumpIfTrue)
+	, InstructionIndex(INDEX_NONE)
+	, Condition()
+	{
+	}
+
+	FMultiplexJumpIfTrueOp(int32 InInstructionIndex, FMultiplexArgument InCondition)
+	: FMultiplexBaseOp(EMultiplexOpCode::JumpIfTrue)
+	, InstructionIndex(InInstructionIndex)
+	, Condition(InCondition)
+	{
+	}
+
+	int32 InstructionIndex;
+	FMultiplexArgument Condition;
+};
+
+struct RIGVM_API FMultiplexJumpIfFalseOp : public FMultiplexBaseOp
+{
+	FMultiplexJumpIfFalseOp()
+	: FMultiplexBaseOp(EMultiplexOpCode::JumpIfFalse)
+	, InstructionIndex(INDEX_NONE)
+	, Condition()
+	{
+	}
+
+	FMultiplexJumpIfFalseOp(int32 InInstructionIndex, FMultiplexArgument InCondition)
+	: FMultiplexBaseOp(EMultiplexOpCode::JumpIfFalse)
+	, InstructionIndex(InInstructionIndex)
+	, Condition(InCondition)
+	{
+	}
+
+	int32 InstructionIndex;
+	FMultiplexArgument Condition;
+};
+
+struct RIGVM_API FMultiplexExecuteOp : public FMultiplexBaseOp
+{
+	FMultiplexExecuteOp()
+	: FMultiplexBaseOp(EMultiplexOpCode::Execute)
+	, FunctionIndex(INDEX_NONE)
+	, ArgumentCount(0)
+	{
+	}
+
+	FMultiplexExecuteOp(uint16 InFunctionIndex, uint16 InArgumentCount)
+	: FMultiplexBaseOp(EMultiplexOpCode::Execute)
+	, FunctionIndex(InFunctionIndex)
+	, ArgumentCount(InArgumentCount)
+	{
+	}
+
+	uint16 FunctionIndex;
+	uint16 ArgumentCount;
+};
+
+struct RIGVM_API FMultiplexExitOp : public FMultiplexBaseOp
+{
+	FMultiplexExitOp()
+	: FMultiplexBaseOp(EMultiplexOpCode::Exit)
+	{
+	}
+};
+
+USTRUCT()
+struct RIGVM_API FMultiplexByteCodeTableEntry
+{
+	GENERATED_USTRUCT_BODY()
+
+	FMultiplexByteCodeTableEntry(EMultiplexOpCode InOpCode = EMultiplexOpCode::Invalid, uint64 InByteCodeIndex = UINT64_MAX)
+		: OpCode(InOpCode)
+		, ByteCodeIndex(InByteCodeIndex)
+	{
+	}
+
+	UPROPERTY()
+	EMultiplexOpCode OpCode;
+
+	UPROPERTY()
+	uint64 ByteCodeIndex;
+};
+
+USTRUCT()
+struct RIGVM_API FMultiplexByteCodeTable
+{
+	GENERATED_USTRUCT_BODY()
+
+public:
+
+	FMultiplexByteCodeTable();
+
+	void Reset();
+	FORCEINLINE bool IsValidIndex(int32 InIndex) const { return Entries.IsValidIndex(InIndex); }
+	FORCEINLINE int32 Num() const { return Entries.Num(); }
+	FORCEINLINE const FMultiplexByteCodeTableEntry& operator[](int32 InIndex) const { return Entries[InIndex]; }
+
+private:
+
+	// hide utility constructor
+	FMultiplexByteCodeTable(const FMultiplexByteCode& InByteCode);
+
+	UPROPERTY()
+	TArray<FMultiplexByteCodeTableEntry> Entries;
+
+	friend struct FMultiplexByteCode;
+};
+
+USTRUCT()
+struct RIGVM_API FMultiplexByteCode
+{
+	GENERATED_USTRUCT_BODY()
+
+public:
+
+	FMultiplexByteCode();
+
+	void Reset();
+	uint64 Num() const;
+
+	uint64 AddCopyOp(const FMultiplexArgument& InSource, const FMultiplexArgument& InTarget, int32 InSourceOffset = INDEX_NONE, int32 InTargetOffset = INDEX_NONE, int32 InNumBytes = INDEX_NONE);
+	uint64 AddIncrementOp(const FMultiplexArgument& InArg);
+	uint64 AddDecrementOp(const FMultiplexArgument& InArg);
+	uint64 AddEqualsOp(const FMultiplexArgument& InA, const FMultiplexArgument& InB, const FMultiplexArgument& InResult);
+	uint64 AddNotEqualsOp(const FMultiplexArgument& InA, const FMultiplexArgument& InB, const FMultiplexArgument& InResult);
+	uint64 AddJumpOp(uint64 InByteCodeIndex);
+	uint64 AddJumpIfTrueOp(uint64 InByteCodeIndex, const FMultiplexArgument& InCondition);
+	uint64 AddJumpIfFalseOp(uint64 InByteCodeIndex, const FMultiplexArgument& InCondition);
+	uint64 AddExecuteOp(uint16 InFunctionIndex, const TArrayView<FMultiplexArgument>& InArguments);
+	uint64 AddExitOp();
+
+	FORCEINLINE FMultiplexByteCodeTable GetTable() const
+	{
+		return FMultiplexByteCodeTable(*this);
+	}
+
+	FORCEINLINE EMultiplexOpCode GetOpCodeAt(uint64 InByteCodeIndex) const
+	{
+		ensure(InByteCodeIndex >= 0 && InByteCodeIndex < ByteCode.Num());
+		return (EMultiplexOpCode)ByteCode[InByteCodeIndex];
+	}
+
+	uint64 GetOpNumBytesAt(uint64 InByteCodeIndex, bool bIncludeArguments = true) const;
+
+	template<class OpType>
+	FORCEINLINE const OpType& GetOpAt(uint64 InByteCodeIndex) const
+	{
+		ensure(InByteCodeIndex >= 0 && InByteCodeIndex <= ByteCode.Num() - sizeof(OpType));
+		return *(const OpType*)(ByteCode.GetData() + InByteCodeIndex);
+	}
+
+	template<class OpType>
+	FORCEINLINE const OpType& GetOpAt(const FMultiplexByteCodeTableEntry& InEntry) const
+	{
+		ensure(OpType().OpCode == InEntry.OpCode);
+		return GetOpAt<OpType>(InEntry.ByteCodeIndex);
+	}
+
+	FORCEINLINE TArrayView<FMultiplexArgument> GetArgumentsAt(uint64 InByteCodeIndex, uint16 InArgumentCount) const
+	{
+		ensure(InByteCodeIndex >= 0 && InByteCodeIndex <= ByteCode.Num() - sizeof(FMultiplexArgument) * InArgumentCount);
+		return TArrayView<FMultiplexArgument>((FMultiplexArgument*)(ByteCode.GetData() + InByteCodeIndex), InArgumentCount);
+	}
+
+	FORCEINLINE TArrayView<FMultiplexArgument> GetArgumentsForExecuteOp(uint64 InByteCodeIndex) const
+	{
+		const FMultiplexExecuteOp& ExecuteOp = GetOpAt<FMultiplexExecuteOp>(InByteCodeIndex);
+		return GetArgumentsAt(InByteCodeIndex + sizeof(FMultiplexExecuteOp), ExecuteOp.ArgumentCount);
+	}
+
+	FORCEINLINE TArrayView<FMultiplexArgument> GetArgumentsForExecuteOp(const FMultiplexByteCodeTableEntry& InEntry) const
+	{
+		const FMultiplexExecuteOp& ExecuteOp = GetOpAt<FMultiplexExecuteOp>(InEntry);
+		return GetArgumentsAt(InEntry.ByteCodeIndex + sizeof(FMultiplexExecuteOp), ExecuteOp.ArgumentCount);
+	}
+
+	FORCEINLINE const TArrayView<uint8> GetByteCode() const
+	{
+		const uint8* Data = ByteCode.GetData();
+		return TArrayView<uint8>((uint8*)Data, ByteCode.Num());
+	}
+
+private:
+
+	template<class OpType>
+	FORCEINLINE uint64 AddOp(const OpType& InOp)
+	{
+		uint64 Address = (uint64)ByteCode.AddUninitialized(sizeof(OpType));
+		FMemory::Memcpy(ByteCode.GetData() + Address, &InOp, sizeof(OpType));
+		return Address;
+	}
+
+	// storage for all functions
+	UPROPERTY()
+	TArray<uint8> ByteCode;
+};
diff --git a/Engine/Source/Runtime/RigVM/Public/MultiplexRegistry.h b/Engine/Source/Runtime/RigVM/Public/MultiplexRegistry.h
new file mode 100644
--- /dev/null
+++ b/Engine/Source/Runtime/RigVM/Public/MultiplexRegistry.h
@@ -0,0 +1,51 @@
+// Copyright 1998-2019 Epic Games, Inc. All Rights Reserved.
+
+#pragma once
+
+#include "CoreMinimal.h"
+#include "MultiplexStorage.h"
+
+typedef void (*FMultiplexFunctionPtr)(const TArrayView<FMultiplexArgument>&, FMultiplexStorage**, const TArrayView<void*>&);
+
+struct RIGVM_API FMultiplexFunction
+{
+	const TCHAR* Name;
+	FMultiplexFunctionPtr FunctionPtr;
+
+	FMultiplexFunction()
+		: Name(nullptr)
+		, FunctionPtr(nullptr)
+	{
+	}
+
+	FMultiplexFunction(const TCHAR* InName, FMultiplexFunctionPtr InFunctionPtr)
+		: Name(InName)
+		, FunctionPtr(InFunctionPtr)
+	{
+	}
+};
+
+struct RIGVM_API FMultiplexRegistry
+{
+public:
+	static FMultiplexRegistry& Get();
+
+	void Register(const TCHAR* InName, FMultiplexFunctionPtr InFunctionPtr);
+	void Refresh();
+
+	FMultiplexFunctionPtr Find(const TCHAR* InName) const;
+
+private:
+
+	// disable default constructor
+	FMultiplexRegistry() {}
+	// disable copy constructor
+	FMultiplexRegistry(const FMultiplexRegistry&) {}
+	// disable assignment operator
+	FMultiplexRegistry& operator= (const FMultiplexRegistry &InOther) { return *this; }
+
+	// storage for all functions
+	TArray<FMultiplexFunction> Functions;
+
+	static FMultiplexRegistry s_MultiplexRegistry;
+};
diff --git a/Engine/Source/Runtime/RigVM/Public/MultiplexStorage.h b/Engine/Source/Runtime/RigVM/Public/MultiplexStorage.h
new file mode 100644
--- /dev/null
+++ b/Engine/Source/Runtime/RigVM/Public/MultiplexStorage.h
@@ -0,0 +1,547 @@
+// Copyright 1998-2019 Epic Games, Inc. All Rights Reserved.
+
+#pragma once
+
+#include "CoreMinimal.h"
+#include "UObject/ObjectMacros.h"
+#include "MultiplexStorage.generated.h"
+
+struct RIGVM_API FMultiplexArgument
+{
+public:
+
+	FMultiplexArgument()
+		:Address(INDEX_NONE)
+	{
+	}
+
+	FMultiplexArgument(int32 InAddress)
+		:Address(InAddress)
+	{
+	}
+
+	static FMultiplexArgument MakeArgument(int32 InAddress)
+	{
+		return FMultiplexArgument(InAddress);
+	}
+
+	static FMultiplexArgument MakeLiteral(int32 InAddress)
+	{
+		return FMultiplexArgument(-1 - InAddress);
+	}
+
+	FORCEINLINE bool operator ==(int32 InOther) const { return Address == InOther; }
+
+	FORCEINLINE bool IsLiteral() const { return Address < 0; }
+	FORCEINLINE int32 StorageType() const { return Address < 0 ? 1 : 0; }
+	FORCEINLINE int32 Index() const { return Address < 0 ? -(Address + 1) : Address; }
+	FORCEINLINE int32 NativeAddress() const { return Address; }
+
+private:
+	int32 Address;
+};
+
+UENUM()
+enum class EMultiplexAddressType : uint8
+{
+	Plain,
+	String,
+	Name,
+	Struct,
+	Invalid
+};
+
+USTRUCT()
+struct RIGVM_API FMultiplexAddress
+{
+	GENERATED_USTRUCT_BODY()
+
+		FMultiplexAddress()
+		: Type(EMultiplexAddressType::Invalid)
+		, Pointer(nullptr)
+		, ByteIndex(INDEX_NONE)
+		, ElementSize(0)
+		, ElementCount(0)
+		, AlignmentBytes(0)
+		, Name(NAME_None)
+		, ScriptStructIndex(INDEX_NONE)
+	{
+	}
+
+	UPROPERTY()
+	EMultiplexAddressType Type;
+
+	void* Pointer;
+
+	UPROPERTY()
+	int32 ByteIndex;
+
+	UPROPERTY()
+	int32 ElementSize;
+
+	UPROPERTY()
+	int32 ElementCount;
+
+	UPROPERTY()
+	int32 AlignmentBytes;
+
+	UPROPERTY()
+	FName Name;
+
+	UPROPERTY()
+	int32 ScriptStructIndex;
+
+	FORCEINLINE int32 FirstByte() const { return ByteIndex + AlignmentBytes; }
+	FORCEINLINE bool IsArray() const { return ElementCount > 1; }
+	FORCEINLINE bool IsPlain() const { return ScriptStructIndex == INDEX_NONE;  }
+	FORCEINLINE int32 NumBytes(bool bIncludeAlignment = true) const { return ElementCount * ElementSize + (bIncludeAlignment ? AlignmentBytes : 0); }
+
+	template<class T>
+	FORCEINLINE const T* Get() const
+	{
+		ensure(ElementCount > 0);
+		return (const T*)Pointer;
+	}
+
+	template<class T>
+	FORCEINLINE const T& GetRef() const
+	{
+		return *Get<T>();
+	}
+
+	template<class T>
+	FORCEINLINE T* Get()
+	{
+		ensure(ElementCount > 0);
+		return (T*)Pointer;
+	}
+
+	template<class T>
+	FORCEINLINE T& GetRef()
+	{
+		return *Get<T>();
+	}
+
+	template<class T>
+	FORCEINLINE TArrayView<T> GetArray()
+	{
+		ensure(ElementCount > 0);
+		return TArrayView<T>((T*)Pointer, ElementCount);
+	}
+};
+
+USTRUCT()
+struct RIGVM_API FMultiplexStorage
+{
+	GENERATED_USTRUCT_BODY()
+	
+public:
+
+	FMultiplexStorage(bool bInUseNames = true);
+	FMultiplexStorage(const FMultiplexStorage& Other);
+	~FMultiplexStorage();
+
+	FMultiplexStorage& operator= (const FMultiplexStorage &InOther);
+
+	FORCEINLINE bool SupportsNames() const { return bUseNameMap;  }
+	FORCEINLINE int32 Num() const { return Addresses.Num(); }
+	FORCEINLINE const FMultiplexAddress& operator[](int32 InIndex) const { return Addresses[InIndex]; }
+	FORCEINLINE FMultiplexAddress& operator[](int32 InIndex) { return Addresses[InIndex]; }
+	FORCEINLINE const FMultiplexAddress& operator[](const FMultiplexArgument& InArg) const { return Addresses[InArg.Index()]; }
+	FORCEINLINE FMultiplexAddress& operator[](const FMultiplexArgument& InArg) { return Addresses[InArg.Index()]; }
+	FORCEINLINE const FMultiplexAddress& operator[](const FName& InName) const { return Addresses[GetIndex(InName)]; }
+	FORCEINLINE FMultiplexAddress& operator[](const FName& InName) { return Addresses[GetIndex(InName)]; }
+
+	FORCEINLINE TArray<FMultiplexAddress>::RangedForIteratorType      begin() { return Addresses.begin(); }
+	FORCEINLINE TArray<FMultiplexAddress>::RangedForConstIteratorType begin() const { return Addresses.begin(); }
+	FORCEINLINE TArray<FMultiplexAddress>::RangedForIteratorType      end() { return Addresses.end(); }
+	FORCEINLINE TArray<FMultiplexAddress>::RangedForConstIteratorType end() const { return Addresses.end(); }
+
+	FORCEINLINE const void* GetData(int32 InAddressIndex) const
+	{
+		ensure(Addresses.IsValidIndex(InAddressIndex));
+		const FMultiplexAddress& Address = Addresses[InAddressIndex];
+		ensure(Address.ElementCount > 0);
+		return (const void*)&Data[Address.FirstByte()];
+	}
+
+	FORCEINLINE void* GetData(int32 InAddressIndex)
+	{
+		ensure(Addresses.IsValidIndex(InAddressIndex));
+		const FMultiplexAddress& Address = Addresses[InAddressIndex];
+		ensure(Address.ElementCount > 0);
+		return (void*)&Data[Address.FirstByte()];
+	}
+
+	template<class T>
+	FORCEINLINE const T* Get(int32 InAddressIndex) const
+	{
+		ensure(Addresses.IsValidIndex(InAddressIndex));
+		const FMultiplexAddress& Address = Addresses[InAddressIndex];
+		ensure(Address.ElementCount > 0);
+		return (const T*)&Data[Address.FirstByte()];
+	}
+
+	template<class T>
+	FORCEINLINE const T* Get(const FMultiplexArgument& InArgument) const
+	{
+		return Get<T>(InArgument.Index());
+	}
+
+	template<class T>
+	FORCEINLINE const T& GetRef(int32 InAddressIndex) const
+	{
+		return *Get<T>(InAddressIndex);
+	}
+
+	template<class T>
+	FORCEINLINE const T& GetRef(const FMultiplexArgument& InArgument) const
+	{
+		return GetRef<T>(InArgument.Index());
+	}
+
+	template<class T>
+	FORCEINLINE T* Get(int32 InAddressIndex)
+	{
+		ensure(Addresses.IsValidIndex(InAddressIndex));
+		const FMultiplexAddress& Address = Addresses[InAddressIndex];
+		ensure(Address.ElementCount > 0);
+		return (T*)&Data[Address.FirstByte()];
+	}
+
+	template<class T>
+	FORCEINLINE T* Get(const FMultiplexArgument& InArgument)
+	{
+		return Get<T>(InArgument.Index());
+	}
+
+	template<class T>
+	FORCEINLINE T& GetRef(int32 InAddressIndex)
+	{
+		return *Get<T>(InAddressIndex);
+	}
+
+	template<class T>
+	FORCEINLINE T& GetRef(const FMultiplexArgument& InArgument)
+	{
+		return GetRef<T>(InArgument.Index());
+	}
+
+	template<class T>
+	FORCEINLINE TArrayView<T> GetArray(int32 InAddressIndex)
+	{
+		ensure(Addresses.IsValidIndex(InAddressIndex));
+		const FMultiplexAddress& Address = Addresses[InAddressIndex];
+		ensure(Address.ElementCount > 0);
+		return TArrayView<T>((T*)&Data[Address.FirstByte()], Address.ElementCount);
+	}
+	
+	template<class T>
+	FORCEINLINE TArrayView<T> GetArray(const FMultiplexArgument& InArgument)
+	{
+		return GetArray<T>(InArgument.Index());
+	}
+
+	FORCEINLINE UScriptStruct* GetScriptStruct(int32 InAddressIndex) const
+	{
+		ensure(Addresses.IsValidIndex(InAddressIndex));
+		const FMultiplexAddress& Address = Addresses[InAddressIndex];
+		if (Address.ScriptStructIndex != INDEX_NONE)
+		{
+			ensure(ScriptStructs.IsValidIndex(Address.ScriptStructIndex));
+			return ScriptStructs[Address.ScriptStructIndex];
+		}
+		return nullptr;
+	}
+
+	bool Copy(
+		int32 InSourceAddressIndex,
+		int32 InTargetAddressIndex,
+		const FMultiplexStorage* InSourceStorage = nullptr,
+		int32 InSourceByteOffset = INDEX_NONE,
+		int32 InTargetByteOffset = INDEX_NONE,
+		int32 InNumBytes = INDEX_NONE);
+
+	bool Copy(
+		const FName& InSourceName,
+		const FName& InTargetName,
+		const FMultiplexStorage* InSourceStorage = nullptr,
+		int32 InSourceByteOffset = INDEX_NONE,
+		int32 InTargetByteOffset = INDEX_NONE,
+		int32 InNumBytes = INDEX_NONE);
+
+	FORCEINLINE int32 GetIndex(const FName& InName) const
+	{
+		if (!bUseNameMap)
+		{
+			return INDEX_NONE;
+		}
+
+		if (NameMap.Num() != Addresses.Num())
+		{
+			for (int32 Index = 0; Index < Addresses.Num(); Index++)
+			{
+				if (Addresses[Index].Name == InName)
+				{
+					return Index;
+				}
+			}
+		}
+		else
+		{
+			const int32* Index = NameMap.Find(InName);
+			if (Index != nullptr)
+			{
+				return *Index;
+			}
+		}
+
+		return INDEX_NONE;
+	}
+
+	FORCEINLINE bool IsNameAvailable(const FName& InPotentialNewName) const\
+	{
+		if (!bUseNameMap)
+		{
+			return false;
+		}
+		return GetIndex(InPotentialNewName) == INDEX_NONE;
+	}
+
+	void Reset();
+
+	FORCEINLINE int32 AddPlainArray(const FName& InNewName, int32 InElementSize, int32 InCount, const void* InDataPtr = nullptr)
+	{
+		return Allocate(InNewName, InElementSize, InCount, InDataPtr);
+	}
+
+	template<class T>
+	FORCEINLINE int32 AddPlainArray(const FName& InNewName, int32 InCount, const T* InDataPtr = nullptr)
+	{
+		return Allocate(InNewName, sizeof(T), InCount, (const void*)InDataPtr);
+	}
+
+	template<class T>
+	FORCEINLINE int32 AddPlainArray(const FName& InNewName, const TArray<T>& InArray)
+	{
+		return AddPlainArray<T>(InNewName, InArray.Num(), InArray.GetData());
+	}
+
+	template<class T>
+	FORCEINLINE int32 AddPlainArray(const TArray<T>& InArray)
+	{
+		return AddPlainArray<T>(NAME_None, InArray);
+	}
+
+	FORCEINLINE int32 AddPlain(const FName& InNewName, int32 InElementSize, const void* InValuePtr)
+	{
+		return AddPlainArray(InNewName, InElementSize, 1, InValuePtr);
+	}
+
+	FORCEINLINE int32 AddPlain(int32 InElementSize, const void* InValuePtr)
+	{
+		return AddPlain(NAME_None, InElementSize, InValuePtr);
+	}
+
+	template<class T>
+	FORCEINLINE int32 AddPlain(const FName& InNewName, const T& InValue)
+	{
+		return AddPlainArray<T>(InNewName, 1, &InValue);
+	}
+
+	template<class T>
+	FORCEINLINE int32 AddPlain(const T& InValue)
+	{
+		return AddPlain<T>(NAME_None, InValue);
+	}
+
+	FORCEINLINE int32 AddNameArray(const FName& InNewName, int32 InCount, const FName* InDataPtr = nullptr)
+	{
+		int32 Address = Allocate(InNewName, sizeof(FName), InCount, nullptr);
+		Addresses[Address].Type = EMultiplexAddressType::Name;
+
+		Construct(Address);
+
+		if(InDataPtr)
+		{
+			FName* DataPtr = (FName*)GetData(Address);
+			for (int32 Index = 0; Index < InCount; Index++)
+			{
+				DataPtr[Index] = InDataPtr[Index];
+			}
+		}
+
+		return Address;
+	}
+
+	FORCEINLINE int32 AddNameArray(const FName& InNewName, const TArray<FName>& InArray)
+	{
+		return AddNameArray(InNewName, InArray.Num(), InArray.GetData());
+	}
+
+	FORCEINLINE int32 AddNameArray(const TArray<FName>& InArray)
+	{
+		return AddNameArray(NAME_None, InArray);
+	}
+
+	FORCEINLINE int32 AddName(const FName& InNewName, const FName& InValue)
+	{
+		return AddNameArray(InNewName, 1, &InValue);
+	}
+
+	FORCEINLINE int32 AddName(const FName& InValue)
+	{
+		return AddName(NAME_None, InValue);
+	}
+
+	FORCEINLINE int32 AddStringArray(const FName& InNewName, int32 InCount, const FString* InDataPtr = nullptr)
+	{
+		int32 Address = Allocate(InNewName, sizeof(FString), InCount, nullptr);
+		Addresses[Address].Type = EMultiplexAddressType::String;
+
+		Construct(Address);
+
+		if(InDataPtr)
+		{
+			FString* DataPtr = (FString*)GetData(Address);
+			for (int32 Index = 0; Index < InCount; Index++)
+			{
+				DataPtr[Index] = InDataPtr[Index];
+			}
+		}
+
+		return Address;
+	}
+
+	FORCEINLINE int32 AddStringArray(const FName& InNewName, const TArray<FString>& InArray)
+	{
+		return AddStringArray(InNewName, InArray.Num(), InArray.GetData());
+	}
+
+	FORCEINLINE int32 AddStringArray(const TArray<FString>& InArray)
+	{
+		return AddStringArray(NAME_None, InArray);
+	}
+
+	FORCEINLINE int32 AddString(const FName& InNewName, const FString& InValue)
+	{
+		return AddStringArray(InNewName, 1, &InValue);
+	}
+
+	FORCEINLINE int32 AddString(const FString& InValue)
+	{
+		return AddString(NAME_None, InValue);
+	}
+
+
+	FORCEINLINE int32 AddStructArray(const FName& InNewName, UScriptStruct* InScriptStruct, int32 InCount, const void* InDataPtr = nullptr)
+	{
+		int32 Address = Allocate(InNewName, InScriptStruct->GetStructureSize(), InCount, nullptr, false);
+		if (Address == INDEX_NONE)
+		{
+			return INDEX_NONE;
+		}
+
+		Addresses[Address].Type = EMultiplexAddressType::Struct;
+		Addresses[Address].ScriptStructIndex = FindOrAddScriptStruct(InScriptStruct);
+
+		UpdateAddresses();
+
+		// construct the content
+		Construct(Address);
+
+		// copy values from the provided data
+		if (InDataPtr != nullptr)
+		{
+			InScriptStruct->CopyScriptStruct(GetData(Address), InDataPtr, InCount);
+		}
+
+		return Address;
+	}
+
+	FORCEINLINE int32 AddStructArray(UScriptStruct* InScriptStruct, int32 InCount, const void* InDataPtr = nullptr)
+	{
+		return AddStructArray(NAME_None, InScriptStruct, InCount, InDataPtr);
+	}
+
+	template<class T>
+	FORCEINLINE int32 AddStructArray(const FName& InNewName, int32 InCount, const T* InDataPtr = nullptr)
+	{
+		// if you are hitting this - you might need to use AddPlainArray instead!
+		UScriptStruct* Struct = T::StaticStruct();
+		if (Struct == nullptr)
+		{
+			return INDEX_NONE;
+		}
+	
+		return AddStructArray(InNewName, Struct, InCount, InDataPtr);
+	}
+
+	template<class T>
+	FORCEINLINE int32 AddStructArray(const FName& InNewName, const TArray<T>& InArray)
+	{
+		return AddStructArray<T>(InNewName, InArray.Num(), InArray.GetData());
+	}
+
+	template<class T>
+	FORCEINLINE int32 AddStructArray(const TArray<T>& InArray)
+	{
+		return AddStructArray<T>(NAME_None, InArray);
+	}
+
+	FORCEINLINE int32 AddStruct(const FName& InNewName, UScriptStruct* InScriptStruct, const void* InValuePtr)
+	{
+		return AddStructArray(InNewName, InScriptStruct, 1, InValuePtr);
+	}
+
+	FORCEINLINE int32 AddStruct(UScriptStruct* InScriptStruct, const void* InValuePtr)
+	{
+		return AddStruct(NAME_None, InScriptStruct, InValuePtr);
+	}
+
+	template<class T>
+	FORCEINLINE int32 AddStruct(const FName& InNewName, const T& InValue)
+	{
+		return AddStructArray<T>(InNewName, 1, &InValue);
+	}
+
+	template<class T>
+	FORCEINLINE int32 AddStruct(const T& InValue)
+	{
+		return AddStruct<T>(NAME_None, InValue);
+	}
+
+	bool Remove(int32 InAddressIndex);
+	bool Remove(const FName& InAddressName);
+	FName Rename(int32 InAddressIndex, const FName& InNewName);
+	FName Rename(const FName& InOldName, const FName& InNewName);
+	bool Resize(int32 InAddressIndex, int32 InNewElementCount);
+	bool Resize(const FName& InAddressName, int32 InNewElementCount);
+
+	void UpdateAddresses();
+
+private:
+
+	int32 Allocate(const FName& InNewName, int32 InElementSize, int32 InCount, const void* InDataPtr = nullptr, bool bUpdateAddresses = true);
+	int32 Allocate(int32 InElementSize, int32 InCount, const void* InDataPtr = nullptr, bool bUpdateAddresses = true);
+	bool Construct(int32 InAddressIndex, int32 InElementIndex = INDEX_NONE);
+	bool Destroy(int32 InAddressIndex, int32 InElementIndex = INDEX_NONE);
+
+	void FillWithZeroes(int32 InAddressIndex);
+	int32 FindOrAddScriptStruct(UScriptStruct* InScriptStruct);
+
+	UPROPERTY()
+	bool bUseNameMap;
+
+	UPROPERTY()
+	TArray<FMultiplexAddress> Addresses;
+
+	UPROPERTY()
+	TArray<uint8> Data;
+
+	UPROPERTY()
+	TArray<UScriptStruct*> ScriptStructs;
+
+	UPROPERTY(transient)
+	TMap<FName, int32> NameMap;
+};
diff --git a/Engine/Source/Runtime/RigVM/Public/MultiplexVM.h b/Engine/Source/Runtime/RigVM/Public/MultiplexVM.h
new file mode 100644
--- /dev/null
+++ b/Engine/Source/Runtime/RigVM/Public/MultiplexVM.h
@@ -0,0 +1,46 @@
+// Copyright 1998-2019 Epic Games, Inc. All Rights Reserved.
+
+#pragma once
+
+#include "CoreMinimal.h"
+#include "MultiplexStorage.h"
+#include "MultiplexRegistry.h"
+#include "MultiplexByteCode.h"
+#include "MultiplexVM.generated.h"
+
+UCLASS()
+class RIGVM_API UMultiplexVM : public UObject
+{
+	GENERATED_BODY()
+
+public:
+
+	void Reset();
+
+	bool Execute(FMultiplexStorage** Storage = nullptr, TArrayView<void*> AdditionalArgs = TArrayView<void*>());
+
+	UFUNCTION()
+	int32 AddMultiplexFunction(UScriptStruct* InMultiplexStruct, const FName& InFunctionName);
+
+	UPROPERTY()
+	FMultiplexStorage Literals;
+
+	UPROPERTY()
+	FMultiplexStorage WorkState;
+
+	UPROPERTY()
+	FMultiplexByteCode ByteCode;
+
+	UPROPERTY()
+	FMultiplexByteCodeTable Instructions;
+
+private:
+
+	void ResolveFunctionsIfRequired();
+	void RefreshInstructionsIfRequired();
+
+	UPROPERTY()
+	TArray<FString> FunctionNames;
+
+	TArray<FMultiplexFunctionPtr> Functions;
+};
diff --git a/Engine/Source/Runtime/RigVM/Public/RigVMModule.h b/Engine/Source/Runtime/RigVM/Public/RigVMModule.h
new file mode 100644
--- /dev/null
+++ b/Engine/Source/Runtime/RigVM/Public/RigVMModule.h
@@ -0,0 +1,11 @@
+// Copyright 1998-2019 Epic Games, Inc. All Rights Reserved.
+
+/*=============================================================================
+	RigVM.h: Render core module definitions.
+=============================================================================*/
+
+#pragma once
+
+#include "CoreMinimal.h"
+
+DECLARE_LOG_CATEGORY_EXTERN(LogRigVM, Log, All);
\ No newline at end of file
diff --git a/Engine/Source/Runtime/RigVM/RigVM.Build.cs b/Engine/Source/Runtime/RigVM/RigVM.Build.cs
new file mode 100644
--- /dev/null
+++ b/Engine/Source/Runtime/RigVM/RigVM.Build.cs
@@ -0,0 +1,17 @@
+// Copyright 1998-2019 Epic Games, Inc. All Rights Reserved.
+
+using UnrealBuildTool;
+
+public class RigVM : ModuleRules
+{
+    public RigVM(ReadOnlyTargetRules Target) : base(Target)
+    {
+        PublicDependencyModuleNames.AddRange(
+            new string[] {
+                "Core",
+                "CoreUObject",
+                "Engine",
+            }
+        );
+    }
+}
