RigVM: Renaming Multiplex to RigVM

diff --git a/Engine/Plugins/Experimental/ControlRig/Source/ControlRig/Private/Units/Animation/RigUnit_AnimEasing.cpp b/Engine/Plugins/Experimental/ControlRig/Source/ControlRig/Private/Units/Animation/RigUnit_AnimEasing.cpp
--- a/Engine/Plugins/Experimental/ControlRig/Source/ControlRig/Private/Units/Animation/RigUnit_AnimEasing.cpp
+++ b/Engine/Plugins/Experimental/ControlRig/Source/ControlRig/Private/Units/Animation/RigUnit_AnimEasing.cpp
@@ -3,7 +3,7 @@
 #include "Units/Animation/RigUnit_AnimEasing.h"
 #include "Units/RigUnitContext.h"
 
-UE_RigUnit_AnimEasing_IMPLEMENT_MULTIPLEX(void, Execute, const FRigUnitContext& Context)
+UE_RigUnit_AnimEasing_IMPLEMENT_RIGVM(void, Execute, const FRigUnitContext& Context)
 {
     DECLARE_SCOPE_HIERARCHICAL_COUNTER_RIGUNIT()
 	if (FMath::IsNearlyEqual(SourceMinimum, SourceMaximum))
diff --git a/Engine/Plugins/Experimental/ControlRig/Source/ControlRig/Private/Units/Animation/RigUnit_AnimEasing.h b/Engine/Plugins/Experimental/ControlRig/Source/ControlRig/Private/Units/Animation/RigUnit_AnimEasing.h
--- a/Engine/Plugins/Experimental/ControlRig/Source/ControlRig/Private/Units/Animation/RigUnit_AnimEasing.h
+++ b/Engine/Plugins/Experimental/ControlRig/Source/ControlRig/Private/Units/Animation/RigUnit_AnimEasing.h
@@ -39,7 +39,7 @@ struct FRigUnit_AnimEasing : public FRigUnit_AnimBase
 		SourceMaximum = TargetMaximum = 1.f;
 	}
 
-	MULTIPLEX_METHOD()
+	RIGVM_METHOD()
 	virtual void Execute(const FRigUnitContext& Context) override;
 
 	UPROPERTY(meta=(Input))
diff --git a/Engine/Plugins/Experimental/ControlRig/Source/ControlRig/Private/Units/Animation/RigUnit_AnimEvalRichCurve.cpp b/Engine/Plugins/Experimental/ControlRig/Source/ControlRig/Private/Units/Animation/RigUnit_AnimEvalRichCurve.cpp
--- a/Engine/Plugins/Experimental/ControlRig/Source/ControlRig/Private/Units/Animation/RigUnit_AnimEvalRichCurve.cpp
+++ b/Engine/Plugins/Experimental/ControlRig/Source/ControlRig/Private/Units/Animation/RigUnit_AnimEvalRichCurve.cpp
@@ -3,7 +3,7 @@
 #include "Units/Animation/RigUnit_AnimEvalRichCurve.h"
 #include "Units/RigUnitContext.h"
 
-UE_RigUnit_AnimEvalRichCurve_IMPLEMENT_MULTIPLEX(void, Execute, const FRigUnitContext& Context)
+UE_RigUnit_AnimEvalRichCurve_IMPLEMENT_RIGVM(void, Execute, const FRigUnitContext& Context)
 {
     DECLARE_SCOPE_HIERARCHICAL_COUNTER_RIGUNIT()
 	if (FMath::IsNearlyEqual(SourceMinimum, SourceMaximum))
diff --git a/Engine/Plugins/Experimental/ControlRig/Source/ControlRig/Private/Units/Animation/RigUnit_AnimEvalRichCurve.h b/Engine/Plugins/Experimental/ControlRig/Source/ControlRig/Private/Units/Animation/RigUnit_AnimEvalRichCurve.h
--- a/Engine/Plugins/Experimental/ControlRig/Source/ControlRig/Private/Units/Animation/RigUnit_AnimEvalRichCurve.h
+++ b/Engine/Plugins/Experimental/ControlRig/Source/ControlRig/Private/Units/Animation/RigUnit_AnimEvalRichCurve.h
@@ -24,7 +24,7 @@ struct FRigUnit_AnimEvalRichCurve : public FRigUnit_AnimBase
 		SourceMaximum = TargetMaximum = 1.f;
 	}
 
-	MULTIPLEX_METHOD()
+	RIGVM_METHOD()
 	virtual void Execute(const FRigUnitContext& Context) override;
 
 	UPROPERTY(meta=(Input))
diff --git a/Engine/Plugins/Experimental/ControlRig/Source/ControlRig/Private/Units/Animation/RigUnit_AnimRichCurve.cpp b/Engine/Plugins/Experimental/ControlRig/Source/ControlRig/Private/Units/Animation/RigUnit_AnimRichCurve.cpp
--- a/Engine/Plugins/Experimental/ControlRig/Source/ControlRig/Private/Units/Animation/RigUnit_AnimRichCurve.cpp
+++ b/Engine/Plugins/Experimental/ControlRig/Source/ControlRig/Private/Units/Animation/RigUnit_AnimRichCurve.cpp
@@ -3,7 +3,7 @@
 #include "Units/Animation/RigUnit_AnimRichCurve.h"
 #include "Units/RigUnitContext.h"
 
-UE_RigUnit_AnimRichCurve_IMPLEMENT_MULTIPLEX(void, Execute, const FRigUnitContext& Context)
+UE_RigUnit_AnimRichCurve_IMPLEMENT_RIGVM(void, Execute, const FRigUnitContext& Context)
 {
     DECLARE_SCOPE_HIERARCHICAL_COUNTER_RIGUNIT()
 }
diff --git a/Engine/Plugins/Experimental/ControlRig/Source/ControlRig/Private/Units/Animation/RigUnit_AnimRichCurve.h b/Engine/Plugins/Experimental/ControlRig/Source/ControlRig/Private/Units/Animation/RigUnit_AnimRichCurve.h
--- a/Engine/Plugins/Experimental/ControlRig/Source/ControlRig/Private/Units/Animation/RigUnit_AnimRichCurve.h
+++ b/Engine/Plugins/Experimental/ControlRig/Source/ControlRig/Private/Units/Animation/RigUnit_AnimRichCurve.h
@@ -21,7 +21,7 @@ struct FRigUnit_AnimRichCurve : public FRigUnit_AnimBase
 		Curve.GetRichCurve()->AddKey(1.f, 1.f);
 	}
 
-	MULTIPLEX_METHOD()
+	RIGVM_METHOD()
 	virtual void Execute(const FRigUnitContext& Context) override;
 
 	UPROPERTY(meta=(Input, Output, Constant))
diff --git a/Engine/Plugins/Experimental/ControlRig/Source/ControlRig/Private/Units/Animation/RigUnit_GetDeltaTime.cpp b/Engine/Plugins/Experimental/ControlRig/Source/ControlRig/Private/Units/Animation/RigUnit_GetDeltaTime.cpp
--- a/Engine/Plugins/Experimental/ControlRig/Source/ControlRig/Private/Units/Animation/RigUnit_GetDeltaTime.cpp
+++ b/Engine/Plugins/Experimental/ControlRig/Source/ControlRig/Private/Units/Animation/RigUnit_GetDeltaTime.cpp
@@ -3,7 +3,7 @@
 #include "Units/Animation/RigUnit_GetDeltaTime.h"
 #include "Units/RigUnitContext.h"
 
-UE_RigUnit_GetDeltaTime_IMPLEMENT_MULTIPLEX(void, Execute, const FRigUnitContext& Context)
+UE_RigUnit_GetDeltaTime_IMPLEMENT_RIGVM(void, Execute, const FRigUnitContext& Context)
 {
     DECLARE_SCOPE_HIERARCHICAL_COUNTER_RIGUNIT()
 	Result = Context.DeltaTime;
diff --git a/Engine/Plugins/Experimental/ControlRig/Source/ControlRig/Private/Units/Animation/RigUnit_GetDeltaTime.h b/Engine/Plugins/Experimental/ControlRig/Source/ControlRig/Private/Units/Animation/RigUnit_GetDeltaTime.h
--- a/Engine/Plugins/Experimental/ControlRig/Source/ControlRig/Private/Units/Animation/RigUnit_GetDeltaTime.h
+++ b/Engine/Plugins/Experimental/ControlRig/Source/ControlRig/Private/Units/Animation/RigUnit_GetDeltaTime.h
@@ -18,7 +18,7 @@ struct FRigUnit_GetDeltaTime : public FRigUnit_AnimBase
 		Result = 0.f;
 	}
 
-	MULTIPLEX_METHOD()
+	RIGVM_METHOD()
 	virtual void Execute(const FRigUnitContext& Context) override;
 
 	UPROPERTY(meta=(Output))
diff --git a/Engine/Plugins/Experimental/ControlRig/Source/ControlRig/Private/Units/Animation/RigUnit_GetWorldTime.cpp b/Engine/Plugins/Experimental/ControlRig/Source/ControlRig/Private/Units/Animation/RigUnit_GetWorldTime.cpp
--- a/Engine/Plugins/Experimental/ControlRig/Source/ControlRig/Private/Units/Animation/RigUnit_GetWorldTime.cpp
+++ b/Engine/Plugins/Experimental/ControlRig/Source/ControlRig/Private/Units/Animation/RigUnit_GetWorldTime.cpp
@@ -3,7 +3,7 @@
 #include "Units/Animation/RigUnit_GetWorldTime.h"
 #include "Units/RigUnitContext.h"
 
-UE_RigUnit_GetWorldTime_IMPLEMENT_MULTIPLEX(void, Execute, const FRigUnitContext& Context)
+UE_RigUnit_GetWorldTime_IMPLEMENT_RIGVM(void, Execute, const FRigUnitContext& Context)
 {
     DECLARE_SCOPE_HIERARCHICAL_COUNTER_RIGUNIT()
 	FDateTime Now = FDateTime::Now();
diff --git a/Engine/Plugins/Experimental/ControlRig/Source/ControlRig/Private/Units/Animation/RigUnit_GetWorldTime.h b/Engine/Plugins/Experimental/ControlRig/Source/ControlRig/Private/Units/Animation/RigUnit_GetWorldTime.h
--- a/Engine/Plugins/Experimental/ControlRig/Source/ControlRig/Private/Units/Animation/RigUnit_GetWorldTime.h
+++ b/Engine/Plugins/Experimental/ControlRig/Source/ControlRig/Private/Units/Animation/RigUnit_GetWorldTime.h
@@ -18,7 +18,7 @@ struct FRigUnit_GetWorldTime : public FRigUnit_AnimBase
 		Year = Month = Day = WeekDay = Hours = Minutes = Seconds = OverallSeconds = 0.f;
 	}
 
-	MULTIPLEX_METHOD()
+	RIGVM_METHOD()
 	virtual void Execute(const FRigUnitContext& Context) override;
 
 	UPROPERTY(meta = (Output))
diff --git a/Engine/Plugins/Experimental/ControlRig/Source/ControlRig/Private/Units/Control/RigUnit_Control.cpp b/Engine/Plugins/Experimental/ControlRig/Source/ControlRig/Private/Units/Control/RigUnit_Control.cpp
--- a/Engine/Plugins/Experimental/ControlRig/Source/ControlRig/Private/Units/Control/RigUnit_Control.cpp
+++ b/Engine/Plugins/Experimental/ControlRig/Source/ControlRig/Private/Units/Control/RigUnit_Control.cpp
@@ -3,7 +3,7 @@
 #include "Units/Control/RigUnit_Control.h"
 #include "Units/RigUnitContext.h"
 
-UE_RigUnit_Control_IMPLEMENT_MULTIPLEX(void, Execute, const FRigUnitContext& Context)
+UE_RigUnit_Control_IMPLEMENT_RIGVM(void, Execute, const FRigUnitContext& Context)
 {
     DECLARE_SCOPE_HIERARCHICAL_COUNTER_RIGUNIT()
 	if (Context.State == EControlRigState::Init)
diff --git a/Engine/Plugins/Experimental/ControlRig/Source/ControlRig/Private/Units/Debug/RigUnit_DebugBezier.cpp b/Engine/Plugins/Experimental/ControlRig/Source/ControlRig/Private/Units/Debug/RigUnit_DebugBezier.cpp
--- a/Engine/Plugins/Experimental/ControlRig/Source/ControlRig/Private/Units/Debug/RigUnit_DebugBezier.cpp
+++ b/Engine/Plugins/Experimental/ControlRig/Source/ControlRig/Private/Units/Debug/RigUnit_DebugBezier.cpp
@@ -3,7 +3,7 @@
 #include "Units/Debug/RigUnit_DebugBezier.h"
 #include "Units/RigUnitContext.h"
 
-UE_RigUnit_DebugBezier_IMPLEMENT_MULTIPLEX(void, Execute, const FRigUnitContext& Context)
+UE_RigUnit_DebugBezier_IMPLEMENT_RIGVM(void, Execute, const FRigUnitContext& Context)
 {
     DECLARE_SCOPE_HIERARCHICAL_COUNTER_RIGUNIT()
 	if (Context.State == EControlRigState::Init)
diff --git a/Engine/Plugins/Experimental/ControlRig/Source/ControlRig/Private/Units/Debug/RigUnit_DebugBezier.h b/Engine/Plugins/Experimental/ControlRig/Source/ControlRig/Private/Units/Debug/RigUnit_DebugBezier.h
--- a/Engine/Plugins/Experimental/ControlRig/Source/ControlRig/Private/Units/Debug/RigUnit_DebugBezier.h
+++ b/Engine/Plugins/Experimental/ControlRig/Source/ControlRig/Private/Units/Debug/RigUnit_DebugBezier.h
@@ -23,7 +23,7 @@ struct FRigUnit_DebugBezier : public FRigUnit_DebugBaseMutable
 		bEnabled = true;
 	}
 
-	MULTIPLEX_METHOD()
+	RIGVM_METHOD()
 	virtual void Execute(const FRigUnitContext& Context) override;
 
 	UPROPERTY(meta = (Input))
diff --git a/Engine/Plugins/Experimental/ControlRig/Source/ControlRig/Private/Units/Debug/RigUnit_DebugHierarchy.cpp b/Engine/Plugins/Experimental/ControlRig/Source/ControlRig/Private/Units/Debug/RigUnit_DebugHierarchy.cpp
--- a/Engine/Plugins/Experimental/ControlRig/Source/ControlRig/Private/Units/Debug/RigUnit_DebugHierarchy.cpp
+++ b/Engine/Plugins/Experimental/ControlRig/Source/ControlRig/Private/Units/Debug/RigUnit_DebugHierarchy.cpp
@@ -3,7 +3,7 @@
 #include "Units/Debug/RigUnit_DebugHierarchy.h"
 #include "Units/RigUnitContext.h"
 
-UE_RigUnit_DebugHierarchy_IMPLEMENT_MULTIPLEX(void, Execute, const FRigUnitContext& Context)
+UE_RigUnit_DebugHierarchy_IMPLEMENT_RIGVM(void, Execute, const FRigUnitContext& Context)
 {
     DECLARE_SCOPE_HIERARCHICAL_COUNTER_RIGUNIT()
 	if (Context.State == EControlRigState::Init)
diff --git a/Engine/Plugins/Experimental/ControlRig/Source/ControlRig/Private/Units/Debug/RigUnit_DebugHierarchy.h b/Engine/Plugins/Experimental/ControlRig/Source/ControlRig/Private/Units/Debug/RigUnit_DebugHierarchy.h
--- a/Engine/Plugins/Experimental/ControlRig/Source/ControlRig/Private/Units/Debug/RigUnit_DebugHierarchy.h
+++ b/Engine/Plugins/Experimental/ControlRig/Source/ControlRig/Private/Units/Debug/RigUnit_DebugHierarchy.h
@@ -21,7 +21,7 @@ struct FRigUnit_DebugHierarchy : public FRigUnit_DebugBaseMutable
 		bEnabled = true;
 	}
 
-	MULTIPLEX_METHOD()
+	RIGVM_METHOD()
 	virtual void Execute(const FRigUnitContext& Context) override;
 
 	UPROPERTY(meta = (Input))
diff --git a/Engine/Plugins/Experimental/ControlRig/Source/ControlRig/Private/Units/Debug/RigUnit_DebugLine.cpp b/Engine/Plugins/Experimental/ControlRig/Source/ControlRig/Private/Units/Debug/RigUnit_DebugLine.cpp
--- a/Engine/Plugins/Experimental/ControlRig/Source/ControlRig/Private/Units/Debug/RigUnit_DebugLine.cpp
+++ b/Engine/Plugins/Experimental/ControlRig/Source/ControlRig/Private/Units/Debug/RigUnit_DebugLine.cpp
@@ -3,7 +3,7 @@
 #include "Units/Debug/RigUnit_DebugLine.h"
 #include "Units/RigUnitContext.h"
 
-UE_RigUnit_DebugLine_IMPLEMENT_MULTIPLEX(void, Execute, const FRigUnitContext& Context)
+UE_RigUnit_DebugLine_IMPLEMENT_RIGVM(void, Execute, const FRigUnitContext& Context)
 {
     DECLARE_SCOPE_HIERARCHICAL_COUNTER_RIGUNIT()
 	if (Context.State == EControlRigState::Init)
diff --git a/Engine/Plugins/Experimental/ControlRig/Source/ControlRig/Private/Units/Debug/RigUnit_DebugLine.h b/Engine/Plugins/Experimental/ControlRig/Source/ControlRig/Private/Units/Debug/RigUnit_DebugLine.h
--- a/Engine/Plugins/Experimental/ControlRig/Source/ControlRig/Private/Units/Debug/RigUnit_DebugLine.h
+++ b/Engine/Plugins/Experimental/ControlRig/Source/ControlRig/Private/Units/Debug/RigUnit_DebugLine.h
@@ -19,7 +19,7 @@ struct FRigUnit_DebugLine : public FRigUnit_DebugBaseMutable
 		bEnabled = true;
 	}
 
-	MULTIPLEX_METHOD()
+	RIGVM_METHOD()
 	virtual void Execute(const FRigUnitContext& Context) override;
 
 	UPROPERTY(meta = (Input))
diff --git a/Engine/Plugins/Experimental/ControlRig/Source/ControlRig/Private/Units/Debug/RigUnit_DebugLineStrip.cpp b/Engine/Plugins/Experimental/ControlRig/Source/ControlRig/Private/Units/Debug/RigUnit_DebugLineStrip.cpp
--- a/Engine/Plugins/Experimental/ControlRig/Source/ControlRig/Private/Units/Debug/RigUnit_DebugLineStrip.cpp
+++ b/Engine/Plugins/Experimental/ControlRig/Source/ControlRig/Private/Units/Debug/RigUnit_DebugLineStrip.cpp
@@ -3,7 +3,7 @@
 #include "Units/Debug/RigUnit_DebugLineStrip.h"
 #include "Units/RigUnitContext.h"
 
-UE_RigUnit_DebugLineStrip_IMPLEMENT_MULTIPLEX(void, Execute, const FRigUnitContext& Context)
+UE_RigUnit_DebugLineStrip_IMPLEMENT_RIGVM(void, Execute, const FRigUnitContext& Context)
 {
     DECLARE_SCOPE_HIERARCHICAL_COUNTER_RIGUNIT()
 	if (Context.State == EControlRigState::Init)
diff --git a/Engine/Plugins/Experimental/ControlRig/Source/ControlRig/Private/Units/Debug/RigUnit_DebugLineStrip.h b/Engine/Plugins/Experimental/ControlRig/Source/ControlRig/Private/Units/Debug/RigUnit_DebugLineStrip.h
--- a/Engine/Plugins/Experimental/ControlRig/Source/ControlRig/Private/Units/Debug/RigUnit_DebugLineStrip.h
+++ b/Engine/Plugins/Experimental/ControlRig/Source/ControlRig/Private/Units/Debug/RigUnit_DebugLineStrip.h
@@ -18,7 +18,7 @@ struct FRigUnit_DebugLineStrip : public FRigUnit_DebugBaseMutable
 		bEnabled = true;
 	}
 
-	MULTIPLEX_METHOD()
+	RIGVM_METHOD()
 	virtual void Execute(const FRigUnitContext& Context) override;
 
 	UPROPERTY(meta = (Input))
diff --git a/Engine/Plugins/Experimental/ControlRig/Source/ControlRig/Private/Units/Debug/RigUnit_DebugPoint.cpp b/Engine/Plugins/Experimental/ControlRig/Source/ControlRig/Private/Units/Debug/RigUnit_DebugPoint.cpp
--- a/Engine/Plugins/Experimental/ControlRig/Source/ControlRig/Private/Units/Debug/RigUnit_DebugPoint.cpp
+++ b/Engine/Plugins/Experimental/ControlRig/Source/ControlRig/Private/Units/Debug/RigUnit_DebugPoint.cpp
@@ -3,7 +3,7 @@
 #include "Units/Debug/RigUnit_DebugPoint.h"
 #include "Units/RigUnitContext.h"
 
-UE_RigUnit_DebugPoint_IMPLEMENT_MULTIPLEX(void, Execute, const FRigUnitContext& Context)
+UE_RigUnit_DebugPoint_IMPLEMENT_RIGVM(void, Execute, const FRigUnitContext& Context)
 {
     DECLARE_SCOPE_HIERARCHICAL_COUNTER_RIGUNIT()
 	if (Context.State == EControlRigState::Init)
@@ -40,7 +40,7 @@ UE_RigUnit_DebugPoint_IMPLEMENT_MULTIPLEX(void, Execute, const FRigUnitContext&
 	}
 }
 
-UE_RigUnit_DebugPointMutable_IMPLEMENT_MULTIPLEX(void, Execute, const FRigUnitContext& Context)
+UE_RigUnit_DebugPointMutable_IMPLEMENT_RIGVM(void, Execute, const FRigUnitContext& Context)
 {
     DECLARE_SCOPE_HIERARCHICAL_COUNTER_RIGUNIT()
 	if (Context.State == EControlRigState::Init)
diff --git a/Engine/Plugins/Experimental/ControlRig/Source/ControlRig/Private/Units/Debug/RigUnit_DebugPoint.h b/Engine/Plugins/Experimental/ControlRig/Source/ControlRig/Private/Units/Debug/RigUnit_DebugPoint.h
--- a/Engine/Plugins/Experimental/ControlRig/Source/ControlRig/Private/Units/Debug/RigUnit_DebugPoint.h
+++ b/Engine/Plugins/Experimental/ControlRig/Source/ControlRig/Private/Units/Debug/RigUnit_DebugPoint.h
@@ -34,7 +34,7 @@ struct FRigUnit_DebugPoint : public FRigUnit_DebugBase
 		bEnabled = true;
 	}
 
-	MULTIPLEX_METHOD()
+	RIGVM_METHOD()
 	virtual void Execute(const FRigUnitContext& Context) override;
 
 	UPROPERTY(meta = (Input, Output))
@@ -78,7 +78,7 @@ struct FRigUnit_DebugPointMutable : public FRigUnit_DebugBaseMutable
 		bEnabled = true;
 	}
 
-	MULTIPLEX_METHOD()
+	RIGVM_METHOD()
 	virtual void Execute(const FRigUnitContext& Context) override;
 
 	UPROPERTY(meta = (Input))
diff --git a/Engine/Plugins/Experimental/ControlRig/Source/ControlRig/Private/Units/Debug/RigUnit_DebugPrimitives.cpp b/Engine/Plugins/Experimental/ControlRig/Source/ControlRig/Private/Units/Debug/RigUnit_DebugPrimitives.cpp
--- a/Engine/Plugins/Experimental/ControlRig/Source/ControlRig/Private/Units/Debug/RigUnit_DebugPrimitives.cpp
+++ b/Engine/Plugins/Experimental/ControlRig/Source/ControlRig/Private/Units/Debug/RigUnit_DebugPrimitives.cpp
@@ -3,7 +3,7 @@
 #include "Units/Debug/RigUnit_DebugPrimitives.h"
 #include "Units/RigUnitContext.h"
 
-UE_RigUnit_DebugRectangle_IMPLEMENT_MULTIPLEX(void, Execute, const FRigUnitContext& Context)
+UE_RigUnit_DebugRectangle_IMPLEMENT_RIGVM(void, Execute, const FRigUnitContext& Context)
 {
     DECLARE_SCOPE_HIERARCHICAL_COUNTER_RIGUNIT()
 	if (Context.State == EControlRigState::Init)
@@ -25,7 +25,7 @@ UE_RigUnit_DebugRectangle_IMPLEMENT_MULTIPLEX(void, Execute, const FRigUnitConte
 	Context.DrawInterface->DrawRectangle(WorldOffset, DrawTransform, Scale, Color, Thickness);
 }
 
-UE_RigUnit_DebugArc_IMPLEMENT_MULTIPLEX(void, Execute, const FRigUnitContext& Context)
+UE_RigUnit_DebugArc_IMPLEMENT_RIGVM(void, Execute, const FRigUnitContext& Context)
 {
     DECLARE_SCOPE_HIERARCHICAL_COUNTER_RIGUNIT()
 	if (Context.State == EControlRigState::Init)
diff --git a/Engine/Plugins/Experimental/ControlRig/Source/ControlRig/Private/Units/Debug/RigUnit_DebugPrimitives.h b/Engine/Plugins/Experimental/ControlRig/Source/ControlRig/Private/Units/Debug/RigUnit_DebugPrimitives.h
--- a/Engine/Plugins/Experimental/ControlRig/Source/ControlRig/Private/Units/Debug/RigUnit_DebugPrimitives.h
+++ b/Engine/Plugins/Experimental/ControlRig/Source/ControlRig/Private/Units/Debug/RigUnit_DebugPrimitives.h
@@ -20,7 +20,7 @@ struct FRigUnit_DebugRectangle : public FRigUnit_DebugBaseMutable
 		bEnabled = true;
 	}
 
-	MULTIPLEX_METHOD()
+	RIGVM_METHOD()
 	virtual void Execute(const FRigUnitContext& Context) override;
 
 	UPROPERTY(meta = (Input))
@@ -62,7 +62,7 @@ struct FRigUnit_DebugArc : public FRigUnit_DebugBaseMutable
 		bEnabled = true;
 	}
 
-	MULTIPLEX_METHOD()
+	RIGVM_METHOD()
 	virtual void Execute(const FRigUnitContext& Context) override;
 
 	UPROPERTY(meta = (Input))
diff --git a/Engine/Plugins/Experimental/ControlRig/Source/ControlRig/Private/Units/Debug/RigUnit_DebugTransform.cpp b/Engine/Plugins/Experimental/ControlRig/Source/ControlRig/Private/Units/Debug/RigUnit_DebugTransform.cpp
--- a/Engine/Plugins/Experimental/ControlRig/Source/ControlRig/Private/Units/Debug/RigUnit_DebugTransform.cpp
+++ b/Engine/Plugins/Experimental/ControlRig/Source/ControlRig/Private/Units/Debug/RigUnit_DebugTransform.cpp
@@ -3,7 +3,7 @@
 #include "Units/Debug/RigUnit_DebugTransform.h"
 #include "Units/RigUnitContext.h"
 
-UE_RigUnit_DebugTransform_IMPLEMENT_MULTIPLEX(void, Execute, const FRigUnitContext& Context)
+UE_RigUnit_DebugTransform_IMPLEMENT_RIGVM(void, Execute, const FRigUnitContext& Context)
 {
     DECLARE_SCOPE_HIERARCHICAL_COUNTER_RIGUNIT()
 	if (Context.State == EControlRigState::Init)
@@ -43,7 +43,7 @@ UE_RigUnit_DebugTransform_IMPLEMENT_MULTIPLEX(void, Execute, const FRigUnitConte
 	}
 }
 
-UE_RigUnit_DebugTransformMutable_IMPLEMENT_MULTIPLEX(void, Execute, const FRigUnitContext& Context)
+UE_RigUnit_DebugTransformMutable_IMPLEMENT_RIGVM(void, Execute, const FRigUnitContext& Context)
 {
     DECLARE_SCOPE_HIERARCHICAL_COUNTER_RIGUNIT()
 	if (Context.State == EControlRigState::Init)
@@ -83,7 +83,7 @@ UE_RigUnit_DebugTransformMutable_IMPLEMENT_MULTIPLEX(void, Execute, const FRigUn
 	}
 }
 
-UE_RigUnit_DebugTransformArrayMutable_IMPLEMENT_MULTIPLEX(void, Execute, const FRigUnitContext& Context)
+UE_RigUnit_DebugTransformArrayMutable_IMPLEMENT_RIGVM(void, Execute, const FRigUnitContext& Context)
 {
     DECLARE_SCOPE_HIERARCHICAL_COUNTER_RIGUNIT()
 	if (Context.State == EControlRigState::Init)
diff --git a/Engine/Plugins/Experimental/ControlRig/Source/ControlRig/Private/Units/Debug/RigUnit_DebugTransform.h b/Engine/Plugins/Experimental/ControlRig/Source/ControlRig/Private/Units/Debug/RigUnit_DebugTransform.h
--- a/Engine/Plugins/Experimental/ControlRig/Source/ControlRig/Private/Units/Debug/RigUnit_DebugTransform.h
+++ b/Engine/Plugins/Experimental/ControlRig/Source/ControlRig/Private/Units/Debug/RigUnit_DebugTransform.h
@@ -36,7 +36,7 @@ struct FRigUnit_DebugTransform : public FRigUnit_DebugBase
 		bEnabled = true;
 	}
 
-	MULTIPLEX_METHOD()
+	RIGVM_METHOD()
 	virtual void Execute(const FRigUnitContext& Context) override;
 
 	UPROPERTY(meta = (Input, Output))
@@ -79,7 +79,7 @@ struct FRigUnit_DebugTransformMutable : public FRigUnit_DebugBaseMutable
 		bEnabled = true;
 	}
 
-	MULTIPLEX_METHOD()
+	RIGVM_METHOD()
 	virtual void Execute(const FRigUnitContext& Context) override;
 
 	UPROPERTY(meta = (Input))
@@ -131,7 +131,7 @@ struct FRigUnit_DebugTransformArrayMutable : public FRigUnit_DebugBaseMutable
 		bEnabled = true;
 	}
 
-	MULTIPLEX_METHOD()
+	RIGVM_METHOD()
 	virtual void Execute(const FRigUnitContext& Context) override;
 
 	UPROPERTY(meta = (Input))
diff --git a/Engine/Plugins/Experimental/ControlRig/Source/ControlRig/Private/Units/Debug/RigUnit_ProfilingBracket.cpp b/Engine/Plugins/Experimental/ControlRig/Source/ControlRig/Private/Units/Debug/RigUnit_ProfilingBracket.cpp
--- a/Engine/Plugins/Experimental/ControlRig/Source/ControlRig/Private/Units/Debug/RigUnit_ProfilingBracket.cpp
+++ b/Engine/Plugins/Experimental/ControlRig/Source/ControlRig/Private/Units/Debug/RigUnit_ProfilingBracket.cpp
@@ -4,7 +4,7 @@
 #include "Units/RigUnitContext.h"
 #include "KismetAnimationLibrary.h"
 
-UE_RigUnit_StartProfilingTimer_IMPLEMENT_MULTIPLEX(void, Execute, const FRigUnitContext& Context)
+UE_RigUnit_StartProfilingTimer_IMPLEMENT_RIGVM(void, Execute, const FRigUnitContext& Context)
 {
     DECLARE_SCOPE_HIERARCHICAL_COUNTER_RIGUNIT()
 	if (Context.State == EControlRigState::Update)
@@ -13,7 +13,7 @@ UE_RigUnit_StartProfilingTimer_IMPLEMENT_MULTIPLEX(void, Execute, const FRigUnit
 	}
 }
 
-UE_RigUnit_EndProfilingTimer_IMPLEMENT_MULTIPLEX(void, Execute, const FRigUnitContext& Context)
+UE_RigUnit_EndProfilingTimer_IMPLEMENT_RIGVM(void, Execute, const FRigUnitContext& Context)
 {
     DECLARE_SCOPE_HIERARCHICAL_COUNTER_RIGUNIT()
 	if (Context.State == EControlRigState::Init)
diff --git a/Engine/Plugins/Experimental/ControlRig/Source/ControlRig/Private/Units/Debug/RigUnit_ProfilingBracket.h b/Engine/Plugins/Experimental/ControlRig/Source/ControlRig/Private/Units/Debug/RigUnit_ProfilingBracket.h
--- a/Engine/Plugins/Experimental/ControlRig/Source/ControlRig/Private/Units/Debug/RigUnit_ProfilingBracket.h
+++ b/Engine/Plugins/Experimental/ControlRig/Source/ControlRig/Private/Units/Debug/RigUnit_ProfilingBracket.h
@@ -14,7 +14,7 @@ struct FRigUnit_StartProfilingTimer : public FRigUnit_DebugBaseMutable
 	{
 	}
 
-	MULTIPLEX_METHOD()
+	RIGVM_METHOD()
 	virtual void Execute(const FRigUnitContext& Context) override;
 };
 
@@ -31,7 +31,7 @@ struct FRigUnit_EndProfilingTimer : public FRigUnit_DebugBaseMutable
 		Prefix = TEXT("Timer");
 	}
 
-	MULTIPLEX_METHOD()
+	RIGVM_METHOD()
 	virtual void Execute(const FRigUnitContext& Context) override;
 
 	UPROPERTY(meta = (Input, Constant))
diff --git a/Engine/Plugins/Experimental/ControlRig/Source/ControlRig/Private/Units/Execution/RigUnit_BeginExecution.cpp b/Engine/Plugins/Experimental/ControlRig/Source/ControlRig/Private/Units/Execution/RigUnit_BeginExecution.cpp
--- a/Engine/Plugins/Experimental/ControlRig/Source/ControlRig/Private/Units/Execution/RigUnit_BeginExecution.cpp
+++ b/Engine/Plugins/Experimental/ControlRig/Source/ControlRig/Private/Units/Execution/RigUnit_BeginExecution.cpp
@@ -3,7 +3,7 @@
 #include "RigUnit_BeginExecution.h"
 #include "Units/RigUnitContext.h"
 
-UE_RigUnit_BeginExecution_IMPLEMENT_MULTIPLEX(void, Execute, const FRigUnitContext& Context)
+UE_RigUnit_BeginExecution_IMPLEMENT_RIGVM(void, Execute, const FRigUnitContext& Context)
 {
     DECLARE_SCOPE_HIERARCHICAL_COUNTER_RIGUNIT()
 	ExecuteContext.Hierarchy = (FRigHierarchyContainer*)Context.Hierarchy;
diff --git a/Engine/Plugins/Experimental/ControlRig/Source/ControlRig/Private/Units/Execution/RigUnit_BeginExecution.h b/Engine/Plugins/Experimental/ControlRig/Source/ControlRig/Private/Units/Execution/RigUnit_BeginExecution.h
--- a/Engine/Plugins/Experimental/ControlRig/Source/ControlRig/Private/Units/Execution/RigUnit_BeginExecution.h
+++ b/Engine/Plugins/Experimental/ControlRig/Source/ControlRig/Private/Units/Execution/RigUnit_BeginExecution.h
@@ -11,7 +11,7 @@ struct FRigUnit_BeginExecution : public FRigUnit
 {
 	GENERATED_BODY()
 
-	MULTIPLEX_METHOD()
+	RIGVM_METHOD()
 	virtual void Execute(const FRigUnitContext& Context) override;
 
 	// The execution result
diff --git a/Engine/Plugins/Experimental/ControlRig/Source/ControlRig/Private/Units/Hierarchy/RigUnit_AddBoneTransform.cpp b/Engine/Plugins/Experimental/ControlRig/Source/ControlRig/Private/Units/Hierarchy/RigUnit_AddBoneTransform.cpp
--- a/Engine/Plugins/Experimental/ControlRig/Source/ControlRig/Private/Units/Hierarchy/RigUnit_AddBoneTransform.cpp
+++ b/Engine/Plugins/Experimental/ControlRig/Source/ControlRig/Private/Units/Hierarchy/RigUnit_AddBoneTransform.cpp
@@ -8,7 +8,7 @@ FString FRigUnit_AddBoneTransform::GetUnitLabel() const
 	return FString::Printf(TEXT("Offset Transform %s"), *Bone.ToString());
 }
 
-UE_RigUnit_AddBoneTransform_IMPLEMENT_MULTIPLEX(void, Execute, const FRigUnitContext& Context)
+UE_RigUnit_AddBoneTransform_IMPLEMENT_RIGVM(void, Execute, const FRigUnitContext& Context)
 {
     DECLARE_SCOPE_HIERARCHICAL_COUNTER_RIGUNIT()
 	FRigBoneHierarchy* Hierarchy = ExecuteContext.GetBones();
diff --git a/Engine/Plugins/Experimental/ControlRig/Source/ControlRig/Private/Units/Hierarchy/RigUnit_AddBoneTransform.h b/Engine/Plugins/Experimental/ControlRig/Source/ControlRig/Private/Units/Hierarchy/RigUnit_AddBoneTransform.h
--- a/Engine/Plugins/Experimental/ControlRig/Source/ControlRig/Private/Units/Hierarchy/RigUnit_AddBoneTransform.h
+++ b/Engine/Plugins/Experimental/ControlRig/Source/ControlRig/Private/Units/Hierarchy/RigUnit_AddBoneTransform.h
@@ -21,7 +21,7 @@ struct FRigUnit_AddBoneTransform : public FRigUnitMutable
 	{}
 
 	virtual FString GetUnitLabel() const override;
-	MULTIPLEX_METHOD()
+	RIGVM_METHOD()
 	virtual void Execute(const FRigUnitContext& Context) override;
 
 	/**
diff --git a/Engine/Plugins/Experimental/ControlRig/Source/ControlRig/Private/Units/Hierarchy/RigUnit_GetBoneTransform.cpp b/Engine/Plugins/Experimental/ControlRig/Source/ControlRig/Private/Units/Hierarchy/RigUnit_GetBoneTransform.cpp
--- a/Engine/Plugins/Experimental/ControlRig/Source/ControlRig/Private/Units/Hierarchy/RigUnit_GetBoneTransform.cpp
+++ b/Engine/Plugins/Experimental/ControlRig/Source/ControlRig/Private/Units/Hierarchy/RigUnit_GetBoneTransform.cpp
@@ -8,7 +8,7 @@ FString FRigUnit_GetBoneTransform::GetUnitLabel() const
 	return FString::Printf(TEXT("Get Transform %s"), *Bone.ToString());
 }
 
-UE_RigUnit_GetBoneTransform_IMPLEMENT_MULTIPLEX(void, Execute, const FRigUnitContext& Context)
+UE_RigUnit_GetBoneTransform_IMPLEMENT_RIGVM(void, Execute, const FRigUnitContext& Context)
 {
     DECLARE_SCOPE_HIERARCHICAL_COUNTER_RIGUNIT()
 	const FRigBoneHierarchy* Hierarchy = Context.GetBones();
diff --git a/Engine/Plugins/Experimental/ControlRig/Source/ControlRig/Private/Units/Hierarchy/RigUnit_GetBoneTransform.h b/Engine/Plugins/Experimental/ControlRig/Source/ControlRig/Private/Units/Hierarchy/RigUnit_GetBoneTransform.h
--- a/Engine/Plugins/Experimental/ControlRig/Source/ControlRig/Private/Units/Hierarchy/RigUnit_GetBoneTransform.h
+++ b/Engine/Plugins/Experimental/ControlRig/Source/ControlRig/Private/Units/Hierarchy/RigUnit_GetBoneTransform.h
@@ -19,7 +19,7 @@ struct FRigUnit_GetBoneTransform : public FRigUnit
 	{}
 
 	virtual FString GetUnitLabel() const override;
-	MULTIPLEX_METHOD()
+	RIGVM_METHOD()
 	virtual void Execute(const FRigUnitContext& Context) override;
 
 	/**
diff --git a/Engine/Plugins/Experimental/ControlRig/Source/ControlRig/Private/Units/Hierarchy/RigUnit_GetCurveValue.cpp b/Engine/Plugins/Experimental/ControlRig/Source/ControlRig/Private/Units/Hierarchy/RigUnit_GetCurveValue.cpp
--- a/Engine/Plugins/Experimental/ControlRig/Source/ControlRig/Private/Units/Hierarchy/RigUnit_GetCurveValue.cpp
+++ b/Engine/Plugins/Experimental/ControlRig/Source/ControlRig/Private/Units/Hierarchy/RigUnit_GetCurveValue.cpp
@@ -8,7 +8,7 @@ FString FRigUnit_GetCurveValue::GetUnitLabel() const
 	return FString::Printf(TEXT("Get Curve %s"), *Curve.ToString());
 }
 
-UE_RigUnit_GetCurveValue_IMPLEMENT_MULTIPLEX(void, Execute, const FRigUnitContext& Context)
+UE_RigUnit_GetCurveValue_IMPLEMENT_RIGVM(void, Execute, const FRigUnitContext& Context)
 {
     DECLARE_SCOPE_HIERARCHICAL_COUNTER_RIGUNIT()
 	const FRigCurveContainer* CurveContainer = Context.GetCurves();
diff --git a/Engine/Plugins/Experimental/ControlRig/Source/ControlRig/Private/Units/Hierarchy/RigUnit_GetCurveValue.h b/Engine/Plugins/Experimental/ControlRig/Source/ControlRig/Private/Units/Hierarchy/RigUnit_GetCurveValue.h
--- a/Engine/Plugins/Experimental/ControlRig/Source/ControlRig/Private/Units/Hierarchy/RigUnit_GetCurveValue.h
+++ b/Engine/Plugins/Experimental/ControlRig/Source/ControlRig/Private/Units/Hierarchy/RigUnit_GetCurveValue.h
@@ -20,7 +20,7 @@ struct FRigUnit_GetCurveValue : public FRigUnit
 	{}
 
 	virtual FString GetUnitLabel() const override;
-	MULTIPLEX_METHOD()
+	RIGVM_METHOD()
 	virtual void Execute(const FRigUnitContext& Context) override;
 
 	/**
diff --git a/Engine/Plugins/Experimental/ControlRig/Source/ControlRig/Private/Units/Hierarchy/RigUnit_GetInitialBoneTransform.cpp b/Engine/Plugins/Experimental/ControlRig/Source/ControlRig/Private/Units/Hierarchy/RigUnit_GetInitialBoneTransform.cpp
--- a/Engine/Plugins/Experimental/ControlRig/Source/ControlRig/Private/Units/Hierarchy/RigUnit_GetInitialBoneTransform.cpp
+++ b/Engine/Plugins/Experimental/ControlRig/Source/ControlRig/Private/Units/Hierarchy/RigUnit_GetInitialBoneTransform.cpp
@@ -8,7 +8,7 @@ FString FRigUnit_GetInitialBoneTransform::GetUnitLabel() const
 	return FString::Printf(TEXT("Get Initial Transform %s"), *Bone.ToString());
 }
 
-UE_RigUnit_GetInitialBoneTransform_IMPLEMENT_MULTIPLEX(void, Execute, const FRigUnitContext& Context)
+UE_RigUnit_GetInitialBoneTransform_IMPLEMENT_RIGVM(void, Execute, const FRigUnitContext& Context)
 {
     DECLARE_SCOPE_HIERARCHICAL_COUNTER_RIGUNIT()
 	const FRigBoneHierarchy* Hierarchy = Context.GetBones();
diff --git a/Engine/Plugins/Experimental/ControlRig/Source/ControlRig/Private/Units/Hierarchy/RigUnit_GetInitialBoneTransform.h b/Engine/Plugins/Experimental/ControlRig/Source/ControlRig/Private/Units/Hierarchy/RigUnit_GetInitialBoneTransform.h
--- a/Engine/Plugins/Experimental/ControlRig/Source/ControlRig/Private/Units/Hierarchy/RigUnit_GetInitialBoneTransform.h
+++ b/Engine/Plugins/Experimental/ControlRig/Source/ControlRig/Private/Units/Hierarchy/RigUnit_GetInitialBoneTransform.h
@@ -19,7 +19,7 @@ struct FRigUnit_GetInitialBoneTransform : public FRigUnit
 	{}
 
 	virtual FString GetUnitLabel() const override;
-	MULTIPLEX_METHOD()
+	RIGVM_METHOD()
 	virtual void Execute(const FRigUnitContext& Context) override;
 
 	/**
diff --git a/Engine/Plugins/Experimental/ControlRig/Source/ControlRig/Private/Units/Hierarchy/RigUnit_GetRelativeBoneTransform.cpp b/Engine/Plugins/Experimental/ControlRig/Source/ControlRig/Private/Units/Hierarchy/RigUnit_GetRelativeBoneTransform.cpp
--- a/Engine/Plugins/Experimental/ControlRig/Source/ControlRig/Private/Units/Hierarchy/RigUnit_GetRelativeBoneTransform.cpp
+++ b/Engine/Plugins/Experimental/ControlRig/Source/ControlRig/Private/Units/Hierarchy/RigUnit_GetRelativeBoneTransform.cpp
@@ -8,7 +8,7 @@ FString FRigUnit_GetRelativeBoneTransform::GetUnitLabel() const
 	return FString::Printf(TEXT("Get Relative Transform %s"), *Bone.ToString());
 }
 
-UE_RigUnit_GetRelativeBoneTransform_IMPLEMENT_MULTIPLEX(void, Execute, const FRigUnitContext& Context)
+UE_RigUnit_GetRelativeBoneTransform_IMPLEMENT_RIGVM(void, Execute, const FRigUnitContext& Context)
 {
     DECLARE_SCOPE_HIERARCHICAL_COUNTER_RIGUNIT()
 	const FRigBoneHierarchy* Hierarchy = Context.GetBones();
diff --git a/Engine/Plugins/Experimental/ControlRig/Source/ControlRig/Private/Units/Hierarchy/RigUnit_GetRelativeBoneTransform.h b/Engine/Plugins/Experimental/ControlRig/Source/ControlRig/Private/Units/Hierarchy/RigUnit_GetRelativeBoneTransform.h
--- a/Engine/Plugins/Experimental/ControlRig/Source/ControlRig/Private/Units/Hierarchy/RigUnit_GetRelativeBoneTransform.h
+++ b/Engine/Plugins/Experimental/ControlRig/Source/ControlRig/Private/Units/Hierarchy/RigUnit_GetRelativeBoneTransform.h
@@ -19,7 +19,7 @@ struct FRigUnit_GetRelativeBoneTransform : public FRigUnit
 	{}
 
 	virtual FString GetUnitLabel() const override;
-	MULTIPLEX_METHOD()
+	RIGVM_METHOD()
 	virtual void Execute(const FRigUnitContext& Context) override;
 
 	/**
diff --git a/Engine/Plugins/Experimental/ControlRig/Source/ControlRig/Private/Units/Hierarchy/RigUnit_SetBoneRotation.cpp b/Engine/Plugins/Experimental/ControlRig/Source/ControlRig/Private/Units/Hierarchy/RigUnit_SetBoneRotation.cpp
--- a/Engine/Plugins/Experimental/ControlRig/Source/ControlRig/Private/Units/Hierarchy/RigUnit_SetBoneRotation.cpp
+++ b/Engine/Plugins/Experimental/ControlRig/Source/ControlRig/Private/Units/Hierarchy/RigUnit_SetBoneRotation.cpp
@@ -8,7 +8,7 @@ FString FRigUnit_SetBoneRotation::GetUnitLabel() const
 	return FString::Printf(TEXT("Set Rotation %s"), *Bone.ToString());
 }
 
-UE_RigUnit_SetBoneRotation_IMPLEMENT_MULTIPLEX(void, Execute, const FRigUnitContext& Context)
+UE_RigUnit_SetBoneRotation_IMPLEMENT_RIGVM(void, Execute, const FRigUnitContext& Context)
 {
     DECLARE_SCOPE_HIERARCHICAL_COUNTER_RIGUNIT()
 	FRigBoneHierarchy* Hierarchy = ExecuteContext.GetBones();
diff --git a/Engine/Plugins/Experimental/ControlRig/Source/ControlRig/Private/Units/Hierarchy/RigUnit_SetBoneRotation.h b/Engine/Plugins/Experimental/ControlRig/Source/ControlRig/Private/Units/Hierarchy/RigUnit_SetBoneRotation.h
--- a/Engine/Plugins/Experimental/ControlRig/Source/ControlRig/Private/Units/Hierarchy/RigUnit_SetBoneRotation.h
+++ b/Engine/Plugins/Experimental/ControlRig/Source/ControlRig/Private/Units/Hierarchy/RigUnit_SetBoneRotation.h
@@ -22,7 +22,7 @@ struct FRigUnit_SetBoneRotation : public FRigUnitMutable
 	{}
 
 	virtual FString GetUnitLabel() const override;
-	MULTIPLEX_METHOD()
+	RIGVM_METHOD()
 	virtual void Execute(const FRigUnitContext& Context) override;
 
 	/**
diff --git a/Engine/Plugins/Experimental/ControlRig/Source/ControlRig/Private/Units/Hierarchy/RigUnit_SetBoneTransform.cpp b/Engine/Plugins/Experimental/ControlRig/Source/ControlRig/Private/Units/Hierarchy/RigUnit_SetBoneTransform.cpp
--- a/Engine/Plugins/Experimental/ControlRig/Source/ControlRig/Private/Units/Hierarchy/RigUnit_SetBoneTransform.cpp
+++ b/Engine/Plugins/Experimental/ControlRig/Source/ControlRig/Private/Units/Hierarchy/RigUnit_SetBoneTransform.cpp
@@ -8,7 +8,7 @@ FString FRigUnit_SetBoneTransform::GetUnitLabel() const
 	return FString::Printf(TEXT("Set Transform %s"), *Bone.ToString());
 }
 
-UE_RigUnit_SetBoneTransform_IMPLEMENT_MULTIPLEX(void, Execute, const FRigUnitContext& Context)
+UE_RigUnit_SetBoneTransform_IMPLEMENT_RIGVM(void, Execute, const FRigUnitContext& Context)
 {
     DECLARE_SCOPE_HIERARCHICAL_COUNTER_RIGUNIT()
 	FRigBoneHierarchy* Hierarchy = ExecuteContext.GetBones();
diff --git a/Engine/Plugins/Experimental/ControlRig/Source/ControlRig/Private/Units/Hierarchy/RigUnit_SetBoneTransform.h b/Engine/Plugins/Experimental/ControlRig/Source/ControlRig/Private/Units/Hierarchy/RigUnit_SetBoneTransform.h
--- a/Engine/Plugins/Experimental/ControlRig/Source/ControlRig/Private/Units/Hierarchy/RigUnit_SetBoneTransform.h
+++ b/Engine/Plugins/Experimental/ControlRig/Source/ControlRig/Private/Units/Hierarchy/RigUnit_SetBoneTransform.h
@@ -22,7 +22,7 @@ struct FRigUnit_SetBoneTransform : public FRigUnitMutable
 
 	virtual FString GetUnitLabel() const override;
 
-	MULTIPLEX_METHOD()
+	RIGVM_METHOD()
 	virtual void Execute(const FRigUnitContext& Context) override;
 
 	/**
diff --git a/Engine/Plugins/Experimental/ControlRig/Source/ControlRig/Private/Units/Hierarchy/RigUnit_SetBoneTranslation.cpp b/Engine/Plugins/Experimental/ControlRig/Source/ControlRig/Private/Units/Hierarchy/RigUnit_SetBoneTranslation.cpp
--- a/Engine/Plugins/Experimental/ControlRig/Source/ControlRig/Private/Units/Hierarchy/RigUnit_SetBoneTranslation.cpp
+++ b/Engine/Plugins/Experimental/ControlRig/Source/ControlRig/Private/Units/Hierarchy/RigUnit_SetBoneTranslation.cpp
@@ -8,7 +8,7 @@ FString FRigUnit_SetBoneTranslation::GetUnitLabel() const
 	return FString::Printf(TEXT("Set Translation %s"), *Bone.ToString());
 }
 
-UE_RigUnit_SetBoneTranslation_IMPLEMENT_MULTIPLEX(void, Execute, const FRigUnitContext& Context)
+UE_RigUnit_SetBoneTranslation_IMPLEMENT_RIGVM(void, Execute, const FRigUnitContext& Context)
 {
     DECLARE_SCOPE_HIERARCHICAL_COUNTER_RIGUNIT()
 	FRigBoneHierarchy* Hierarchy = ExecuteContext.GetBones();
diff --git a/Engine/Plugins/Experimental/ControlRig/Source/ControlRig/Private/Units/Hierarchy/RigUnit_SetBoneTranslation.h b/Engine/Plugins/Experimental/ControlRig/Source/ControlRig/Private/Units/Hierarchy/RigUnit_SetBoneTranslation.h
--- a/Engine/Plugins/Experimental/ControlRig/Source/ControlRig/Private/Units/Hierarchy/RigUnit_SetBoneTranslation.h
+++ b/Engine/Plugins/Experimental/ControlRig/Source/ControlRig/Private/Units/Hierarchy/RigUnit_SetBoneTranslation.h
@@ -22,7 +22,7 @@ struct FRigUnit_SetBoneTranslation : public FRigUnitMutable
 	{}
 
 	virtual FString GetUnitLabel() const override;
-	MULTIPLEX_METHOD()
+	RIGVM_METHOD()
 	virtual void Execute(const FRigUnitContext& Context) override;
 
 	/**
diff --git a/Engine/Plugins/Experimental/ControlRig/Source/ControlRig/Private/Units/Hierarchy/RigUnit_SetCurveValue.cpp b/Engine/Plugins/Experimental/ControlRig/Source/ControlRig/Private/Units/Hierarchy/RigUnit_SetCurveValue.cpp
--- a/Engine/Plugins/Experimental/ControlRig/Source/ControlRig/Private/Units/Hierarchy/RigUnit_SetCurveValue.cpp
+++ b/Engine/Plugins/Experimental/ControlRig/Source/ControlRig/Private/Units/Hierarchy/RigUnit_SetCurveValue.cpp
@@ -8,7 +8,7 @@ FString FRigUnit_SetCurveValue::GetUnitLabel() const
 	return FString::Printf(TEXT("Set Curve [%s]"), *Curve.ToString());
 }
 
-UE_RigUnit_SetCurveValue_IMPLEMENT_MULTIPLEX(void, Execute, const FRigUnitContext& Context)
+UE_RigUnit_SetCurveValue_IMPLEMENT_RIGVM(void, Execute, const FRigUnitContext& Context)
 {
     DECLARE_SCOPE_HIERARCHICAL_COUNTER_RIGUNIT()
 	FRigCurveContainer* CurveContainer = ExecuteContext.GetCurves();
diff --git a/Engine/Plugins/Experimental/ControlRig/Source/ControlRig/Private/Units/Hierarchy/RigUnit_SetCurveValue.h b/Engine/Plugins/Experimental/ControlRig/Source/ControlRig/Private/Units/Hierarchy/RigUnit_SetCurveValue.h
--- a/Engine/Plugins/Experimental/ControlRig/Source/ControlRig/Private/Units/Hierarchy/RigUnit_SetCurveValue.h
+++ b/Engine/Plugins/Experimental/ControlRig/Source/ControlRig/Private/Units/Hierarchy/RigUnit_SetCurveValue.h
@@ -20,7 +20,7 @@ struct FRigUnit_SetCurveValue : public FRigUnitMutable
 	{}
 
 	virtual FString GetUnitLabel() const override;
-	MULTIPLEX_METHOD()
+	RIGVM_METHOD()
 	virtual void Execute(const FRigUnitContext& Context) override;
 
 	/**
diff --git a/Engine/Plugins/Experimental/ControlRig/Source/ControlRig/Private/Units/Hierarchy/RigUnit_SetRelativeBoneTransform.cpp b/Engine/Plugins/Experimental/ControlRig/Source/ControlRig/Private/Units/Hierarchy/RigUnit_SetRelativeBoneTransform.cpp
--- a/Engine/Plugins/Experimental/ControlRig/Source/ControlRig/Private/Units/Hierarchy/RigUnit_SetRelativeBoneTransform.cpp
+++ b/Engine/Plugins/Experimental/ControlRig/Source/ControlRig/Private/Units/Hierarchy/RigUnit_SetRelativeBoneTransform.cpp
@@ -8,7 +8,7 @@ FString FRigUnit_SetRelativeBoneTransform::GetUnitLabel() const
 	return FString::Printf(TEXT("Set Relative Transform %s"), *Bone.ToString());
 }
 
-UE_RigUnit_SetRelativeBoneTransform_IMPLEMENT_MULTIPLEX(void, Execute, const FRigUnitContext& Context)
+UE_RigUnit_SetRelativeBoneTransform_IMPLEMENT_RIGVM(void, Execute, const FRigUnitContext& Context)
 {
     DECLARE_SCOPE_HIERARCHICAL_COUNTER_RIGUNIT()
 	FRigBoneHierarchy* Hierarchy = ExecuteContext.GetBones();
diff --git a/Engine/Plugins/Experimental/ControlRig/Source/ControlRig/Private/Units/Hierarchy/RigUnit_SetRelativeBoneTransform.h b/Engine/Plugins/Experimental/ControlRig/Source/ControlRig/Private/Units/Hierarchy/RigUnit_SetRelativeBoneTransform.h
--- a/Engine/Plugins/Experimental/ControlRig/Source/ControlRig/Private/Units/Hierarchy/RigUnit_SetRelativeBoneTransform.h
+++ b/Engine/Plugins/Experimental/ControlRig/Source/ControlRig/Private/Units/Hierarchy/RigUnit_SetRelativeBoneTransform.h
@@ -21,7 +21,7 @@ struct FRigUnit_SetRelativeBoneTransform : public FRigUnitMutable
 	{}
 
 	virtual FString GetUnitLabel() const override;
-	MULTIPLEX_METHOD()
+	RIGVM_METHOD()
 	virtual void Execute(const FRigUnitContext& Context) override;
 
 	/**
diff --git a/Engine/Plugins/Experimental/ControlRig/Source/ControlRig/Private/Units/Highlevel/Harmonics/RigUnit_BoneHarmonics.cpp b/Engine/Plugins/Experimental/ControlRig/Source/ControlRig/Private/Units/Highlevel/Harmonics/RigUnit_BoneHarmonics.cpp
--- a/Engine/Plugins/Experimental/ControlRig/Source/ControlRig/Private/Units/Highlevel/Harmonics/RigUnit_BoneHarmonics.cpp
+++ b/Engine/Plugins/Experimental/ControlRig/Source/ControlRig/Private/Units/Highlevel/Harmonics/RigUnit_BoneHarmonics.cpp
@@ -3,7 +3,7 @@
 #include "Units/Highlevel/Harmonics/RigUnit_BoneHarmonics.h"
 #include "Units/RigUnitContext.h"
 
-UE_RigUnit_BoneHarmonics_IMPLEMENT_MULTIPLEX(void, Execute, const FRigUnitContext& Context)
+UE_RigUnit_BoneHarmonics_IMPLEMENT_RIGVM(void, Execute, const FRigUnitContext& Context)
 {
     DECLARE_SCOPE_HIERARCHICAL_COUNTER_RIGUNIT()
 	FRigBoneHierarchy* Hierarchy = ExecuteContext.GetBones();
diff --git a/Engine/Plugins/Experimental/ControlRig/Source/ControlRig/Private/Units/Highlevel/Harmonics/RigUnit_BoneHarmonics.h b/Engine/Plugins/Experimental/ControlRig/Source/ControlRig/Private/Units/Highlevel/Harmonics/RigUnit_BoneHarmonics.h
--- a/Engine/Plugins/Experimental/ControlRig/Source/ControlRig/Private/Units/Highlevel/Harmonics/RigUnit_BoneHarmonics.h
+++ b/Engine/Plugins/Experimental/ControlRig/Source/ControlRig/Private/Units/Highlevel/Harmonics/RigUnit_BoneHarmonics.h
@@ -70,7 +70,7 @@ struct FRigUnit_BoneHarmonics : public FRigUnit_HighlevelBaseMutable
 		bPropagateToChildren = false;
 	}
 
-	MULTIPLEX_METHOD()
+	RIGVM_METHOD()
 	virtual void Execute(const FRigUnitContext& Context) override;
 
 	/** The bones to drive. */
diff --git a/Engine/Plugins/Experimental/ControlRig/Source/ControlRig/Private/Units/Highlevel/Harmonics/RigUnit_ChainHarmonics.cpp b/Engine/Plugins/Experimental/ControlRig/Source/ControlRig/Private/Units/Highlevel/Harmonics/RigUnit_ChainHarmonics.cpp
--- a/Engine/Plugins/Experimental/ControlRig/Source/ControlRig/Private/Units/Highlevel/Harmonics/RigUnit_ChainHarmonics.cpp
+++ b/Engine/Plugins/Experimental/ControlRig/Source/ControlRig/Private/Units/Highlevel/Harmonics/RigUnit_ChainHarmonics.cpp
@@ -3,7 +3,7 @@
 #include "Units/Highlevel/Harmonics/RigUnit_ChainHarmonics.h"
 #include "Units/RigUnitContext.h"
 
-UE_RigUnit_ChainHarmonics_IMPLEMENT_MULTIPLEX(void, Execute, const FRigUnitContext& Context)
+UE_RigUnit_ChainHarmonics_IMPLEMENT_RIGVM(void, Execute, const FRigUnitContext& Context)
 {
     DECLARE_SCOPE_HIERARCHICAL_COUNTER_RIGUNIT()
 	FRigBoneHierarchy* Hierarchy = ExecuteContext.GetBones();
diff --git a/Engine/Plugins/Experimental/ControlRig/Source/ControlRig/Private/Units/Highlevel/Harmonics/RigUnit_ChainHarmonics.h b/Engine/Plugins/Experimental/ControlRig/Source/ControlRig/Private/Units/Highlevel/Harmonics/RigUnit_ChainHarmonics.h
--- a/Engine/Plugins/Experimental/ControlRig/Source/ControlRig/Private/Units/Highlevel/Harmonics/RigUnit_ChainHarmonics.h
+++ b/Engine/Plugins/Experimental/ControlRig/Source/ControlRig/Private/Units/Highlevel/Harmonics/RigUnit_ChainHarmonics.h
@@ -197,7 +197,7 @@ struct FRigUnit_ChainHarmonics : public FRigUnit_HighlevelBaseMutable
 		DrawWorldOffset = FTransform::Identity;
 	}
 
-	MULTIPLEX_METHOD()
+	RIGVM_METHOD()
 	virtual void Execute(const FRigUnitContext& Context) override;
 
 	UPROPERTY(meta = (Input, Constant, BoneName))
diff --git a/Engine/Plugins/Experimental/ControlRig/Source/ControlRig/Private/Units/Highlevel/Hierarchy/RigUnit_AimBone.cpp b/Engine/Plugins/Experimental/ControlRig/Source/ControlRig/Private/Units/Highlevel/Hierarchy/RigUnit_AimBone.cpp
--- a/Engine/Plugins/Experimental/ControlRig/Source/ControlRig/Private/Units/Highlevel/Hierarchy/RigUnit_AimBone.cpp
+++ b/Engine/Plugins/Experimental/ControlRig/Source/ControlRig/Private/Units/Highlevel/Hierarchy/RigUnit_AimBone.cpp
@@ -3,7 +3,7 @@
 #include "Units/Highlevel/Hierarchy/RigUnit_AimBone.h"
 #include "Units/RigUnitContext.h"
 
-UE_RigUnit_AimBone_IMPLEMENT_MULTIPLEX(void, Execute, const FRigUnitContext& Context)
+UE_RigUnit_AimBone_IMPLEMENT_RIGVM(void, Execute, const FRigUnitContext& Context)
 {
     DECLARE_SCOPE_HIERARCHICAL_COUNTER_RIGUNIT()
 	FRigBoneHierarchy* Hierarchy = ExecuteContext.GetBones();
diff --git a/Engine/Plugins/Experimental/ControlRig/Source/ControlRig/Private/Units/Highlevel/Hierarchy/RigUnit_AimBone.h b/Engine/Plugins/Experimental/ControlRig/Source/ControlRig/Private/Units/Highlevel/Hierarchy/RigUnit_AimBone.h
--- a/Engine/Plugins/Experimental/ControlRig/Source/ControlRig/Private/Units/Highlevel/Hierarchy/RigUnit_AimBone.h
+++ b/Engine/Plugins/Experimental/ControlRig/Source/ControlRig/Private/Units/Highlevel/Hierarchy/RigUnit_AimBone.h
@@ -106,7 +106,7 @@ struct FRigUnit_AimBone : public FRigUnit_HighlevelBaseMutable
 		SecondaryCachedSpaceIndex = INDEX_NONE;
 	}
 
-	MULTIPLEX_METHOD()
+	RIGVM_METHOD()
 	virtual void Execute(const FRigUnitContext& Context) override;
 
 	/** 
diff --git a/Engine/Plugins/Experimental/ControlRig/Source/ControlRig/Private/Units/Highlevel/Hierarchy/RigUnit_CCDIK.cpp b/Engine/Plugins/Experimental/ControlRig/Source/ControlRig/Private/Units/Highlevel/Hierarchy/RigUnit_CCDIK.cpp
--- a/Engine/Plugins/Experimental/ControlRig/Source/ControlRig/Private/Units/Highlevel/Hierarchy/RigUnit_CCDIK.cpp
+++ b/Engine/Plugins/Experimental/ControlRig/Source/ControlRig/Private/Units/Highlevel/Hierarchy/RigUnit_CCDIK.cpp
@@ -3,7 +3,7 @@
 #include "RigUnit_CCDIK.h"
 #include "Units/RigUnitContext.h"
 
-UE_RigUnit_CCDIK_IMPLEMENT_MULTIPLEX(void, Execute, const FRigUnitContext& Context)
+UE_RigUnit_CCDIK_IMPLEMENT_RIGVM(void, Execute, const FRigUnitContext& Context)
 {
     DECLARE_SCOPE_HIERARCHICAL_COUNTER_RIGUNIT()
 	FRigBoneHierarchy* Hierarchy = ExecuteContext.GetBones();
diff --git a/Engine/Plugins/Experimental/ControlRig/Source/ControlRig/Private/Units/Highlevel/Hierarchy/RigUnit_CCDIK.h b/Engine/Plugins/Experimental/ControlRig/Source/ControlRig/Private/Units/Highlevel/Hierarchy/RigUnit_CCDIK.h
--- a/Engine/Plugins/Experimental/ControlRig/Source/ControlRig/Private/Units/Highlevel/Hierarchy/RigUnit_CCDIK.h
+++ b/Engine/Plugins/Experimental/ControlRig/Source/ControlRig/Private/Units/Highlevel/Hierarchy/RigUnit_CCDIK.h
@@ -66,7 +66,7 @@ struct FRigUnit_CCDIK : public FRigUnit_HighlevelBaseMutable
 {
 	GENERATED_BODY()
 
-	MULTIPLEX_METHOD()
+	RIGVM_METHOD()
 	virtual void Execute(const FRigUnitContext& Context) override;
 
 	FRigUnit_CCDIK()
diff --git a/Engine/Plugins/Experimental/ControlRig/Source/ControlRig/Private/Units/Highlevel/Hierarchy/RigUnit_DistributeRotation.cpp b/Engine/Plugins/Experimental/ControlRig/Source/ControlRig/Private/Units/Highlevel/Hierarchy/RigUnit_DistributeRotation.cpp
--- a/Engine/Plugins/Experimental/ControlRig/Source/ControlRig/Private/Units/Highlevel/Hierarchy/RigUnit_DistributeRotation.cpp
+++ b/Engine/Plugins/Experimental/ControlRig/Source/ControlRig/Private/Units/Highlevel/Hierarchy/RigUnit_DistributeRotation.cpp
@@ -3,7 +3,7 @@
 #include "Units/Highlevel/Hierarchy/RigUnit_DistributeRotation.h"
 #include "Units/RigUnitContext.h"
 
-UE_RigUnit_DistributeRotation_IMPLEMENT_MULTIPLEX(void, Execute, const FRigUnitContext& Context)
+UE_RigUnit_DistributeRotation_IMPLEMENT_RIGVM(void, Execute, const FRigUnitContext& Context)
 {
     DECLARE_SCOPE_HIERARCHICAL_COUNTER_RIGUNIT()
 	FRigBoneHierarchy* Hierarchy = ExecuteContext.GetBones();
diff --git a/Engine/Plugins/Experimental/ControlRig/Source/ControlRig/Private/Units/Highlevel/Hierarchy/RigUnit_DistributeRotation.h b/Engine/Plugins/Experimental/ControlRig/Source/ControlRig/Private/Units/Highlevel/Hierarchy/RigUnit_DistributeRotation.h
--- a/Engine/Plugins/Experimental/ControlRig/Source/ControlRig/Private/Units/Highlevel/Hierarchy/RigUnit_DistributeRotation.h
+++ b/Engine/Plugins/Experimental/ControlRig/Source/ControlRig/Private/Units/Highlevel/Hierarchy/RigUnit_DistributeRotation.h
@@ -68,7 +68,7 @@ struct FRigUnit_DistributeRotation : public FRigUnit_HighlevelBaseMutable
 		bPropagateToChildren = false;
 	}
 
-	MULTIPLEX_METHOD()
+	RIGVM_METHOD()
 	virtual void Execute(const FRigUnitContext& Context) override;
 
 	/** 
diff --git a/Engine/Plugins/Experimental/ControlRig/Source/ControlRig/Private/Units/Highlevel/Hierarchy/RigUnit_FABRIK.cpp b/Engine/Plugins/Experimental/ControlRig/Source/ControlRig/Private/Units/Highlevel/Hierarchy/RigUnit_FABRIK.cpp
--- a/Engine/Plugins/Experimental/ControlRig/Source/ControlRig/Private/Units/Highlevel/Hierarchy/RigUnit_FABRIK.cpp
+++ b/Engine/Plugins/Experimental/ControlRig/Source/ControlRig/Private/Units/Highlevel/Hierarchy/RigUnit_FABRIK.cpp
@@ -3,7 +3,7 @@
 #include "RigUnit_FABRIK.h"
 #include "Units/RigUnitContext.h"
 
-UE_RigUnit_FABRIK_IMPLEMENT_MULTIPLEX(void, Execute, const FRigUnitContext& Context)
+UE_RigUnit_FABRIK_IMPLEMENT_RIGVM(void, Execute, const FRigUnitContext& Context)
 {
     DECLARE_SCOPE_HIERARCHICAL_COUNTER_RIGUNIT()
 	FRigBoneHierarchy* Hierarchy = ExecuteContext.GetBones();
diff --git a/Engine/Plugins/Experimental/ControlRig/Source/ControlRig/Private/Units/Highlevel/Hierarchy/RigUnit_FABRIK.h b/Engine/Plugins/Experimental/ControlRig/Source/ControlRig/Private/Units/Highlevel/Hierarchy/RigUnit_FABRIK.h
--- a/Engine/Plugins/Experimental/ControlRig/Source/ControlRig/Private/Units/Highlevel/Hierarchy/RigUnit_FABRIK.h
+++ b/Engine/Plugins/Experimental/ControlRig/Source/ControlRig/Private/Units/Highlevel/Hierarchy/RigUnit_FABRIK.h
@@ -36,7 +36,7 @@ struct FRigUnit_FABRIK : public FRigUnit_HighlevelBaseMutable
 {
 	GENERATED_BODY()
 
-	MULTIPLEX_METHOD()
+	RIGVM_METHOD()
 	virtual void Execute(const FRigUnitContext& Context) override;
 
 	FRigUnit_FABRIK()
diff --git a/Engine/Plugins/Experimental/ControlRig/Source/ControlRig/Private/Units/Highlevel/Hierarchy/RigUnit_FitChainToCurve.cpp b/Engine/Plugins/Experimental/ControlRig/Source/ControlRig/Private/Units/Highlevel/Hierarchy/RigUnit_FitChainToCurve.cpp
--- a/Engine/Plugins/Experimental/ControlRig/Source/ControlRig/Private/Units/Highlevel/Hierarchy/RigUnit_FitChainToCurve.cpp
+++ b/Engine/Plugins/Experimental/ControlRig/Source/ControlRig/Private/Units/Highlevel/Hierarchy/RigUnit_FitChainToCurve.cpp
@@ -3,7 +3,7 @@
 #include "Units/Highlevel/Hierarchy/RigUnit_FitChainToCurve.h"
 #include "Units/RigUnitContext.h"
 
-UE_RigUnit_FitChainToCurve_IMPLEMENT_MULTIPLEX(void, Execute, const FRigUnitContext& Context)
+UE_RigUnit_FitChainToCurve_IMPLEMENT_RIGVM(void, Execute, const FRigUnitContext& Context)
 {
     DECLARE_SCOPE_HIERARCHICAL_COUNTER_RIGUNIT()
 	FRigBoneHierarchy* Hierarchy = ExecuteContext.GetBones();
diff --git a/Engine/Plugins/Experimental/ControlRig/Source/ControlRig/Private/Units/Highlevel/Hierarchy/RigUnit_FitChainToCurve.h b/Engine/Plugins/Experimental/ControlRig/Source/ControlRig/Private/Units/Highlevel/Hierarchy/RigUnit_FitChainToCurve.h
--- a/Engine/Plugins/Experimental/ControlRig/Source/ControlRig/Private/Units/Highlevel/Hierarchy/RigUnit_FitChainToCurve.h
+++ b/Engine/Plugins/Experimental/ControlRig/Source/ControlRig/Private/Units/Highlevel/Hierarchy/RigUnit_FitChainToCurve.h
@@ -147,7 +147,7 @@ struct FRigUnit_FitChainToCurve : public FRigUnit_HighlevelBaseMutable
 		DebugSettings = FRigUnit_FitChainToCurve_DebugSettings();
 	}
 
-	MULTIPLEX_METHOD()
+	RIGVM_METHOD()
 	virtual void Execute(const FRigUnitContext& Context) override;
 
 	/** 
diff --git a/Engine/Plugins/Experimental/ControlRig/Source/ControlRig/Private/Units/Highlevel/Hierarchy/RigUnit_ModifyBoneTransforms.cpp b/Engine/Plugins/Experimental/ControlRig/Source/ControlRig/Private/Units/Highlevel/Hierarchy/RigUnit_ModifyBoneTransforms.cpp
--- a/Engine/Plugins/Experimental/ControlRig/Source/ControlRig/Private/Units/Highlevel/Hierarchy/RigUnit_ModifyBoneTransforms.cpp
+++ b/Engine/Plugins/Experimental/ControlRig/Source/ControlRig/Private/Units/Highlevel/Hierarchy/RigUnit_ModifyBoneTransforms.cpp
@@ -3,7 +3,7 @@
 #include "RigUnit_ModifyBoneTransforms.h"
 #include "Units/RigUnitContext.h"
 
-UE_RigUnit_ModifyBoneTransforms_IMPLEMENT_MULTIPLEX(void, Execute, const FRigUnitContext& Context)
+UE_RigUnit_ModifyBoneTransforms_IMPLEMENT_RIGVM(void, Execute, const FRigUnitContext& Context)
 {
     DECLARE_SCOPE_HIERARCHICAL_COUNTER_RIGUNIT()
 
diff --git a/Engine/Plugins/Experimental/ControlRig/Source/ControlRig/Private/Units/Highlevel/Hierarchy/RigUnit_ModifyBoneTransforms.h b/Engine/Plugins/Experimental/ControlRig/Source/ControlRig/Private/Units/Highlevel/Hierarchy/RigUnit_ModifyBoneTransforms.h
--- a/Engine/Plugins/Experimental/ControlRig/Source/ControlRig/Private/Units/Highlevel/Hierarchy/RigUnit_ModifyBoneTransforms.h
+++ b/Engine/Plugins/Experimental/ControlRig/Source/ControlRig/Private/Units/Highlevel/Hierarchy/RigUnit_ModifyBoneTransforms.h
@@ -78,7 +78,7 @@ struct FRigUnit_ModifyBoneTransforms : public FRigUnit_HighlevelBaseMutable
 		, Mode(EControlRigModifyBoneMode::AdditiveLocal)
 	{}
 
-	MULTIPLEX_METHOD()
+	RIGVM_METHOD()
 	virtual void Execute(const FRigUnitContext& Context) override;
 
 	/**
diff --git a/Engine/Plugins/Experimental/ControlRig/Source/ControlRig/Private/Units/Highlevel/Hierarchy/RigUnit_SlideChain.cpp b/Engine/Plugins/Experimental/ControlRig/Source/ControlRig/Private/Units/Highlevel/Hierarchy/RigUnit_SlideChain.cpp
--- a/Engine/Plugins/Experimental/ControlRig/Source/ControlRig/Private/Units/Highlevel/Hierarchy/RigUnit_SlideChain.cpp
+++ b/Engine/Plugins/Experimental/ControlRig/Source/ControlRig/Private/Units/Highlevel/Hierarchy/RigUnit_SlideChain.cpp
@@ -4,7 +4,7 @@
 #include "Math/ControlRigMathLibrary.h"
 #include "Units/RigUnitContext.h"
 
-UE_RigUnit_SlideChain_IMPLEMENT_MULTIPLEX(void, Execute, const FRigUnitContext& Context)
+UE_RigUnit_SlideChain_IMPLEMENT_RIGVM(void, Execute, const FRigUnitContext& Context)
 {
 
 	FRigBoneHierarchy* Hierarchy = ExecuteContext.GetBones();
diff --git a/Engine/Plugins/Experimental/ControlRig/Source/ControlRig/Private/Units/Highlevel/Hierarchy/RigUnit_SlideChain.h b/Engine/Plugins/Experimental/ControlRig/Source/ControlRig/Private/Units/Highlevel/Hierarchy/RigUnit_SlideChain.h
--- a/Engine/Plugins/Experimental/ControlRig/Source/ControlRig/Private/Units/Highlevel/Hierarchy/RigUnit_SlideChain.h
+++ b/Engine/Plugins/Experimental/ControlRig/Source/ControlRig/Private/Units/Highlevel/Hierarchy/RigUnit_SlideChain.h
@@ -47,7 +47,7 @@ struct FRigUnit_SlideChain: public FRigUnit_HighlevelBaseMutable
 		bPropagateToChildren = false;
 	}
 
-	MULTIPLEX_METHOD()
+	RIGVM_METHOD()
 	virtual void Execute(const FRigUnitContext& Context) override;
 
 	/** 
diff --git a/Engine/Plugins/Experimental/ControlRig/Source/ControlRig/Private/Units/Highlevel/Hierarchy/RigUnit_SpringIK.cpp b/Engine/Plugins/Experimental/ControlRig/Source/ControlRig/Private/Units/Highlevel/Hierarchy/RigUnit_SpringIK.cpp
--- a/Engine/Plugins/Experimental/ControlRig/Source/ControlRig/Private/Units/Highlevel/Hierarchy/RigUnit_SpringIK.cpp
+++ b/Engine/Plugins/Experimental/ControlRig/Source/ControlRig/Private/Units/Highlevel/Hierarchy/RigUnit_SpringIK.cpp
@@ -3,7 +3,7 @@
 #include "Units/Highlevel/Hierarchy/RigUnit_SpringIK.h"
 #include "Units/RigUnitContext.h"
 
-UE_RigUnit_SpringIK_IMPLEMENT_MULTIPLEX(void, Execute, const FRigUnitContext& Context)
+UE_RigUnit_SpringIK_IMPLEMENT_RIGVM(void, Execute, const FRigUnitContext& Context)
 {
     DECLARE_SCOPE_HIERARCHICAL_COUNTER_RIGUNIT()
 	FRigBoneHierarchy* Hierarchy = ExecuteContext.GetBones();
diff --git a/Engine/Plugins/Experimental/ControlRig/Source/ControlRig/Private/Units/Highlevel/Hierarchy/RigUnit_SpringIK.h b/Engine/Plugins/Experimental/ControlRig/Source/ControlRig/Private/Units/Highlevel/Hierarchy/RigUnit_SpringIK.h
--- a/Engine/Plugins/Experimental/ControlRig/Source/ControlRig/Private/Units/Highlevel/Hierarchy/RigUnit_SpringIK.h
+++ b/Engine/Plugins/Experimental/ControlRig/Source/ControlRig/Private/Units/Highlevel/Hierarchy/RigUnit_SpringIK.h
@@ -96,7 +96,7 @@ struct FRigUnit_SpringIK : public FRigUnit_HighlevelBaseMutable
 		DebugSettings = FRigUnit_SpringIK_DebugSettings();
 	}
 
-	MULTIPLEX_METHOD()
+	RIGVM_METHOD()
 	virtual void Execute(const FRigUnitContext& Context) override;
 
 	/** 
diff --git a/Engine/Plugins/Experimental/ControlRig/Source/ControlRig/Private/Units/Highlevel/Hierarchy/RigUnit_TransformConstraint.cpp b/Engine/Plugins/Experimental/ControlRig/Source/ControlRig/Private/Units/Highlevel/Hierarchy/RigUnit_TransformConstraint.cpp
--- a/Engine/Plugins/Experimental/ControlRig/Source/ControlRig/Private/Units/Highlevel/Hierarchy/RigUnit_TransformConstraint.cpp
+++ b/Engine/Plugins/Experimental/ControlRig/Source/ControlRig/Private/Units/Highlevel/Hierarchy/RigUnit_TransformConstraint.cpp
@@ -5,7 +5,7 @@
 #include "HelperUtil.h"
 #include "AnimationCoreLibrary.h"
 
-UE_RigUnit_TransformConstraint_IMPLEMENT_MULTIPLEX(void, Execute, const FRigUnitContext& Context)
+UE_RigUnit_TransformConstraint_IMPLEMENT_RIGVM(void, Execute, const FRigUnitContext& Context)
 {
     DECLARE_SCOPE_HIERARCHICAL_COUNTER_RIGUNIT()
 
diff --git a/Engine/Plugins/Experimental/ControlRig/Source/ControlRig/Private/Units/Highlevel/Hierarchy/RigUnit_TransformConstraint.h b/Engine/Plugins/Experimental/ControlRig/Source/ControlRig/Private/Units/Highlevel/Hierarchy/RigUnit_TransformConstraint.h
--- a/Engine/Plugins/Experimental/ControlRig/Source/ControlRig/Private/Units/Highlevel/Hierarchy/RigUnit_TransformConstraint.h
+++ b/Engine/Plugins/Experimental/ControlRig/Source/ControlRig/Private/Units/Highlevel/Hierarchy/RigUnit_TransformConstraint.h
@@ -58,7 +58,7 @@ struct FRigUnit_TransformConstraint : public FRigUnit_HighlevelBaseMutable
 		: BaseTransformSpace(ETransformSpaceMode::GlobalSpace)
 	{}
 
-	MULTIPLEX_METHOD()
+	RIGVM_METHOD()
 	virtual void Execute(const FRigUnitContext& Context) override;
 
 	UPROPERTY(EditAnywhere, Category = "Constraint", meta = (Input, Constant, BoneName))
diff --git a/Engine/Plugins/Experimental/ControlRig/Source/ControlRig/Private/Units/Highlevel/Hierarchy/RigUnit_TwoBoneIKSimple.cpp b/Engine/Plugins/Experimental/ControlRig/Source/ControlRig/Private/Units/Highlevel/Hierarchy/RigUnit_TwoBoneIKSimple.cpp
--- a/Engine/Plugins/Experimental/ControlRig/Source/ControlRig/Private/Units/Highlevel/Hierarchy/RigUnit_TwoBoneIKSimple.cpp
+++ b/Engine/Plugins/Experimental/ControlRig/Source/ControlRig/Private/Units/Highlevel/Hierarchy/RigUnit_TwoBoneIKSimple.cpp
@@ -4,7 +4,7 @@
 #include "Units/RigUnitContext.h"
 #include "Math/ControlRigMathLibrary.h"
 
-UE_RigUnit_TwoBoneIKSimple_IMPLEMENT_MULTIPLEX(void, Execute, const FRigUnitContext& Context)
+UE_RigUnit_TwoBoneIKSimple_IMPLEMENT_RIGVM(void, Execute, const FRigUnitContext& Context)
 {
     DECLARE_SCOPE_HIERARCHICAL_COUNTER_RIGUNIT()
 	FRigBoneHierarchy* Hierarchy = ExecuteContext.GetBones();
diff --git a/Engine/Plugins/Experimental/ControlRig/Source/ControlRig/Private/Units/Highlevel/Hierarchy/RigUnit_TwoBoneIKSimple.h b/Engine/Plugins/Experimental/ControlRig/Source/ControlRig/Private/Units/Highlevel/Hierarchy/RigUnit_TwoBoneIKSimple.h
--- a/Engine/Plugins/Experimental/ControlRig/Source/ControlRig/Private/Units/Highlevel/Hierarchy/RigUnit_TwoBoneIKSimple.h
+++ b/Engine/Plugins/Experimental/ControlRig/Source/ControlRig/Private/Units/Highlevel/Hierarchy/RigUnit_TwoBoneIKSimple.h
@@ -67,7 +67,7 @@ struct FRigUnit_TwoBoneIKSimple : public FRigUnit_HighlevelBaseMutable
 		PoleVectorSpaceIndex = INDEX_NONE;
 	}
 
-	MULTIPLEX_METHOD()
+	RIGVM_METHOD()
 	virtual void Execute(const FRigUnitContext& Context) override;
 
 	/** 
diff --git a/Engine/Plugins/Experimental/ControlRig/Source/ControlRig/Private/Units/Math/RigUnit_MathBool.cpp b/Engine/Plugins/Experimental/ControlRig/Source/ControlRig/Private/Units/Math/RigUnit_MathBool.cpp
--- a/Engine/Plugins/Experimental/ControlRig/Source/ControlRig/Private/Units/Math/RigUnit_MathBool.cpp
+++ b/Engine/Plugins/Experimental/ControlRig/Source/ControlRig/Private/Units/Math/RigUnit_MathBool.cpp
@@ -3,37 +3,37 @@
 #include "Units/Math/RigUnit_MathBool.h"
 #include "Units/RigUnitContext.h"
 
-UE_RigUnit_MathBoolNot_IMPLEMENT_MULTIPLEX(void, Execute, const FRigUnitContext& Context)
+UE_RigUnit_MathBoolNot_IMPLEMENT_RIGVM(void, Execute, const FRigUnitContext& Context)
 {
     DECLARE_SCOPE_HIERARCHICAL_COUNTER_RIGUNIT()
 	Result = !Value;
 }
 
-UE_RigUnit_MathBoolAnd_IMPLEMENT_MULTIPLEX(void, Execute, const FRigUnitContext& Context)
+UE_RigUnit_MathBoolAnd_IMPLEMENT_RIGVM(void, Execute, const FRigUnitContext& Context)
 {
     DECLARE_SCOPE_HIERARCHICAL_COUNTER_RIGUNIT()
 	Result = A && B;
 }
 
-UE_RigUnit_MathBoolNand_IMPLEMENT_MULTIPLEX(void, Execute, const FRigUnitContext& Context)
+UE_RigUnit_MathBoolNand_IMPLEMENT_RIGVM(void, Execute, const FRigUnitContext& Context)
 {
     DECLARE_SCOPE_HIERARCHICAL_COUNTER_RIGUNIT()
 	Result = (!A) && (!B);
 }
 
-UE_RigUnit_MathBoolOr_IMPLEMENT_MULTIPLEX(void, Execute, const FRigUnitContext& Context)
+UE_RigUnit_MathBoolOr_IMPLEMENT_RIGVM(void, Execute, const FRigUnitContext& Context)
 {
     DECLARE_SCOPE_HIERARCHICAL_COUNTER_RIGUNIT()
 	Result = A || B;
 }
 
-UE_RigUnit_MathBoolEquals_IMPLEMENT_MULTIPLEX(void, Execute, const FRigUnitContext& Context)
+UE_RigUnit_MathBoolEquals_IMPLEMENT_RIGVM(void, Execute, const FRigUnitContext& Context)
 {
     DECLARE_SCOPE_HIERARCHICAL_COUNTER_RIGUNIT()
 	Result = A == B;
 }
 
-UE_RigUnit_MathBoolNotEquals_IMPLEMENT_MULTIPLEX(void, Execute, const FRigUnitContext& Context)
+UE_RigUnit_MathBoolNotEquals_IMPLEMENT_RIGVM(void, Execute, const FRigUnitContext& Context)
 {
     DECLARE_SCOPE_HIERARCHICAL_COUNTER_RIGUNIT()
 	Result = A != B;
diff --git a/Engine/Plugins/Experimental/ControlRig/Source/ControlRig/Private/Units/Math/RigUnit_MathBool.h b/Engine/Plugins/Experimental/ControlRig/Source/ControlRig/Private/Units/Math/RigUnit_MathBool.h
--- a/Engine/Plugins/Experimental/ControlRig/Source/ControlRig/Private/Units/Math/RigUnit_MathBool.h
+++ b/Engine/Plugins/Experimental/ControlRig/Source/ControlRig/Private/Units/Math/RigUnit_MathBool.h
@@ -98,7 +98,7 @@ struct FRigUnit_MathBoolNot : public FRigUnit_MathBoolUnaryOp
 {
 	GENERATED_BODY()
 
-	MULTIPLEX_METHOD()
+	RIGVM_METHOD()
 	virtual void Execute(const FRigUnitContext& Context) override;
 };
 
@@ -110,7 +110,7 @@ struct FRigUnit_MathBoolAnd : public FRigUnit_MathBoolBinaryOp
 {
 	GENERATED_BODY()
 
-	MULTIPLEX_METHOD()
+	RIGVM_METHOD()
 	virtual void Execute(const FRigUnitContext& Context) override;
 };
 
@@ -122,7 +122,7 @@ struct FRigUnit_MathBoolNand : public FRigUnit_MathBoolBinaryOp
 {
 	GENERATED_BODY()
 
-	MULTIPLEX_METHOD()
+	RIGVM_METHOD()
 	virtual void Execute(const FRigUnitContext& Context) override;
 };
 
@@ -134,7 +134,7 @@ struct FRigUnit_MathBoolOr : public FRigUnit_MathBoolBinaryOp
 {
 	GENERATED_BODY()
 
-	MULTIPLEX_METHOD()
+	RIGVM_METHOD()
 	virtual void Execute(const FRigUnitContext& Context) override;
 };
 
@@ -145,7 +145,7 @@ USTRUCT(meta=(DisplayName="Equals", PrototypeName="Equals", Keywords="Same,=="))
 struct FRigUnit_MathBoolEquals : public FRigUnit_MathBoolBase
 {
 	GENERATED_BODY()
-	MULTIPLEX_METHOD()
+	RIGVM_METHOD()
 	virtual void Execute(const FRigUnitContext& Context) override;
 
 	FRigUnit_MathBoolEquals()
@@ -171,7 +171,7 @@ USTRUCT(meta=(DisplayName="Not Equals", PrototypeName="NotEquals", Keywords="Dif
 struct FRigUnit_MathBoolNotEquals : public FRigUnit_MathBoolBase
 {
 	GENERATED_BODY()
-	MULTIPLEX_METHOD()
+	RIGVM_METHOD()
 	virtual void Execute(const FRigUnitContext& Context) override;
 
 	FRigUnit_MathBoolNotEquals()
diff --git a/Engine/Plugins/Experimental/ControlRig/Source/ControlRig/Private/Units/Math/RigUnit_MathFloat.cpp b/Engine/Plugins/Experimental/ControlRig/Source/ControlRig/Private/Units/Math/RigUnit_MathFloat.cpp
--- a/Engine/Plugins/Experimental/ControlRig/Source/ControlRig/Private/Units/Math/RigUnit_MathFloat.cpp
+++ b/Engine/Plugins/Experimental/ControlRig/Source/ControlRig/Private/Units/Math/RigUnit_MathFloat.cpp
@@ -3,25 +3,25 @@
 #include "Units/Math/RigUnit_MathFloat.h"
 #include "Units/RigUnitContext.h"
 
-UE_RigUnit_MathFloatAdd_IMPLEMENT_MULTIPLEX(void, Execute, const FRigUnitContext& Context)
+UE_RigUnit_MathFloatAdd_IMPLEMENT_RIGVM(void, Execute, const FRigUnitContext& Context)
 {
     DECLARE_SCOPE_HIERARCHICAL_COUNTER_RIGUNIT()
 	Result = A + B;
 }
 
-UE_RigUnit_MathFloatSub_IMPLEMENT_MULTIPLEX(void, Execute, const FRigUnitContext& Context)
+UE_RigUnit_MathFloatSub_IMPLEMENT_RIGVM(void, Execute, const FRigUnitContext& Context)
 {
     DECLARE_SCOPE_HIERARCHICAL_COUNTER_RIGUNIT()
 	Result = A - B;
 }
 
-UE_RigUnit_MathFloatMul_IMPLEMENT_MULTIPLEX(void, Execute, const FRigUnitContext& Context)
+UE_RigUnit_MathFloatMul_IMPLEMENT_RIGVM(void, Execute, const FRigUnitContext& Context)
 {
     DECLARE_SCOPE_HIERARCHICAL_COUNTER_RIGUNIT()
 	Result = A * B;
 }
 
-UE_RigUnit_MathFloatDiv_IMPLEMENT_MULTIPLEX(void, Execute, const FRigUnitContext& Context)
+UE_RigUnit_MathFloatDiv_IMPLEMENT_RIGVM(void, Execute, const FRigUnitContext& Context)
 {
     DECLARE_SCOPE_HIERARCHICAL_COUNTER_RIGUNIT()
 	if(FMath::IsNearlyZero(B))
@@ -33,7 +33,7 @@ UE_RigUnit_MathFloatDiv_IMPLEMENT_MULTIPLEX(void, Execute, const FRigUnitContext
 	Result = A / B;
 }
 
-UE_RigUnit_MathFloatMod_IMPLEMENT_MULTIPLEX(void, Execute, const FRigUnitContext& Context)
+UE_RigUnit_MathFloatMod_IMPLEMENT_RIGVM(void, Execute, const FRigUnitContext& Context)
 {
     DECLARE_SCOPE_HIERARCHICAL_COUNTER_RIGUNIT()
 	if(FMath::IsNearlyZero(B) || B < 0.f)
@@ -45,25 +45,25 @@ UE_RigUnit_MathFloatMod_IMPLEMENT_MULTIPLEX(void, Execute, const FRigUnitContext
 	Result = FMath::Fmod(A, B);
 }
 
-UE_RigUnit_MathFloatMin_IMPLEMENT_MULTIPLEX(void, Execute, const FRigUnitContext& Context)
+UE_RigUnit_MathFloatMin_IMPLEMENT_RIGVM(void, Execute, const FRigUnitContext& Context)
 {
     DECLARE_SCOPE_HIERARCHICAL_COUNTER_RIGUNIT()
 	Result = FMath::Min<float>(A, B);
 }
 
-UE_RigUnit_MathFloatMax_IMPLEMENT_MULTIPLEX(void, Execute, const FRigUnitContext& Context)
+UE_RigUnit_MathFloatMax_IMPLEMENT_RIGVM(void, Execute, const FRigUnitContext& Context)
 {
     DECLARE_SCOPE_HIERARCHICAL_COUNTER_RIGUNIT()
 	Result = FMath::Max<float>(A, B);
 }
 
-UE_RigUnit_MathFloatPow_IMPLEMENT_MULTIPLEX(void, Execute, const FRigUnitContext& Context)
+UE_RigUnit_MathFloatPow_IMPLEMENT_RIGVM(void, Execute, const FRigUnitContext& Context)
 {
     DECLARE_SCOPE_HIERARCHICAL_COUNTER_RIGUNIT()
 	Result = FMath::Pow(A, B);
 }
 
-UE_RigUnit_MathFloatSqrt_IMPLEMENT_MULTIPLEX(void, Execute, const FRigUnitContext& Context)
+UE_RigUnit_MathFloatSqrt_IMPLEMENT_RIGVM(void, Execute, const FRigUnitContext& Context)
 {
     DECLARE_SCOPE_HIERARCHICAL_COUNTER_RIGUNIT()
 	if(Value < 0.f)
@@ -75,55 +75,55 @@ UE_RigUnit_MathFloatSqrt_IMPLEMENT_MULTIPLEX(void, Execute, const FRigUnitContex
 	Result = FMath::Sqrt(Value);
 }
 
-UE_RigUnit_MathFloatNegate_IMPLEMENT_MULTIPLEX(void, Execute, const FRigUnitContext& Context)
+UE_RigUnit_MathFloatNegate_IMPLEMENT_RIGVM(void, Execute, const FRigUnitContext& Context)
 {
     DECLARE_SCOPE_HIERARCHICAL_COUNTER_RIGUNIT()
 	Result = -Value;
 }
 
-UE_RigUnit_MathFloatAbs_IMPLEMENT_MULTIPLEX(void, Execute, const FRigUnitContext& Context)
+UE_RigUnit_MathFloatAbs_IMPLEMENT_RIGVM(void, Execute, const FRigUnitContext& Context)
 {
     DECLARE_SCOPE_HIERARCHICAL_COUNTER_RIGUNIT()
 	Result = FMath::Abs(Value);
 }
 
-UE_RigUnit_MathFloatFloor_IMPLEMENT_MULTIPLEX(void, Execute, const FRigUnitContext& Context)
+UE_RigUnit_MathFloatFloor_IMPLEMENT_RIGVM(void, Execute, const FRigUnitContext& Context)
 {
     DECLARE_SCOPE_HIERARCHICAL_COUNTER_RIGUNIT()
 	Result = FMath::FloorToFloat(Value);
 }
 
-UE_RigUnit_MathFloatCeil_IMPLEMENT_MULTIPLEX(void, Execute, const FRigUnitContext& Context)
+UE_RigUnit_MathFloatCeil_IMPLEMENT_RIGVM(void, Execute, const FRigUnitContext& Context)
 {
     DECLARE_SCOPE_HIERARCHICAL_COUNTER_RIGUNIT()
 	Result = FMath::CeilToFloat(Value);
 }
 
-UE_RigUnit_MathFloatRound_IMPLEMENT_MULTIPLEX(void, Execute, const FRigUnitContext& Context)
+UE_RigUnit_MathFloatRound_IMPLEMENT_RIGVM(void, Execute, const FRigUnitContext& Context)
 {
     DECLARE_SCOPE_HIERARCHICAL_COUNTER_RIGUNIT()
 	Result = FMath::RoundToFloat(Value);
 }
 
-UE_RigUnit_MathFloatSign_IMPLEMENT_MULTIPLEX(void, Execute, const FRigUnitContext& Context)
+UE_RigUnit_MathFloatSign_IMPLEMENT_RIGVM(void, Execute, const FRigUnitContext& Context)
 {
     DECLARE_SCOPE_HIERARCHICAL_COUNTER_RIGUNIT()
 	Result = Value < 0.f ? -1.f : 1.f;
 }
 
-UE_RigUnit_MathFloatClamp_IMPLEMENT_MULTIPLEX(void, Execute, const FRigUnitContext& Context)
+UE_RigUnit_MathFloatClamp_IMPLEMENT_RIGVM(void, Execute, const FRigUnitContext& Context)
 {
     DECLARE_SCOPE_HIERARCHICAL_COUNTER_RIGUNIT()
 	Result = FMath::Clamp<float>(Value, Minimum, Maximum);
 }
 
-UE_RigUnit_MathFloatLerp_IMPLEMENT_MULTIPLEX(void, Execute, const FRigUnitContext& Context)
+UE_RigUnit_MathFloatLerp_IMPLEMENT_RIGVM(void, Execute, const FRigUnitContext& Context)
 {
     DECLARE_SCOPE_HIERARCHICAL_COUNTER_RIGUNIT()
 	Result = FMath::Lerp<float>(A, B, T);
 }
 
-UE_RigUnit_MathFloatRemap_IMPLEMENT_MULTIPLEX(void, Execute, const FRigUnitContext& Context)
+UE_RigUnit_MathFloatRemap_IMPLEMENT_RIGVM(void, Execute, const FRigUnitContext& Context)
 {
     DECLARE_SCOPE_HIERARCHICAL_COUNTER_RIGUNIT()
 	float Ratio = 0.f;
@@ -142,43 +142,43 @@ UE_RigUnit_MathFloatRemap_IMPLEMENT_MULTIPLEX(void, Execute, const FRigUnitConte
 	Result = FMath::Lerp<float>(TargetMinimum, TargetMaximum, Ratio);
 }
 
-UE_RigUnit_MathFloatEquals_IMPLEMENT_MULTIPLEX(void, Execute, const FRigUnitContext& Context)
+UE_RigUnit_MathFloatEquals_IMPLEMENT_RIGVM(void, Execute, const FRigUnitContext& Context)
 {
     DECLARE_SCOPE_HIERARCHICAL_COUNTER_RIGUNIT()
 	Result = A == B;
 }
 
-UE_RigUnit_MathFloatNotEquals_IMPLEMENT_MULTIPLEX(void, Execute, const FRigUnitContext& Context)
+UE_RigUnit_MathFloatNotEquals_IMPLEMENT_RIGVM(void, Execute, const FRigUnitContext& Context)
 {
     DECLARE_SCOPE_HIERARCHICAL_COUNTER_RIGUNIT()
 	Result = A != B;
 }
 
-UE_RigUnit_MathFloatGreater_IMPLEMENT_MULTIPLEX(void, Execute, const FRigUnitContext& Context)
+UE_RigUnit_MathFloatGreater_IMPLEMENT_RIGVM(void, Execute, const FRigUnitContext& Context)
 {
     DECLARE_SCOPE_HIERARCHICAL_COUNTER_RIGUNIT()
 	Result = A > B;
 }
 
-UE_RigUnit_MathFloatLess_IMPLEMENT_MULTIPLEX(void, Execute, const FRigUnitContext& Context)
+UE_RigUnit_MathFloatLess_IMPLEMENT_RIGVM(void, Execute, const FRigUnitContext& Context)
 {
     DECLARE_SCOPE_HIERARCHICAL_COUNTER_RIGUNIT()
 	Result = A < B;
 }
 
-UE_RigUnit_MathFloatGreaterEqual_IMPLEMENT_MULTIPLEX(void, Execute, const FRigUnitContext& Context)
+UE_RigUnit_MathFloatGreaterEqual_IMPLEMENT_RIGVM(void, Execute, const FRigUnitContext& Context)
 {
     DECLARE_SCOPE_HIERARCHICAL_COUNTER_RIGUNIT()
 	Result = A >= B;
 }
 
-UE_RigUnit_MathFloatLessEqual_IMPLEMENT_MULTIPLEX(void, Execute, const FRigUnitContext& Context)
+UE_RigUnit_MathFloatLessEqual_IMPLEMENT_RIGVM(void, Execute, const FRigUnitContext& Context)
 {
     DECLARE_SCOPE_HIERARCHICAL_COUNTER_RIGUNIT()
 	Result = A <= B;
 }
 
-UE_RigUnit_MathFloatIsNearlyZero_IMPLEMENT_MULTIPLEX(void, Execute, const FRigUnitContext& Context)
+UE_RigUnit_MathFloatIsNearlyZero_IMPLEMENT_RIGVM(void, Execute, const FRigUnitContext& Context)
 {
     DECLARE_SCOPE_HIERARCHICAL_COUNTER_RIGUNIT()
 	if(Tolerance < 0.f)
@@ -188,7 +188,7 @@ UE_RigUnit_MathFloatIsNearlyZero_IMPLEMENT_MULTIPLEX(void, Execute, const FRigUn
 	Result = FMath::IsNearlyZero(Value, FMath::Max<float>(Tolerance, SMALL_NUMBER));
 }
 
-UE_RigUnit_MathFloatIsNearlyEqual_IMPLEMENT_MULTIPLEX(void, Execute, const FRigUnitContext& Context)
+UE_RigUnit_MathFloatIsNearlyEqual_IMPLEMENT_RIGVM(void, Execute, const FRigUnitContext& Context)
 {
     DECLARE_SCOPE_HIERARCHICAL_COUNTER_RIGUNIT()
 	if(Tolerance < 0.f)
@@ -198,43 +198,43 @@ UE_RigUnit_MathFloatIsNearlyEqual_IMPLEMENT_MULTIPLEX(void, Execute, const FRigU
 	Result = FMath::IsNearlyEqual(A, B, FMath::Max<float>(Tolerance, SMALL_NUMBER));
 }
 
-UE_RigUnit_MathFloatSelectBool_IMPLEMENT_MULTIPLEX(void, Execute, const FRigUnitContext& Context)
+UE_RigUnit_MathFloatSelectBool_IMPLEMENT_RIGVM(void, Execute, const FRigUnitContext& Context)
 {
     DECLARE_SCOPE_HIERARCHICAL_COUNTER_RIGUNIT()
 	Result = Condition ? IfTrue : IfFalse;
 }
 
-UE_RigUnit_MathFloatDeg_IMPLEMENT_MULTIPLEX(void, Execute, const FRigUnitContext& Context)
+UE_RigUnit_MathFloatDeg_IMPLEMENT_RIGVM(void, Execute, const FRigUnitContext& Context)
 {
     DECLARE_SCOPE_HIERARCHICAL_COUNTER_RIGUNIT()
 	Result = FMath::RadiansToDegrees(Value);
 }
 
-UE_RigUnit_MathFloatRad_IMPLEMENT_MULTIPLEX(void, Execute, const FRigUnitContext& Context)
+UE_RigUnit_MathFloatRad_IMPLEMENT_RIGVM(void, Execute, const FRigUnitContext& Context)
 {
     DECLARE_SCOPE_HIERARCHICAL_COUNTER_RIGUNIT()
 	Result = FMath::DegreesToRadians(Value);
 }
 
-UE_RigUnit_MathFloatSin_IMPLEMENT_MULTIPLEX(void, Execute, const FRigUnitContext& Context)
+UE_RigUnit_MathFloatSin_IMPLEMENT_RIGVM(void, Execute, const FRigUnitContext& Context)
 {
     DECLARE_SCOPE_HIERARCHICAL_COUNTER_RIGUNIT()
 	Result = FMath::Sin(Value);
 }
 
-UE_RigUnit_MathFloatCos_IMPLEMENT_MULTIPLEX(void, Execute, const FRigUnitContext& Context)
+UE_RigUnit_MathFloatCos_IMPLEMENT_RIGVM(void, Execute, const FRigUnitContext& Context)
 {
     DECLARE_SCOPE_HIERARCHICAL_COUNTER_RIGUNIT()
 	Result = FMath::Cos(Value);
 }
 
-UE_RigUnit_MathFloatTan_IMPLEMENT_MULTIPLEX(void, Execute, const FRigUnitContext& Context)
+UE_RigUnit_MathFloatTan_IMPLEMENT_RIGVM(void, Execute, const FRigUnitContext& Context)
 {
     DECLARE_SCOPE_HIERARCHICAL_COUNTER_RIGUNIT()
 	Result = FMath::Tan(Value);
 }
 
-UE_RigUnit_MathFloatAsin_IMPLEMENT_MULTIPLEX(void, Execute, const FRigUnitContext& Context)
+UE_RigUnit_MathFloatAsin_IMPLEMENT_RIGVM(void, Execute, const FRigUnitContext& Context)
 {
     DECLARE_SCOPE_HIERARCHICAL_COUNTER_RIGUNIT()
 	if (!FMath::IsWithinInclusive<float>(Value, -1.f, 1.f))
@@ -246,7 +246,7 @@ UE_RigUnit_MathFloatAsin_IMPLEMENT_MULTIPLEX(void, Execute, const FRigUnitContex
 	Result = FMath::Asin(Value);
 }
 
-UE_RigUnit_MathFloatAcos_IMPLEMENT_MULTIPLEX(void, Execute, const FRigUnitContext& Context)
+UE_RigUnit_MathFloatAcos_IMPLEMENT_RIGVM(void, Execute, const FRigUnitContext& Context)
 {
     DECLARE_SCOPE_HIERARCHICAL_COUNTER_RIGUNIT()
 	if (!FMath::IsWithinInclusive<float>(Value, -1.f, 1.f))
@@ -258,13 +258,13 @@ UE_RigUnit_MathFloatAcos_IMPLEMENT_MULTIPLEX(void, Execute, const FRigUnitContex
 	Result = FMath::Acos(Value);
 }
 
-UE_RigUnit_MathFloatAtan_IMPLEMENT_MULTIPLEX(void, Execute, const FRigUnitContext& Context)
+UE_RigUnit_MathFloatAtan_IMPLEMENT_RIGVM(void, Execute, const FRigUnitContext& Context)
 {
     DECLARE_SCOPE_HIERARCHICAL_COUNTER_RIGUNIT()
 	Result = FMath::Atan(Value);
 }
 
-UE_RigUnit_MathFloatLawOfCosine_IMPLEMENT_MULTIPLEX(void, Execute, const FRigUnitContext& Context)
+UE_RigUnit_MathFloatLawOfCosine_IMPLEMENT_RIGVM(void, Execute, const FRigUnitContext& Context)
 {
     DECLARE_SCOPE_HIERARCHICAL_COUNTER_RIGUNIT()
 	if ((A <= 0.f) || (B <= 0.f) || (C <= 0.f) || (A + B < C) || (A + C < B) || (B + C < A))
diff --git a/Engine/Plugins/Experimental/ControlRig/Source/ControlRig/Private/Units/Math/RigUnit_MathFloat.h b/Engine/Plugins/Experimental/ControlRig/Source/ControlRig/Private/Units/Math/RigUnit_MathFloat.h
--- a/Engine/Plugins/Experimental/ControlRig/Source/ControlRig/Private/Units/Math/RigUnit_MathFloat.h
+++ b/Engine/Plugins/Experimental/ControlRig/Source/ControlRig/Private/Units/Math/RigUnit_MathFloat.h
@@ -112,7 +112,7 @@ struct FRigUnit_MathFloatAdd : public FRigUnit_MathFloatBinaryOp
 {
 	GENERATED_BODY()
 
-	MULTIPLEX_METHOD()
+	RIGVM_METHOD()
 	virtual void Execute(const FRigUnitContext& Context) override;
 };
 
@@ -124,7 +124,7 @@ struct FRigUnit_MathFloatSub : public FRigUnit_MathFloatBinaryOp
 {
 	GENERATED_BODY()
 
-	MULTIPLEX_METHOD()
+	RIGVM_METHOD()
 	virtual void Execute(const FRigUnitContext& Context) override;
 };
 
@@ -142,7 +142,7 @@ struct FRigUnit_MathFloatMul : public FRigUnit_MathFloatBinaryOp
 		B = 1.f;
 	}
 
-	MULTIPLEX_METHOD()
+	RIGVM_METHOD()
 	virtual void Execute(const FRigUnitContext& Context) override;
 };
 
@@ -159,7 +159,7 @@ struct FRigUnit_MathFloatDiv : public FRigUnit_MathFloatBinaryOp
 		B = 1.f;
 	}
 
-	MULTIPLEX_METHOD()
+	RIGVM_METHOD()
 	virtual void Execute(const FRigUnitContext& Context) override;
 };
 
@@ -177,7 +177,7 @@ struct FRigUnit_MathFloatMod : public FRigUnit_MathFloatBinaryOp
 		B = 1.f;
 	}
 
-	MULTIPLEX_METHOD()
+	RIGVM_METHOD()
 	virtual void Execute(const FRigUnitContext& Context) override;
 };
 
@@ -189,7 +189,7 @@ struct FRigUnit_MathFloatMin : public FRigUnit_MathFloatBinaryOp
 {
 	GENERATED_BODY()
 
-	MULTIPLEX_METHOD()
+	RIGVM_METHOD()
 	virtual void Execute(const FRigUnitContext& Context) override;
 };
 
@@ -201,7 +201,7 @@ struct FRigUnit_MathFloatMax : public FRigUnit_MathFloatBinaryOp
 {
 	GENERATED_BODY()
 
-	MULTIPLEX_METHOD()
+	RIGVM_METHOD()
 	virtual void Execute(const FRigUnitContext& Context) override;
 };
 
@@ -219,7 +219,7 @@ struct FRigUnit_MathFloatPow : public FRigUnit_MathFloatBinaryOp
 		B = 1.f;
 	}
 
-	MULTIPLEX_METHOD()
+	RIGVM_METHOD()
 	virtual void Execute(const FRigUnitContext& Context) override;
 };
 
@@ -231,7 +231,7 @@ struct FRigUnit_MathFloatSqrt : public FRigUnit_MathFloatUnaryOp
 {
 	GENERATED_BODY()
 
-	MULTIPLEX_METHOD()
+	RIGVM_METHOD()
 	virtual void Execute(const FRigUnitContext& Context) override;
 };
 
@@ -243,7 +243,7 @@ struct FRigUnit_MathFloatNegate : public FRigUnit_MathFloatUnaryOp
 {
 	GENERATED_BODY()
 
-	MULTIPLEX_METHOD()
+	RIGVM_METHOD()
 	virtual void Execute(const FRigUnitContext& Context) override;
 };
 
@@ -255,7 +255,7 @@ struct FRigUnit_MathFloatAbs : public FRigUnit_MathFloatUnaryOp
 {
 	GENERATED_BODY()
 
-	MULTIPLEX_METHOD()
+	RIGVM_METHOD()
 	virtual void Execute(const FRigUnitContext& Context) override;
 };
 
@@ -267,7 +267,7 @@ struct FRigUnit_MathFloatFloor : public FRigUnit_MathFloatUnaryOp
 {
 	GENERATED_BODY()
 
-	MULTIPLEX_METHOD()
+	RIGVM_METHOD()
 	virtual void Execute(const FRigUnitContext& Context) override;
 };
 
@@ -279,7 +279,7 @@ struct FRigUnit_MathFloatCeil : public FRigUnit_MathFloatUnaryOp
 {
 	GENERATED_BODY()
 
-	MULTIPLEX_METHOD()
+	RIGVM_METHOD()
 	virtual void Execute(const FRigUnitContext& Context) override;
 };
 
@@ -291,7 +291,7 @@ struct FRigUnit_MathFloatRound : public FRigUnit_MathFloatUnaryOp
 {
 	GENERATED_BODY()
 
-	MULTIPLEX_METHOD()
+	RIGVM_METHOD()
 	virtual void Execute(const FRigUnitContext& Context) override;
 };
 
@@ -303,7 +303,7 @@ struct FRigUnit_MathFloatSign : public FRigUnit_MathFloatUnaryOp
 {
 	GENERATED_BODY()
 
-	MULTIPLEX_METHOD()
+	RIGVM_METHOD()
 	virtual void Execute(const FRigUnitContext& Context) override;
 };
 
@@ -320,7 +320,7 @@ struct FRigUnit_MathFloatClamp : public FRigUnit_MathFloatBase
 		Value = Minimum = Maximum = Result = 0.f;
 	}
 
-	MULTIPLEX_METHOD()
+	RIGVM_METHOD()
 	virtual void Execute(const FRigUnitContext& Context) override;
 
 	UPROPERTY(meta=(Input))
@@ -350,7 +350,7 @@ struct FRigUnit_MathFloatLerp : public FRigUnit_MathFloatBase
 		B = 1.f;
 	}
 
-	MULTIPLEX_METHOD()
+	RIGVM_METHOD()
 	virtual void Execute(const FRigUnitContext& Context) override;
 
 	UPROPERTY(meta=(Input))
@@ -381,7 +381,7 @@ struct FRigUnit_MathFloatRemap : public FRigUnit_MathFloatBase
 		bClamp = false;
 	}
 
-	MULTIPLEX_METHOD()
+	RIGVM_METHOD()
 	virtual void Execute(const FRigUnitContext& Context) override;
 
 	UPROPERTY(meta=(Input))
@@ -414,7 +414,7 @@ USTRUCT(meta=(DisplayName="Equals", PrototypeName="Equals", Keywords="Same,=="))
 struct FRigUnit_MathFloatEquals : public FRigUnit_MathFloatBase
 {
 	GENERATED_BODY()
-	MULTIPLEX_METHOD()
+	RIGVM_METHOD()
 	virtual void Execute(const FRigUnitContext& Context) override;
 
 	FRigUnit_MathFloatEquals()
@@ -440,7 +440,7 @@ USTRUCT(meta=(DisplayName="Not Equals", PrototypeName="NotEquals", Keywords="Dif
 struct FRigUnit_MathFloatNotEquals : public FRigUnit_MathFloatBase
 {
 	GENERATED_BODY()
-	MULTIPLEX_METHOD()
+	RIGVM_METHOD()
 	virtual void Execute(const FRigUnitContext& Context) override;
 
 	FRigUnit_MathFloatNotEquals()
@@ -473,7 +473,7 @@ struct FRigUnit_MathFloatGreater : public FRigUnit_MathFloatBase
 		Result = false;
 	}
 
-	MULTIPLEX_METHOD()
+	RIGVM_METHOD()
 	virtual void Execute(const FRigUnitContext& Context) override;
 
 	UPROPERTY(meta=(Input))
@@ -500,7 +500,7 @@ struct FRigUnit_MathFloatLess : public FRigUnit_MathFloatBase
 		Result = false;
 	}
 
-	MULTIPLEX_METHOD()
+	RIGVM_METHOD()
 	virtual void Execute(const FRigUnitContext& Context) override;
 
 	UPROPERTY(meta=(Input))
@@ -527,7 +527,7 @@ struct FRigUnit_MathFloatGreaterEqual : public FRigUnit_MathFloatBase
 		Result = true;
 	}
 
-	MULTIPLEX_METHOD()
+	RIGVM_METHOD()
 	virtual void Execute(const FRigUnitContext& Context) override;
 
 	UPROPERTY(meta=(Input))
@@ -554,7 +554,7 @@ struct FRigUnit_MathFloatLessEqual : public FRigUnit_MathFloatBase
 		Result = true;
 	}
 
-	MULTIPLEX_METHOD()
+	RIGVM_METHOD()
 	virtual void Execute(const FRigUnitContext& Context) override;
 
 	UPROPERTY(meta=(Input))
@@ -581,7 +581,7 @@ struct FRigUnit_MathFloatIsNearlyZero : public FRigUnit_MathFloatBase
 		Result = true;
 	}
 	
-	MULTIPLEX_METHOD()
+	RIGVM_METHOD()
 	virtual void Execute(const FRigUnitContext& Context) override;
 
 	UPROPERTY(meta=(Input))
@@ -608,7 +608,7 @@ struct FRigUnit_MathFloatIsNearlyEqual : public FRigUnit_MathFloatBase
 		Result = true;
 	}
 
-	MULTIPLEX_METHOD()
+	RIGVM_METHOD()
 	virtual void Execute(const FRigUnitContext& Context) override;
 
 	UPROPERTY(meta=(Input))
@@ -638,7 +638,7 @@ struct FRigUnit_MathFloatSelectBool : public FRigUnit_MathFloatBase
 		IfTrue = IfFalse = Result = 0.f;
 	}
 
-	MULTIPLEX_METHOD()
+	RIGVM_METHOD()
 	virtual void Execute(const FRigUnitContext& Context) override;
 
 	UPROPERTY(meta=(Input))
@@ -662,7 +662,7 @@ struct FRigUnit_MathFloatDeg : public FRigUnit_MathFloatUnaryOp
 {
 	GENERATED_BODY()
 
-	MULTIPLEX_METHOD()
+	RIGVM_METHOD()
 	virtual void Execute(const FRigUnitContext& Context) override;
 };
 
@@ -674,7 +674,7 @@ struct FRigUnit_MathFloatRad : public FRigUnit_MathFloatUnaryOp
 {
 	GENERATED_BODY()
 
-	MULTIPLEX_METHOD()
+	RIGVM_METHOD()
 	virtual void Execute(const FRigUnitContext& Context) override;
 };
 
@@ -686,7 +686,7 @@ struct FRigUnit_MathFloatSin : public FRigUnit_MathFloatUnaryOp
 {
 	GENERATED_BODY()
 
-	MULTIPLEX_METHOD()
+	RIGVM_METHOD()
 	virtual void Execute(const FRigUnitContext& Context) override;
 };
 
@@ -698,7 +698,7 @@ struct FRigUnit_MathFloatCos : public FRigUnit_MathFloatUnaryOp
 {
 	GENERATED_BODY()
 
-	MULTIPLEX_METHOD()
+	RIGVM_METHOD()
 	virtual void Execute(const FRigUnitContext& Context) override;
 };
 
@@ -710,7 +710,7 @@ struct FRigUnit_MathFloatTan : public FRigUnit_MathFloatUnaryOp
 {
 	GENERATED_BODY()
 
-	MULTIPLEX_METHOD()
+	RIGVM_METHOD()
 	virtual void Execute(const FRigUnitContext& Context) override;
 };
 
@@ -722,7 +722,7 @@ struct FRigUnit_MathFloatAsin : public FRigUnit_MathFloatUnaryOp
 {
 	GENERATED_BODY()
 
-	MULTIPLEX_METHOD()
+	RIGVM_METHOD()
 	virtual void Execute(const FRigUnitContext& Context) override;
 };
 
@@ -734,7 +734,7 @@ struct FRigUnit_MathFloatAcos : public FRigUnit_MathFloatUnaryOp
 {
 	GENERATED_BODY()
 
-	MULTIPLEX_METHOD()
+	RIGVM_METHOD()
 	virtual void Execute(const FRigUnitContext& Context) override;
 };
 
@@ -746,7 +746,7 @@ struct FRigUnit_MathFloatAtan : public FRigUnit_MathFloatUnaryOp
 {
 	GENERATED_BODY()
 
-	MULTIPLEX_METHOD()
+	RIGVM_METHOD()
 	virtual void Execute(const FRigUnitContext& Context) override;
 };
 
@@ -764,7 +764,7 @@ struct FRigUnit_MathFloatLawOfCosine : public FRigUnit_MathFloatBase
 		bValid = false;
 	}
 
-	MULTIPLEX_METHOD()
+	RIGVM_METHOD()
 	virtual void Execute(const FRigUnitContext& Context) override;
 
 	UPROPERTY(meta = (Input))
diff --git a/Engine/Plugins/Experimental/ControlRig/Source/ControlRig/Private/Units/Math/RigUnit_MathQuaternion.cpp b/Engine/Plugins/Experimental/ControlRig/Source/ControlRig/Private/Units/Math/RigUnit_MathQuaternion.cpp
--- a/Engine/Plugins/Experimental/ControlRig/Source/ControlRig/Private/Units/Math/RigUnit_MathQuaternion.cpp
+++ b/Engine/Plugins/Experimental/ControlRig/Source/ControlRig/Private/Units/Math/RigUnit_MathQuaternion.cpp
@@ -3,7 +3,7 @@
 #include "Units/Math/RigUnit_MathQuaternion.h"
 #include "Units/RigUnitContext.h"
 
-UE_RigUnit_MathQuaternionFromAxisAndAngle_IMPLEMENT_MULTIPLEX(void, Execute, const FRigUnitContext& Context)
+UE_RigUnit_MathQuaternionFromAxisAndAngle_IMPLEMENT_RIGVM(void, Execute, const FRigUnitContext& Context)
 {
     DECLARE_SCOPE_HIERARCHICAL_COUNTER_RIGUNIT()
 	if (Axis.IsNearlyZero())
@@ -15,19 +15,19 @@ UE_RigUnit_MathQuaternionFromAxisAndAngle_IMPLEMENT_MULTIPLEX(void, Execute, con
 	Result = FQuat(Axis.GetUnsafeNormal(), Angle);
 }
 
-UE_RigUnit_MathQuaternionFromEuler_IMPLEMENT_MULTIPLEX(void, Execute, const FRigUnitContext& Context)
+UE_RigUnit_MathQuaternionFromEuler_IMPLEMENT_RIGVM(void, Execute, const FRigUnitContext& Context)
 {
     DECLARE_SCOPE_HIERARCHICAL_COUNTER_RIGUNIT()
 	Result = FControlRigMathLibrary::QuatFromEuler(Euler, RotationOrder);
 }
 
-UE_RigUnit_MathQuaternionFromRotator_IMPLEMENT_MULTIPLEX(void, Execute, const FRigUnitContext& Context)
+UE_RigUnit_MathQuaternionFromRotator_IMPLEMENT_RIGVM(void, Execute, const FRigUnitContext& Context)
 {
     DECLARE_SCOPE_HIERARCHICAL_COUNTER_RIGUNIT()
 	Result = FQuat(Rotator);
 }
 
-UE_RigUnit_MathQuaternionFromTwoVectors_IMPLEMENT_MULTIPLEX(void, Execute, const FRigUnitContext& Context)
+UE_RigUnit_MathQuaternionFromTwoVectors_IMPLEMENT_RIGVM(void, Execute, const FRigUnitContext& Context)
 {
     DECLARE_SCOPE_HIERARCHICAL_COUNTER_RIGUNIT()
 	if (A.IsNearlyZero() || B.IsNearlyZero())
@@ -46,79 +46,79 @@ UE_RigUnit_MathQuaternionFromTwoVectors_IMPLEMENT_MULTIPLEX(void, Execute, const
 	Result = FQuat::FindBetweenVectors(A, B).GetNormalized();
 }
 
-UE_RigUnit_MathQuaternionToAxisAndAngle_IMPLEMENT_MULTIPLEX(void, Execute, const FRigUnitContext& Context)
+UE_RigUnit_MathQuaternionToAxisAndAngle_IMPLEMENT_RIGVM(void, Execute, const FRigUnitContext& Context)
 {
     DECLARE_SCOPE_HIERARCHICAL_COUNTER_RIGUNIT()
 	Value.ToAxisAndAngle(Axis, Angle);
 }
 
-UE_RigUnit_MathQuaternionToEuler_IMPLEMENT_MULTIPLEX(void, Execute, const FRigUnitContext& Context)
+UE_RigUnit_MathQuaternionToEuler_IMPLEMENT_RIGVM(void, Execute, const FRigUnitContext& Context)
 {
     DECLARE_SCOPE_HIERARCHICAL_COUNTER_RIGUNIT()
 	Result = FControlRigMathLibrary::EulerFromQuat(Value, RotationOrder);
 }
 
-UE_RigUnit_MathQuaternionToRotator_IMPLEMENT_MULTIPLEX(void, Execute, const FRigUnitContext& Context)
+UE_RigUnit_MathQuaternionToRotator_IMPLEMENT_RIGVM(void, Execute, const FRigUnitContext& Context)
 {
     DECLARE_SCOPE_HIERARCHICAL_COUNTER_RIGUNIT()
 	Result = Value.Rotator();
 }
 
-UE_RigUnit_MathQuaternionMul_IMPLEMENT_MULTIPLEX(void, Execute, const FRigUnitContext& Context)
+UE_RigUnit_MathQuaternionMul_IMPLEMENT_RIGVM(void, Execute, const FRigUnitContext& Context)
 {
     DECLARE_SCOPE_HIERARCHICAL_COUNTER_RIGUNIT()
 	Result = A * B;
 }
 
-UE_RigUnit_MathQuaternionInverse_IMPLEMENT_MULTIPLEX(void, Execute, const FRigUnitContext& Context)
+UE_RigUnit_MathQuaternionInverse_IMPLEMENT_RIGVM(void, Execute, const FRigUnitContext& Context)
 {
     DECLARE_SCOPE_HIERARCHICAL_COUNTER_RIGUNIT()
 	Result = Value.Inverse();
 }
 
-UE_RigUnit_MathQuaternionSlerp_IMPLEMENT_MULTIPLEX(void, Execute, const FRigUnitContext& Context)
+UE_RigUnit_MathQuaternionSlerp_IMPLEMENT_RIGVM(void, Execute, const FRigUnitContext& Context)
 {
     DECLARE_SCOPE_HIERARCHICAL_COUNTER_RIGUNIT()
 	Result = FQuat::Slerp(A, B, T);
 }
 
-UE_RigUnit_MathQuaternionEquals_IMPLEMENT_MULTIPLEX(void, Execute, const FRigUnitContext& Context)
+UE_RigUnit_MathQuaternionEquals_IMPLEMENT_RIGVM(void, Execute, const FRigUnitContext& Context)
 {
     DECLARE_SCOPE_HIERARCHICAL_COUNTER_RIGUNIT()
 	Result = A == B;
 }
 
-UE_RigUnit_MathQuaternionNotEquals_IMPLEMENT_MULTIPLEX(void, Execute, const FRigUnitContext& Context)
+UE_RigUnit_MathQuaternionNotEquals_IMPLEMENT_RIGVM(void, Execute, const FRigUnitContext& Context)
 {
     DECLARE_SCOPE_HIERARCHICAL_COUNTER_RIGUNIT()
 	Result = A != B;
 }
 
-UE_RigUnit_MathQuaternionSelectBool_IMPLEMENT_MULTIPLEX(void, Execute, const FRigUnitContext& Context)
+UE_RigUnit_MathQuaternionSelectBool_IMPLEMENT_RIGVM(void, Execute, const FRigUnitContext& Context)
 {
     DECLARE_SCOPE_HIERARCHICAL_COUNTER_RIGUNIT()
 	Result = Condition ? IfTrue : IfFalse;
 }
 
-UE_RigUnit_MathQuaternionDot_IMPLEMENT_MULTIPLEX(void, Execute, const FRigUnitContext& Context)
+UE_RigUnit_MathQuaternionDot_IMPLEMENT_RIGVM(void, Execute, const FRigUnitContext& Context)
 {
     DECLARE_SCOPE_HIERARCHICAL_COUNTER_RIGUNIT()
 	Result = A | B;
 }
 
-UE_RigUnit_MathQuaternionUnit_IMPLEMENT_MULTIPLEX(void, Execute, const FRigUnitContext& Context)
+UE_RigUnit_MathQuaternionUnit_IMPLEMENT_RIGVM(void, Execute, const FRigUnitContext& Context)
 {
     DECLARE_SCOPE_HIERARCHICAL_COUNTER_RIGUNIT()
 	Result = Value.GetNormalized();
 }
 
-UE_RigUnit_MathQuaternionRotateVector_IMPLEMENT_MULTIPLEX(void, Execute, const FRigUnitContext& Context)
+UE_RigUnit_MathQuaternionRotateVector_IMPLEMENT_RIGVM(void, Execute, const FRigUnitContext& Context)
 {
     DECLARE_SCOPE_HIERARCHICAL_COUNTER_RIGUNIT()
 	Result = Quaternion.RotateVector(Vector);
 }
 
-UE_RigUnit_MathQuaternionGetAxis_IMPLEMENT_MULTIPLEX(void, Execute, const FRigUnitContext& Context)
+UE_RigUnit_MathQuaternionGetAxis_IMPLEMENT_RIGVM(void, Execute, const FRigUnitContext& Context)
 {
     DECLARE_SCOPE_HIERARCHICAL_COUNTER_RIGUNIT()
 	switch (Axis)
@@ -143,7 +143,7 @@ UE_RigUnit_MathQuaternionGetAxis_IMPLEMENT_MULTIPLEX(void, Execute, const FRigUn
 }
 
 
-UE_RigUnit_MathQuaternionSwingTwist_IMPLEMENT_MULTIPLEX(void, Execute, const FRigUnitContext& Context)
+UE_RigUnit_MathQuaternionSwingTwist_IMPLEMENT_RIGVM(void, Execute, const FRigUnitContext& Context)
 {
     DECLARE_SCOPE_HIERARCHICAL_COUNTER_RIGUNIT()
 	if (TwistAxis.IsNearlyZero())
diff --git a/Engine/Plugins/Experimental/ControlRig/Source/ControlRig/Private/Units/Math/RigUnit_MathQuaternion.h b/Engine/Plugins/Experimental/ControlRig/Source/ControlRig/Private/Units/Math/RigUnit_MathQuaternion.h
--- a/Engine/Plugins/Experimental/ControlRig/Source/ControlRig/Private/Units/Math/RigUnit_MathQuaternion.h
+++ b/Engine/Plugins/Experimental/ControlRig/Source/ControlRig/Private/Units/Math/RigUnit_MathQuaternion.h
@@ -64,7 +64,7 @@ struct FRigUnit_MathQuaternionFromAxisAndAngle : public FRigUnit_MathQuaternionB
 		Result = FQuat::Identity;
 	}
 
-	MULTIPLEX_METHOD()
+	RIGVM_METHOD()
 	virtual void Execute(const FRigUnitContext& Context) override;
 
 	UPROPERTY(meta=(Input))
@@ -92,7 +92,7 @@ struct FRigUnit_MathQuaternionFromEuler : public FRigUnit_MathQuaternionBase
 		Result = FQuat::Identity;
 	}
 
-	MULTIPLEX_METHOD()
+	RIGVM_METHOD()
 	virtual void Execute(const FRigUnitContext& Context) override;
 
 	UPROPERTY(meta = (Input))
@@ -119,7 +119,7 @@ struct FRigUnit_MathQuaternionFromRotator : public FRigUnit_MathQuaternionBase
 		Result = FQuat::Identity;
 	}
 
-	MULTIPLEX_METHOD()
+	RIGVM_METHOD()
 	virtual void Execute(const FRigUnitContext& Context) override;
 
 	UPROPERTY(meta=(Input))
@@ -143,7 +143,7 @@ struct FRigUnit_MathQuaternionFromTwoVectors : public FRigUnit_MathQuaternionBas
 		Result = FQuat::Identity;
 	}
 
-	MULTIPLEX_METHOD()
+	RIGVM_METHOD()
 	virtual void Execute(const FRigUnitContext& Context) override;
 
 	UPROPERTY(meta=(Input))
@@ -171,7 +171,7 @@ struct FRigUnit_MathQuaternionToAxisAndAngle : public FRigUnit_MathQuaternionBas
 		Angle = 0.f;
 	}
 
-	MULTIPLEX_METHOD()
+	RIGVM_METHOD()
 	virtual void Execute(const FRigUnitContext& Context) override;
 
 	UPROPERTY(meta=(Input))
@@ -199,7 +199,7 @@ struct FRigUnit_MathQuaternionToEuler : public FRigUnit_MathQuaternionBase
 		Result = FVector::ZeroVector;
 	}
 
-	MULTIPLEX_METHOD()
+	RIGVM_METHOD()
 	virtual void Execute(const FRigUnitContext& Context) override;
 
 	UPROPERTY(meta=(Input))
@@ -226,7 +226,7 @@ struct FRigUnit_MathQuaternionToRotator : public FRigUnit_MathQuaternionBase
 		Result = FRotator::ZeroRotator;
 	}
 
-	MULTIPLEX_METHOD()
+	RIGVM_METHOD()
 	virtual void Execute(const FRigUnitContext& Context) override;
 
 	UPROPERTY(meta=(Input))
@@ -243,7 +243,7 @@ USTRUCT(meta=(DisplayName="Multiply", PrototypeName="Multiply", Keywords="Produc
 struct FRigUnit_MathQuaternionMul : public FRigUnit_MathQuaternionBinaryOp
 {
 	GENERATED_BODY()
-	MULTIPLEX_METHOD()
+	RIGVM_METHOD()
 	virtual void Execute(const FRigUnitContext& Context) override;
 };
 
@@ -254,7 +254,7 @@ USTRUCT(meta=(DisplayName="Inverse", PrototypeName="Inverse"))
 struct FRigUnit_MathQuaternionInverse : public FRigUnit_MathQuaternionUnaryOp
 {
 	GENERATED_BODY()
-	MULTIPLEX_METHOD()
+	RIGVM_METHOD()
 	virtual void Execute(const FRigUnitContext& Context) override;
 };
 
@@ -272,7 +272,7 @@ struct FRigUnit_MathQuaternionSlerp : public FRigUnit_MathQuaternionBase
 		T = 0.f;
 	}
 
-	MULTIPLEX_METHOD()
+	RIGVM_METHOD()
 	virtual void Execute(const FRigUnitContext& Context) override;
 
 	UPROPERTY(meta=(Input))
@@ -302,7 +302,7 @@ struct FRigUnit_MathQuaternionEquals : public FRigUnit_MathQuaternionBase
 		Result = true;
 	}	
 
-	MULTIPLEX_METHOD()
+	RIGVM_METHOD()
 	virtual void Execute(const FRigUnitContext& Context) override;
 
 	UPROPERTY(meta=(Input))
@@ -329,7 +329,7 @@ struct FRigUnit_MathQuaternionNotEquals : public FRigUnit_MathQuaternionBase
 		Result = false;
 	}
 
-	MULTIPLEX_METHOD()
+	RIGVM_METHOD()
 	virtual void Execute(const FRigUnitContext& Context) override;
 
 	UPROPERTY(meta=(Input))
@@ -356,7 +356,7 @@ struct FRigUnit_MathQuaternionSelectBool : public FRigUnit_MathQuaternionBase
 		Condition = true;
 	}
 
-	MULTIPLEX_METHOD()
+	RIGVM_METHOD()
 	virtual void Execute(const FRigUnitContext& Context) override;
 
 	UPROPERTY(meta=(Input))
@@ -386,7 +386,7 @@ struct FRigUnit_MathQuaternionDot : public FRigUnit_MathQuaternionBase
 		Result = 0.f;
 	}
 
-	MULTIPLEX_METHOD()
+	RIGVM_METHOD()
 	virtual void Execute(const FRigUnitContext& Context) override;
 
 	UPROPERTY(meta=(Input))
@@ -406,7 +406,7 @@ USTRUCT(meta=(DisplayName="Unit", PrototypeName="Unit", Keywords="Normalize"))
 struct FRigUnit_MathQuaternionUnit : public FRigUnit_MathQuaternionUnaryOp
 {
 	GENERATED_BODY()
-	MULTIPLEX_METHOD()
+	RIGVM_METHOD()
 	virtual void Execute(const FRigUnitContext& Context) override;
 };
 
@@ -424,7 +424,7 @@ struct FRigUnit_MathQuaternionRotateVector : public FRigUnit_MathQuaternionBase
 		Vector = Result = FVector::ZeroVector;
 	}
 
-	MULTIPLEX_METHOD()
+	RIGVM_METHOD()
 	virtual void Execute(const FRigUnitContext& Context) override;
 
 	UPROPERTY(meta=(Input))
@@ -452,7 +452,7 @@ struct FRigUnit_MathQuaternionGetAxis: public FRigUnit_MathQuaternionBase
 		Result = FVector(1.f, 0.f, 0.f);
 	}
 
-	MULTIPLEX_METHOD()
+	RIGVM_METHOD()
 	virtual void Execute(const FRigUnitContext& Context) override;
 
 	UPROPERTY(meta = (Input))
@@ -480,7 +480,7 @@ struct FRigUnit_MathQuaternionSwingTwist : public FRigUnit_MathQuaternionBase
 		TwistAxis = FVector(1.f, 0.f, 0.f);
 	}
 
-	MULTIPLEX_METHOD()
+	RIGVM_METHOD()
 	virtual void Execute(const FRigUnitContext& Context) override;
 
 	UPROPERTY(meta = (Input))
diff --git a/Engine/Plugins/Experimental/ControlRig/Source/ControlRig/Private/Units/Math/RigUnit_MathTransform.cpp b/Engine/Plugins/Experimental/ControlRig/Source/ControlRig/Private/Units/Math/RigUnit_MathTransform.cpp
--- a/Engine/Plugins/Experimental/ControlRig/Source/ControlRig/Private/Units/Math/RigUnit_MathTransform.cpp
+++ b/Engine/Plugins/Experimental/ControlRig/Source/ControlRig/Private/Units/Math/RigUnit_MathTransform.cpp
@@ -4,61 +4,61 @@
 #include "Math/ControlRigMathLibrary.h"
 #include "Units/RigUnitContext.h"
 
-UE_RigUnit_MathTransformFromEulerTransform_IMPLEMENT_MULTIPLEX(void, Execute, const FRigUnitContext& Context)
+UE_RigUnit_MathTransformFromEulerTransform_IMPLEMENT_RIGVM(void, Execute, const FRigUnitContext& Context)
 {
     DECLARE_SCOPE_HIERARCHICAL_COUNTER_RIGUNIT()
 	Result = EulerTransform.ToFTransform();
 }
 
-UE_RigUnit_MathTransformToEulerTransform_IMPLEMENT_MULTIPLEX(void, Execute, const FRigUnitContext& Context)
+UE_RigUnit_MathTransformToEulerTransform_IMPLEMENT_RIGVM(void, Execute, const FRigUnitContext& Context)
 {
     DECLARE_SCOPE_HIERARCHICAL_COUNTER_RIGUNIT()
 	Result.FromFTransform(Value);
 }
 
-UE_RigUnit_MathTransformMul_IMPLEMENT_MULTIPLEX(void, Execute, const FRigUnitContext& Context)
+UE_RigUnit_MathTransformMul_IMPLEMENT_RIGVM(void, Execute, const FRigUnitContext& Context)
 {
     DECLARE_SCOPE_HIERARCHICAL_COUNTER_RIGUNIT()
 	Result = A * B;
 }
 
-UE_RigUnit_MathTransformMakeRelative_IMPLEMENT_MULTIPLEX(void, Execute, const FRigUnitContext& Context)
+UE_RigUnit_MathTransformMakeRelative_IMPLEMENT_RIGVM(void, Execute, const FRigUnitContext& Context)
 {
     DECLARE_SCOPE_HIERARCHICAL_COUNTER_RIGUNIT()
 	Local = Global.GetRelativeTransform(Parent);
 }
 
-UE_RigUnit_MathTransformInverse_IMPLEMENT_MULTIPLEX(void, Execute, const FRigUnitContext& Context)
+UE_RigUnit_MathTransformInverse_IMPLEMENT_RIGVM(void, Execute, const FRigUnitContext& Context)
 {
     DECLARE_SCOPE_HIERARCHICAL_COUNTER_RIGUNIT()
 	Result = Value.Inverse();
 }
 
-UE_RigUnit_MathTransformLerp_IMPLEMENT_MULTIPLEX(void, Execute, const FRigUnitContext& Context)
+UE_RigUnit_MathTransformLerp_IMPLEMENT_RIGVM(void, Execute, const FRigUnitContext& Context)
 {
     DECLARE_SCOPE_HIERARCHICAL_COUNTER_RIGUNIT()
 	Result = FControlRigMathLibrary::LerpTransform(A, B, T);
 }
 
-UE_RigUnit_MathTransformSelectBool_IMPLEMENT_MULTIPLEX(void, Execute, const FRigUnitContext& Context)
+UE_RigUnit_MathTransformSelectBool_IMPLEMENT_RIGVM(void, Execute, const FRigUnitContext& Context)
 {
     DECLARE_SCOPE_HIERARCHICAL_COUNTER_RIGUNIT()
 	Result = Condition ? IfTrue : IfFalse;
 }
 
-UE_RigUnit_MathTransformRotateVector_IMPLEMENT_MULTIPLEX(void, Execute, const FRigUnitContext& Context)
+UE_RigUnit_MathTransformRotateVector_IMPLEMENT_RIGVM(void, Execute, const FRigUnitContext& Context)
 {
     DECLARE_SCOPE_HIERARCHICAL_COUNTER_RIGUNIT()
 	Result = Transform.TransformVector(Direction);
 }
 
-UE_RigUnit_MathTransformTransformVector_IMPLEMENT_MULTIPLEX(void, Execute, const FRigUnitContext& Context)
+UE_RigUnit_MathTransformTransformVector_IMPLEMENT_RIGVM(void, Execute, const FRigUnitContext& Context)
 {
     DECLARE_SCOPE_HIERARCHICAL_COUNTER_RIGUNIT()
 	Result = Transform.TransformPosition(Location);
 }
 
-UE_RigUnit_MathTransformFromSRT_IMPLEMENT_MULTIPLEX(void, Execute, const FRigUnitContext& Context)
+UE_RigUnit_MathTransformFromSRT_IMPLEMENT_RIGVM(void, Execute, const FRigUnitContext& Context)
 {
     DECLARE_SCOPE_HIERARCHICAL_COUNTER_RIGUNIT()
 	Transform.SetLocation(Location);
diff --git a/Engine/Plugins/Experimental/ControlRig/Source/ControlRig/Private/Units/Math/RigUnit_MathTransform.h b/Engine/Plugins/Experimental/ControlRig/Source/ControlRig/Private/Units/Math/RigUnit_MathTransform.h
--- a/Engine/Plugins/Experimental/ControlRig/Source/ControlRig/Private/Units/Math/RigUnit_MathTransform.h
+++ b/Engine/Plugins/Experimental/ControlRig/Source/ControlRig/Private/Units/Math/RigUnit_MathTransform.h
@@ -64,7 +64,7 @@ struct FRigUnit_MathTransformFromEulerTransform : public FRigUnit_MathTransformB
 		Result = FTransform::Identity;
 	}
 
-	MULTIPLEX_METHOD()
+	RIGVM_METHOD()
 	virtual void Execute(const FRigUnitContext& Context) override;
 
 	UPROPERTY(meta=(Input))
@@ -88,7 +88,7 @@ struct FRigUnit_MathTransformToEulerTransform : public FRigUnit_MathTransformBas
 		Result = FEulerTransform::Identity;
 	}
 
-	MULTIPLEX_METHOD()
+	RIGVM_METHOD()
 	virtual void Execute(const FRigUnitContext& Context) override;
 
 	UPROPERTY(meta=(Input))
@@ -105,7 +105,7 @@ USTRUCT(meta=(DisplayName="Multiply", PrototypeName="Multiply", Keywords="Produc
 struct FRigUnit_MathTransformMul : public FRigUnit_MathTransformBinaryOp
 {
 	GENERATED_BODY()
-	MULTIPLEX_METHOD()
+	RIGVM_METHOD()
 	virtual void Execute(const FRigUnitContext& Context) override;
 };
 
@@ -122,7 +122,7 @@ struct FRigUnit_MathTransformMakeRelative : public FRigUnit_MathTransformBase
 		Global = Parent = Local = FTransform::Identity;
 	}
 
-	MULTIPLEX_METHOD()
+	RIGVM_METHOD()
 	virtual void Execute(const FRigUnitContext& Context) override;
 
 	UPROPERTY(meta=(Input))
@@ -142,7 +142,7 @@ USTRUCT(meta=(DisplayName="Inverse", PrototypeName="Inverse"))
 struct FRigUnit_MathTransformInverse : public FRigUnit_MathTransformUnaryOp
 {
 	GENERATED_BODY()
-	MULTIPLEX_METHOD()
+	RIGVM_METHOD()
 	virtual void Execute(const FRigUnitContext& Context) override;
 };
 
@@ -160,7 +160,7 @@ struct FRigUnit_MathTransformLerp : public FRigUnit_MathTransformBase
 		T = 0.f;
 	}
 
-	MULTIPLEX_METHOD()
+	RIGVM_METHOD()
 	virtual void Execute(const FRigUnitContext& Context) override;
 
 	UPROPERTY(meta=(Input))
@@ -190,7 +190,7 @@ struct FRigUnit_MathTransformSelectBool : public FRigUnit_MathTransformBase
 		IfTrue = IfFalse = Result = FTransform::Identity;
 	}
 
-	MULTIPLEX_METHOD()
+	RIGVM_METHOD()
 	virtual void Execute(const FRigUnitContext& Context) override;
 
 	UPROPERTY(meta=(Input))
@@ -220,7 +220,7 @@ struct FRigUnit_MathTransformRotateVector : public FRigUnit_MathTransformBase
 		Direction = Result = FVector::ZeroVector;
 	}
 
-	MULTIPLEX_METHOD()
+	RIGVM_METHOD()
 	virtual void Execute(const FRigUnitContext& Context) override;
 
 	UPROPERTY(meta=(Input))
@@ -247,7 +247,7 @@ struct FRigUnit_MathTransformTransformVector : public FRigUnit_MathTransformBase
 		Location = Result = FVector::ZeroVector;
 	}
 	
-	MULTIPLEX_METHOD()
+	RIGVM_METHOD()
 	virtual void Execute(const FRigUnitContext& Context) override;
 
 	UPROPERTY(meta=(Input))
@@ -278,7 +278,7 @@ struct FRigUnit_MathTransformFromSRT : public FRigUnit_MathTransformBase
 		EulerTransform = FEulerTransform::Identity;
 	}
 	
-	MULTIPLEX_METHOD()
+	RIGVM_METHOD()
 	virtual void Execute(const FRigUnitContext& Context) override;
 
 	UPROPERTY(meta=(Input))
diff --git a/Engine/Plugins/Experimental/ControlRig/Source/ControlRig/Private/Units/Math/RigUnit_MathVector.cpp b/Engine/Plugins/Experimental/ControlRig/Source/ControlRig/Private/Units/Math/RigUnit_MathVector.cpp
--- a/Engine/Plugins/Experimental/ControlRig/Source/ControlRig/Private/Units/Math/RigUnit_MathVector.cpp
+++ b/Engine/Plugins/Experimental/ControlRig/Source/ControlRig/Private/Units/Math/RigUnit_MathVector.cpp
@@ -4,37 +4,37 @@
 #include "Units/RigUnitContext.h"
 #include "Math/ControlRigMathLibrary.h"
 
-UE_RigUnit_MathVectorFromFloat_IMPLEMENT_MULTIPLEX(void, Execute, const FRigUnitContext& Context)
+UE_RigUnit_MathVectorFromFloat_IMPLEMENT_RIGVM(void, Execute, const FRigUnitContext& Context)
 {
     DECLARE_SCOPE_HIERARCHICAL_COUNTER_RIGUNIT()
 	Result = FVector(Value, Value, Value);
 }
 
-UE_RigUnit_MathVectorAdd_IMPLEMENT_MULTIPLEX(void, Execute, const FRigUnitContext& Context)
+UE_RigUnit_MathVectorAdd_IMPLEMENT_RIGVM(void, Execute, const FRigUnitContext& Context)
 {
     DECLARE_SCOPE_HIERARCHICAL_COUNTER_RIGUNIT()
 	Result = A + B;
 }
 
-UE_RigUnit_MathVectorSub_IMPLEMENT_MULTIPLEX(void, Execute, const FRigUnitContext& Context)
+UE_RigUnit_MathVectorSub_IMPLEMENT_RIGVM(void, Execute, const FRigUnitContext& Context)
 {
     DECLARE_SCOPE_HIERARCHICAL_COUNTER_RIGUNIT()
 	Result = A - B;
 }
 
-UE_RigUnit_MathVectorMul_IMPLEMENT_MULTIPLEX(void, Execute, const FRigUnitContext& Context)
+UE_RigUnit_MathVectorMul_IMPLEMENT_RIGVM(void, Execute, const FRigUnitContext& Context)
 {
     DECLARE_SCOPE_HIERARCHICAL_COUNTER_RIGUNIT()
 	Result = A * B;
 }
 
-UE_RigUnit_MathVectorScale_IMPLEMENT_MULTIPLEX(void, Execute, const FRigUnitContext& Context)
+UE_RigUnit_MathVectorScale_IMPLEMENT_RIGVM(void, Execute, const FRigUnitContext& Context)
 {
     DECLARE_SCOPE_HIERARCHICAL_COUNTER_RIGUNIT()
 	Result = Value * Factor;
 }
 
-UE_RigUnit_MathVectorDiv_IMPLEMENT_MULTIPLEX(void, Execute, const FRigUnitContext& Context)
+UE_RigUnit_MathVectorDiv_IMPLEMENT_RIGVM(void, Execute, const FRigUnitContext& Context)
 {
     DECLARE_SCOPE_HIERARCHICAL_COUNTER_RIGUNIT()
 	if(FMath::IsNearlyZero(B.X) || FMath::IsNearlyZero(B.Y) || FMath::IsNearlyZero(B.Z))
@@ -57,7 +57,7 @@ UE_RigUnit_MathVectorDiv_IMPLEMENT_MULTIPLEX(void, Execute, const FRigUnitContex
 	Result = A / B;
 }
 
-UE_RigUnit_MathVectorMod_IMPLEMENT_MULTIPLEX(void, Execute, const FRigUnitContext& Context)
+UE_RigUnit_MathVectorMod_IMPLEMENT_RIGVM(void, Execute, const FRigUnitContext& Context)
 {
     DECLARE_SCOPE_HIERARCHICAL_COUNTER_RIGUNIT()
 	if(FMath::IsNearlyZero(B.X) || FMath::IsNearlyZero(B.Y) || FMath::IsNearlyZero(B.Z) || B.X < 0.f || B.Y < 0.f || B.Z < 0.f)
@@ -83,7 +83,7 @@ UE_RigUnit_MathVectorMod_IMPLEMENT_MULTIPLEX(void, Execute, const FRigUnitContex
 	Result.Z = FMath::Fmod(A.Z, B.Z);
 }
 
-UE_RigUnit_MathVectorMin_IMPLEMENT_MULTIPLEX(void, Execute, const FRigUnitContext& Context)
+UE_RigUnit_MathVectorMin_IMPLEMENT_RIGVM(void, Execute, const FRigUnitContext& Context)
 {
     DECLARE_SCOPE_HIERARCHICAL_COUNTER_RIGUNIT()
 	Result.X = FMath::Min<float>(A.X, B.X);
@@ -91,7 +91,7 @@ UE_RigUnit_MathVectorMin_IMPLEMENT_MULTIPLEX(void, Execute, const FRigUnitContex
 	Result.Z = FMath::Min<float>(A.Z, B.Z);
 }
 
-UE_RigUnit_MathVectorMax_IMPLEMENT_MULTIPLEX(void, Execute, const FRigUnitContext& Context)
+UE_RigUnit_MathVectorMax_IMPLEMENT_RIGVM(void, Execute, const FRigUnitContext& Context)
 {
     DECLARE_SCOPE_HIERARCHICAL_COUNTER_RIGUNIT()
 	Result.X = FMath::Max<float>(A.X, B.X);
@@ -99,13 +99,13 @@ UE_RigUnit_MathVectorMax_IMPLEMENT_MULTIPLEX(void, Execute, const FRigUnitContex
 	Result.Z = FMath::Max<float>(A.Z, B.Z);
 }
 
-UE_RigUnit_MathVectorNegate_IMPLEMENT_MULTIPLEX(void, Execute, const FRigUnitContext& Context)
+UE_RigUnit_MathVectorNegate_IMPLEMENT_RIGVM(void, Execute, const FRigUnitContext& Context)
 {
     DECLARE_SCOPE_HIERARCHICAL_COUNTER_RIGUNIT()
 	Result = -Value;
 }
 
-UE_RigUnit_MathVectorAbs_IMPLEMENT_MULTIPLEX(void, Execute, const FRigUnitContext& Context)
+UE_RigUnit_MathVectorAbs_IMPLEMENT_RIGVM(void, Execute, const FRigUnitContext& Context)
 {
     DECLARE_SCOPE_HIERARCHICAL_COUNTER_RIGUNIT()
 	Result.X = FMath::Abs(Value.X);
@@ -113,7 +113,7 @@ UE_RigUnit_MathVectorAbs_IMPLEMENT_MULTIPLEX(void, Execute, const FRigUnitContex
 	Result.Z = FMath::Abs(Value.Z);
 }
 
-UE_RigUnit_MathVectorFloor_IMPLEMENT_MULTIPLEX(void, Execute, const FRigUnitContext& Context)
+UE_RigUnit_MathVectorFloor_IMPLEMENT_RIGVM(void, Execute, const FRigUnitContext& Context)
 {
     DECLARE_SCOPE_HIERARCHICAL_COUNTER_RIGUNIT()
 	Result.X = FMath::FloorToFloat(Value.X);
@@ -121,7 +121,7 @@ UE_RigUnit_MathVectorFloor_IMPLEMENT_MULTIPLEX(void, Execute, const FRigUnitCont
 	Result.Z = FMath::FloorToFloat(Value.Z);
 }
 
-UE_RigUnit_MathVectorCeil_IMPLEMENT_MULTIPLEX(void, Execute, const FRigUnitContext& Context)
+UE_RigUnit_MathVectorCeil_IMPLEMENT_RIGVM(void, Execute, const FRigUnitContext& Context)
 {
     DECLARE_SCOPE_HIERARCHICAL_COUNTER_RIGUNIT()
 	Result.X = FMath::CeilToFloat(Value.X);
@@ -129,7 +129,7 @@ UE_RigUnit_MathVectorCeil_IMPLEMENT_MULTIPLEX(void, Execute, const FRigUnitConte
 	Result.Z = FMath::CeilToFloat(Value.Z);
 }
 
-UE_RigUnit_MathVectorRound_IMPLEMENT_MULTIPLEX(void, Execute, const FRigUnitContext& Context)
+UE_RigUnit_MathVectorRound_IMPLEMENT_RIGVM(void, Execute, const FRigUnitContext& Context)
 {
     DECLARE_SCOPE_HIERARCHICAL_COUNTER_RIGUNIT()
 	Result.X = FMath::RoundToFloat(Value.X);
@@ -137,13 +137,13 @@ UE_RigUnit_MathVectorRound_IMPLEMENT_MULTIPLEX(void, Execute, const FRigUnitCont
 	Result.Z = FMath::RoundToFloat(Value.Z);
 }
 
-UE_RigUnit_MathVectorSign_IMPLEMENT_MULTIPLEX(void, Execute, const FRigUnitContext& Context)
+UE_RigUnit_MathVectorSign_IMPLEMENT_RIGVM(void, Execute, const FRigUnitContext& Context)
 {
     DECLARE_SCOPE_HIERARCHICAL_COUNTER_RIGUNIT()
 	Result = Value.GetSignVector();
 }
 
-UE_RigUnit_MathVectorClamp_IMPLEMENT_MULTIPLEX(void, Execute, const FRigUnitContext& Context)
+UE_RigUnit_MathVectorClamp_IMPLEMENT_RIGVM(void, Execute, const FRigUnitContext& Context)
 {
     DECLARE_SCOPE_HIERARCHICAL_COUNTER_RIGUNIT()
 	Result.X = FMath::Clamp<float>(Value.X, Minimum.X, Maximum.X);
@@ -151,13 +151,13 @@ UE_RigUnit_MathVectorClamp_IMPLEMENT_MULTIPLEX(void, Execute, const FRigUnitCont
 	Result.Z = FMath::Clamp<float>(Value.Z, Minimum.Z, Maximum.Z);
 }
 
-UE_RigUnit_MathVectorLerp_IMPLEMENT_MULTIPLEX(void, Execute, const FRigUnitContext& Context)
+UE_RigUnit_MathVectorLerp_IMPLEMENT_RIGVM(void, Execute, const FRigUnitContext& Context)
 {
     DECLARE_SCOPE_HIERARCHICAL_COUNTER_RIGUNIT()
 	Result = FMath::Lerp<FVector>(A, B, T);
 }
 
-UE_RigUnit_MathVectorRemap_IMPLEMENT_MULTIPLEX(void, Execute, const FRigUnitContext& Context)
+UE_RigUnit_MathVectorRemap_IMPLEMENT_RIGVM(void, Execute, const FRigUnitContext& Context)
 {
     DECLARE_SCOPE_HIERARCHICAL_COUNTER_RIGUNIT()
 	FVector Ratio(0.f, 0.f, 0.f);
@@ -189,19 +189,19 @@ UE_RigUnit_MathVectorRemap_IMPLEMENT_MULTIPLEX(void, Execute, const FRigUnitCont
 	Result = FMath::Lerp<FVector>(TargetMinimum, TargetMaximum, Ratio);
 }
 
-UE_RigUnit_MathVectorEquals_IMPLEMENT_MULTIPLEX(void, Execute, const FRigUnitContext& Context)
+UE_RigUnit_MathVectorEquals_IMPLEMENT_RIGVM(void, Execute, const FRigUnitContext& Context)
 {
     DECLARE_SCOPE_HIERARCHICAL_COUNTER_RIGUNIT()
 	Result = A == B;
 }
 
-UE_RigUnit_MathVectorNotEquals_IMPLEMENT_MULTIPLEX(void, Execute, const FRigUnitContext& Context)
+UE_RigUnit_MathVectorNotEquals_IMPLEMENT_RIGVM(void, Execute, const FRigUnitContext& Context)
 {
     DECLARE_SCOPE_HIERARCHICAL_COUNTER_RIGUNIT()
 	Result = A != B;
 }
 
-UE_RigUnit_MathVectorIsNearlyZero_IMPLEMENT_MULTIPLEX(void, Execute, const FRigUnitContext& Context)
+UE_RigUnit_MathVectorIsNearlyZero_IMPLEMENT_RIGVM(void, Execute, const FRigUnitContext& Context)
 {
     DECLARE_SCOPE_HIERARCHICAL_COUNTER_RIGUNIT()
 	if(Tolerance < 0.f)
@@ -211,7 +211,7 @@ UE_RigUnit_MathVectorIsNearlyZero_IMPLEMENT_MULTIPLEX(void, Execute, const FRigU
 	Result = Value.IsNearlyZero(FMath::Max<float>(Tolerance, SMALL_NUMBER));
 }
 
-UE_RigUnit_MathVectorIsNearlyEqual_IMPLEMENT_MULTIPLEX(void, Execute, const FRigUnitContext& Context)
+UE_RigUnit_MathVectorIsNearlyEqual_IMPLEMENT_RIGVM(void, Execute, const FRigUnitContext& Context)
 {
     DECLARE_SCOPE_HIERARCHICAL_COUNTER_RIGUNIT()
 	if(Tolerance < 0.f)
@@ -221,55 +221,55 @@ UE_RigUnit_MathVectorIsNearlyEqual_IMPLEMENT_MULTIPLEX(void, Execute, const FRig
 	Result = (A - B).IsNearlyZero(FMath::Max<float>(Tolerance, SMALL_NUMBER));
 }
 
-UE_RigUnit_MathVectorSelectBool_IMPLEMENT_MULTIPLEX(void, Execute, const FRigUnitContext& Context)
+UE_RigUnit_MathVectorSelectBool_IMPLEMENT_RIGVM(void, Execute, const FRigUnitContext& Context)
 {
     DECLARE_SCOPE_HIERARCHICAL_COUNTER_RIGUNIT()
 	Result = Condition ? IfTrue : IfFalse;
 }
 
-UE_RigUnit_MathVectorDeg_IMPLEMENT_MULTIPLEX(void, Execute, const FRigUnitContext& Context)
+UE_RigUnit_MathVectorDeg_IMPLEMENT_RIGVM(void, Execute, const FRigUnitContext& Context)
 {
     DECLARE_SCOPE_HIERARCHICAL_COUNTER_RIGUNIT()
 	Result = FMath::RadiansToDegrees(Value);
 }
 
-UE_RigUnit_MathVectorRad_IMPLEMENT_MULTIPLEX(void, Execute, const FRigUnitContext& Context)
+UE_RigUnit_MathVectorRad_IMPLEMENT_RIGVM(void, Execute, const FRigUnitContext& Context)
 {
     DECLARE_SCOPE_HIERARCHICAL_COUNTER_RIGUNIT()
 	Result = FMath::DegreesToRadians(Value);
 }
 
-UE_RigUnit_MathVectorLengthSquared_IMPLEMENT_MULTIPLEX(void, Execute, const FRigUnitContext& Context)
+UE_RigUnit_MathVectorLengthSquared_IMPLEMENT_RIGVM(void, Execute, const FRigUnitContext& Context)
 {
     DECLARE_SCOPE_HIERARCHICAL_COUNTER_RIGUNIT()
 	Result = Value.SizeSquared();
 }
 
-UE_RigUnit_MathVectorLength_IMPLEMENT_MULTIPLEX(void, Execute, const FRigUnitContext& Context)
+UE_RigUnit_MathVectorLength_IMPLEMENT_RIGVM(void, Execute, const FRigUnitContext& Context)
 {
     DECLARE_SCOPE_HIERARCHICAL_COUNTER_RIGUNIT()
 	Result = Value.Size();
 }
 
-UE_RigUnit_MathVectorDistance_IMPLEMENT_MULTIPLEX(void, Execute, const FRigUnitContext& Context)
+UE_RigUnit_MathVectorDistance_IMPLEMENT_RIGVM(void, Execute, const FRigUnitContext& Context)
 {
     DECLARE_SCOPE_HIERARCHICAL_COUNTER_RIGUNIT()
 	Result = FVector::Distance(A, B);
 }
 
-UE_RigUnit_MathVectorCross_IMPLEMENT_MULTIPLEX(void, Execute, const FRigUnitContext& Context)
+UE_RigUnit_MathVectorCross_IMPLEMENT_RIGVM(void, Execute, const FRigUnitContext& Context)
 {
     DECLARE_SCOPE_HIERARCHICAL_COUNTER_RIGUNIT()
 	Result = A ^ B;
 }
 
-UE_RigUnit_MathVectorDot_IMPLEMENT_MULTIPLEX(void, Execute, const FRigUnitContext& Context)
+UE_RigUnit_MathVectorDot_IMPLEMENT_RIGVM(void, Execute, const FRigUnitContext& Context)
 {
     DECLARE_SCOPE_HIERARCHICAL_COUNTER_RIGUNIT()
 	Result = A | B;
 }
 
-UE_RigUnit_MathVectorUnit_IMPLEMENT_MULTIPLEX(void, Execute, const FRigUnitContext& Context)
+UE_RigUnit_MathVectorUnit_IMPLEMENT_RIGVM(void, Execute, const FRigUnitContext& Context)
 {
     DECLARE_SCOPE_HIERARCHICAL_COUNTER_RIGUNIT()
 	if (Value.IsNearlyZero())
@@ -281,7 +281,7 @@ UE_RigUnit_MathVectorUnit_IMPLEMENT_MULTIPLEX(void, Execute, const FRigUnitConte
 	Result = Value.GetUnsafeNormal();
 }
 
-UE_RigUnit_MathVectorMirror_IMPLEMENT_MULTIPLEX(void, Execute, const FRigUnitContext& Context)
+UE_RigUnit_MathVectorMirror_IMPLEMENT_RIGVM(void, Execute, const FRigUnitContext& Context)
 {
     DECLARE_SCOPE_HIERARCHICAL_COUNTER_RIGUNIT()
 	if (Normal.IsNearlyZero())
@@ -293,7 +293,7 @@ UE_RigUnit_MathVectorMirror_IMPLEMENT_MULTIPLEX(void, Execute, const FRigUnitCon
 	Result = Value.MirrorByVector(Normal.GetSafeNormal());
 }
 
-UE_RigUnit_MathVectorAngle_IMPLEMENT_MULTIPLEX(void, Execute, const FRigUnitContext& Context)
+UE_RigUnit_MathVectorAngle_IMPLEMENT_RIGVM(void, Execute, const FRigUnitContext& Context)
 {
     DECLARE_SCOPE_HIERARCHICAL_COUNTER_RIGUNIT()
 	if (A.IsNearlyZero() || B.IsNearlyZero())
@@ -312,7 +312,7 @@ UE_RigUnit_MathVectorAngle_IMPLEMENT_MULTIPLEX(void, Execute, const FRigUnitCont
 	Result = FQuat::FindBetween(A, B).GetAngle();
 }
 
-UE_RigUnit_MathVectorParallel_IMPLEMENT_MULTIPLEX(void, Execute, const FRigUnitContext& Context)
+UE_RigUnit_MathVectorParallel_IMPLEMENT_RIGVM(void, Execute, const FRigUnitContext& Context)
 {
     DECLARE_SCOPE_HIERARCHICAL_COUNTER_RIGUNIT()
 	if (A.IsNearlyZero() || B.IsNearlyZero())
@@ -331,7 +331,7 @@ UE_RigUnit_MathVectorParallel_IMPLEMENT_MULTIPLEX(void, Execute, const FRigUnitC
 	Result = FVector::Parallel(A, B);
 }
 
-UE_RigUnit_MathVectorOrthogonal_IMPLEMENT_MULTIPLEX(void, Execute, const FRigUnitContext& Context)
+UE_RigUnit_MathVectorOrthogonal_IMPLEMENT_RIGVM(void, Execute, const FRigUnitContext& Context)
 {
     DECLARE_SCOPE_HIERARCHICAL_COUNTER_RIGUNIT()
 	if (A.IsNearlyZero() || B.IsNearlyZero())
@@ -350,7 +350,7 @@ UE_RigUnit_MathVectorOrthogonal_IMPLEMENT_MULTIPLEX(void, Execute, const FRigUni
 	Result = FVector::Orthogonal(A, B);
 }
 
-UE_RigUnit_MathVectorBezierFourPoint_IMPLEMENT_MULTIPLEX(void, Execute, const FRigUnitContext& Context)
+UE_RigUnit_MathVectorBezierFourPoint_IMPLEMENT_RIGVM(void, Execute, const FRigUnitContext& Context)
 {
     DECLARE_SCOPE_HIERARCHICAL_COUNTER_RIGUNIT()
 	FControlRigMathLibrary::FourPointBezier(Bezier, T, Result, Tangent);
diff --git a/Engine/Plugins/Experimental/ControlRig/Source/ControlRig/Private/Units/Math/RigUnit_MathVector.h b/Engine/Plugins/Experimental/ControlRig/Source/ControlRig/Private/Units/Math/RigUnit_MathVector.h
--- a/Engine/Plugins/Experimental/ControlRig/Source/ControlRig/Private/Units/Math/RigUnit_MathVector.h
+++ b/Engine/Plugins/Experimental/ControlRig/Source/ControlRig/Private/Units/Math/RigUnit_MathVector.h
@@ -56,7 +56,7 @@ USTRUCT(meta=(DisplayName="From Float", PrototypeName="FromFloat", Keywords="Mak
 struct FRigUnit_MathVectorFromFloat : public FRigUnit_MathVectorBase
 {
 	GENERATED_BODY()
-	MULTIPLEX_METHOD()
+	RIGVM_METHOD()
 	virtual void Execute(const FRigUnitContext& Context) override;
 
 	FRigUnit_MathVectorFromFloat()
@@ -80,7 +80,7 @@ struct FRigUnit_MathVectorAdd : public FRigUnit_MathVectorBinaryOp
 {
 	GENERATED_BODY()
 
-	MULTIPLEX_METHOD()
+	RIGVM_METHOD()
 	virtual void Execute(const FRigUnitContext& Context) override;
 };
 
@@ -92,7 +92,7 @@ struct FRigUnit_MathVectorSub : public FRigUnit_MathVectorBinaryOp
 {
 	GENERATED_BODY()
 
-	MULTIPLEX_METHOD()
+	RIGVM_METHOD()
 	virtual void Execute(const FRigUnitContext& Context) override;
 };
 
@@ -109,7 +109,7 @@ struct FRigUnit_MathVectorMul : public FRigUnit_MathVectorBinaryOp
 		A = B = FVector::OneVector;
 	}
 
-	MULTIPLEX_METHOD()
+	RIGVM_METHOD()
 	virtual void Execute(const FRigUnitContext& Context) override;
 };
 
@@ -127,7 +127,7 @@ struct FRigUnit_MathVectorScale : public FRigUnit_MathVectorBase
 		Factor = 1.f;
 	}
 
-	MULTIPLEX_METHOD()
+	RIGVM_METHOD()
 	virtual void Execute(const FRigUnitContext& Context) override;
 
 	UPROPERTY(meta = (Input))
@@ -153,7 +153,7 @@ struct FRigUnit_MathVectorDiv : public FRigUnit_MathVectorBinaryOp
 		B = FVector::OneVector;
 	}
 
-	MULTIPLEX_METHOD()
+	RIGVM_METHOD()
 	virtual void Execute(const FRigUnitContext& Context) override;
 };
 
@@ -171,7 +171,7 @@ struct FRigUnit_MathVectorMod : public FRigUnit_MathVectorBinaryOp
 		B = FVector::OneVector;
 	}
 
-	MULTIPLEX_METHOD()
+	RIGVM_METHOD()
 	virtual void Execute(const FRigUnitContext& Context) override;
 };
 
@@ -183,7 +183,7 @@ struct FRigUnit_MathVectorMin : public FRigUnit_MathVectorBinaryOp
 {
 	GENERATED_BODY()
 
-	MULTIPLEX_METHOD()
+	RIGVM_METHOD()
 	virtual void Execute(const FRigUnitContext& Context) override;
 };
 
@@ -195,7 +195,7 @@ struct FRigUnit_MathVectorMax : public FRigUnit_MathVectorBinaryOp
 {
 	GENERATED_BODY()
 
-	MULTIPLEX_METHOD()
+	RIGVM_METHOD()
 	virtual void Execute(const FRigUnitContext& Context) override;
 };
 
@@ -207,7 +207,7 @@ struct FRigUnit_MathVectorNegate : public FRigUnit_MathVectorUnaryOp
 {
 	GENERATED_BODY()
 
-	MULTIPLEX_METHOD()
+	RIGVM_METHOD()
 	virtual void Execute(const FRigUnitContext& Context) override;
 };
 
@@ -219,7 +219,7 @@ struct FRigUnit_MathVectorAbs : public FRigUnit_MathVectorUnaryOp
 {
 	GENERATED_BODY()
 
-	MULTIPLEX_METHOD()
+	RIGVM_METHOD()
 	virtual void Execute(const FRigUnitContext& Context) override;
 };
 
@@ -231,7 +231,7 @@ struct FRigUnit_MathVectorFloor : public FRigUnit_MathVectorUnaryOp
 {
 	GENERATED_BODY()
 
-	MULTIPLEX_METHOD()
+	RIGVM_METHOD()
 	virtual void Execute(const FRigUnitContext& Context) override;
 };
 
@@ -243,7 +243,7 @@ struct FRigUnit_MathVectorCeil : public FRigUnit_MathVectorUnaryOp
 {
 	GENERATED_BODY()
 
-	MULTIPLEX_METHOD()
+	RIGVM_METHOD()
 	virtual void Execute(const FRigUnitContext& Context) override;
 };
 
@@ -255,7 +255,7 @@ struct FRigUnit_MathVectorRound : public FRigUnit_MathVectorUnaryOp
 {
 	GENERATED_BODY()
 
-	MULTIPLEX_METHOD()
+	RIGVM_METHOD()
 	virtual void Execute(const FRigUnitContext& Context) override;
 };
 
@@ -267,7 +267,7 @@ struct FRigUnit_MathVectorSign : public FRigUnit_MathVectorUnaryOp
 {
 	GENERATED_BODY()
 
-	MULTIPLEX_METHOD()
+	RIGVM_METHOD()
 	virtual void Execute(const FRigUnitContext& Context) override;
 };
 
@@ -285,7 +285,7 @@ struct FRigUnit_MathVectorClamp : public FRigUnit_MathVectorBase
 		Maximum = FVector::OneVector;
 	}
 
-	MULTIPLEX_METHOD()
+	RIGVM_METHOD()
 	virtual void Execute(const FRigUnitContext& Context) override;
 
 	UPROPERTY(meta=(Input))
@@ -316,7 +316,7 @@ struct FRigUnit_MathVectorLerp : public FRigUnit_MathVectorBase
 		T = 0.f;
 	}
 
-	MULTIPLEX_METHOD()
+	RIGVM_METHOD()
 	virtual void Execute(const FRigUnitContext& Context) override;
 
 	UPROPERTY(meta=(Input))
@@ -347,7 +347,7 @@ struct FRigUnit_MathVectorRemap : public FRigUnit_MathVectorBase
 		bClamp = false;
 	}
 
-	MULTIPLEX_METHOD()
+	RIGVM_METHOD()
 	virtual void Execute(const FRigUnitContext& Context) override;
 
 	UPROPERTY(meta=(Input))
@@ -387,7 +387,7 @@ struct FRigUnit_MathVectorEquals : public FRigUnit_MathVectorBase
 		Result = true;
 	}
 
-	MULTIPLEX_METHOD()
+	RIGVM_METHOD()
 	virtual void Execute(const FRigUnitContext& Context) override;
 
 	UPROPERTY(meta=(Input))
@@ -414,7 +414,7 @@ struct FRigUnit_MathVectorNotEquals : public FRigUnit_MathVectorBase
 		Result = false;
 	}
 
-	MULTIPLEX_METHOD()
+	RIGVM_METHOD()
 	virtual void Execute(const FRigUnitContext& Context) override;
 
 	UPROPERTY(meta=(Input))
@@ -442,7 +442,7 @@ struct FRigUnit_MathVectorIsNearlyZero : public FRigUnit_MathVectorBase
 		Result = true;
 	}
 
-	MULTIPLEX_METHOD()
+	RIGVM_METHOD()
 	virtual void Execute(const FRigUnitContext& Context) override;
 
 	UPROPERTY(meta=(Input))
@@ -470,7 +470,7 @@ struct FRigUnit_MathVectorIsNearlyEqual : public FRigUnit_MathVectorBase
 		Result = true;
 	}
 
-	MULTIPLEX_METHOD()
+	RIGVM_METHOD()
 	virtual void Execute(const FRigUnitContext& Context) override;
 
 	UPROPERTY(meta=(Input))
@@ -500,7 +500,7 @@ struct FRigUnit_MathVectorSelectBool : public FRigUnit_MathVectorBase
 		IfTrue = IfFalse = Result = FVector::ZeroVector;
 	}
 
-	MULTIPLEX_METHOD()
+	RIGVM_METHOD()
 	virtual void Execute(const FRigUnitContext& Context) override;
 
 	UPROPERTY(meta=(Input))
@@ -524,7 +524,7 @@ struct FRigUnit_MathVectorDeg : public FRigUnit_MathVectorUnaryOp
 {
 	GENERATED_BODY()
 
-	MULTIPLEX_METHOD()
+	RIGVM_METHOD()
 	virtual void Execute(const FRigUnitContext& Context) override;
 };
 
@@ -536,7 +536,7 @@ struct FRigUnit_MathVectorRad : public FRigUnit_MathVectorUnaryOp
 {
 	GENERATED_BODY()
 
-	MULTIPLEX_METHOD()
+	RIGVM_METHOD()
 	virtual void Execute(const FRigUnitContext& Context) override;
 };
 
@@ -554,7 +554,7 @@ struct FRigUnit_MathVectorLengthSquared : public FRigUnit_MathVectorBase
 		Result = 0.f;
 	}
 
-	MULTIPLEX_METHOD()
+	RIGVM_METHOD()
 	virtual void Execute(const FRigUnitContext& Context) override;
 
 	UPROPERTY(meta=(Input))
@@ -578,7 +578,7 @@ struct FRigUnit_MathVectorLength : public FRigUnit_MathVectorBase
 		Result = 0.f;
 	}
 
-	MULTIPLEX_METHOD()
+	RIGVM_METHOD()
 	virtual void Execute(const FRigUnitContext& Context) override;
 
 	UPROPERTY(meta=(Input))
@@ -602,7 +602,7 @@ struct FRigUnit_MathVectorDistance : public FRigUnit_MathVectorBase
 		Result = 0.f;
 	}
 
-	MULTIPLEX_METHOD()
+	RIGVM_METHOD()
 	virtual void Execute(const FRigUnitContext& Context) override;
 
 	UPROPERTY(meta=(Input))
@@ -622,7 +622,7 @@ USTRUCT(meta=(DisplayName="Cross", PrototypeName="Cross", Keywords="^"))
 struct FRigUnit_MathVectorCross : public FRigUnit_MathVectorBinaryOp
 {
 	GENERATED_BODY()
-	MULTIPLEX_METHOD()
+	RIGVM_METHOD()
 	virtual void Execute(const FRigUnitContext& Context) override;
 };
 
@@ -640,7 +640,7 @@ struct FRigUnit_MathVectorDot : public FRigUnit_MathVectorBase
 		Result = 0.f;
 	}
 
-	MULTIPLEX_METHOD()
+	RIGVM_METHOD()
 	virtual void Execute(const FRigUnitContext& Context) override;
 
 	UPROPERTY(meta=(Input))
@@ -661,7 +661,7 @@ struct FRigUnit_MathVectorUnit : public FRigUnit_MathVectorUnaryOp
 {
 	GENERATED_BODY()
 
-	MULTIPLEX_METHOD()
+	RIGVM_METHOD()
 	virtual void Execute(const FRigUnitContext& Context) override;
 };
 
@@ -679,7 +679,7 @@ struct FRigUnit_MathVectorMirror : public FRigUnit_MathVectorBase
 		Normal = FVector(1.f, 0.f, 0.f);
 	}
 
-	MULTIPLEX_METHOD()
+	RIGVM_METHOD()
 	virtual void Execute(const FRigUnitContext& Context) override;
 
 	UPROPERTY(meta=(Input))
@@ -706,7 +706,7 @@ struct FRigUnit_MathVectorAngle : public FRigUnit_MathVectorBase
 		Result = 0.f;
 	}
 
-	MULTIPLEX_METHOD()
+	RIGVM_METHOD()
 	virtual void Execute(const FRigUnitContext& Context) override;
 
 	UPROPERTY(meta=(Input))
@@ -733,7 +733,7 @@ struct FRigUnit_MathVectorParallel : public FRigUnit_MathVectorBase
 		Result = true;
 	}
 
-	MULTIPLEX_METHOD()
+	RIGVM_METHOD()
 	virtual void Execute(const FRigUnitContext& Context) override;
 
 	UPROPERTY(meta=(Input))
@@ -761,7 +761,7 @@ struct FRigUnit_MathVectorOrthogonal : public FRigUnit_MathVectorBase
 		Result = true;
 	}
 
-	MULTIPLEX_METHOD()
+	RIGVM_METHOD()
 	virtual void Execute(const FRigUnitContext& Context) override;
 
 	UPROPERTY(meta=(Input))
@@ -789,7 +789,7 @@ struct FRigUnit_MathVectorBezierFourPoint : public FRigUnit_MathVectorBase
 		Result = Tangent = FVector::ZeroVector;
 	}
 
-	MULTIPLEX_METHOD()
+	RIGVM_METHOD()
 	virtual void Execute(const FRigUnitContext& Context) override;
 
 	UPROPERTY(meta=(Input))
diff --git a/Engine/Plugins/Experimental/ControlRig/Source/ControlRig/Private/Units/Math/RigUnit_Noise.cpp b/Engine/Plugins/Experimental/ControlRig/Source/ControlRig/Private/Units/Math/RigUnit_Noise.cpp
--- a/Engine/Plugins/Experimental/ControlRig/Source/ControlRig/Private/Units/Math/RigUnit_Noise.cpp
+++ b/Engine/Plugins/Experimental/ControlRig/Source/ControlRig/Private/Units/Math/RigUnit_Noise.cpp
@@ -3,7 +3,7 @@
 #include "Units/Math/RigUnit_Noise.h"
 #include "Units/RigUnitContext.h"
 
-UE_RigUnit_NoiseFloat_IMPLEMENT_MULTIPLEX(void, Execute, const FRigUnitContext& Context)
+UE_RigUnit_NoiseFloat_IMPLEMENT_RIGVM(void, Execute, const FRigUnitContext& Context)
 {
     DECLARE_SCOPE_HIERARCHICAL_COUNTER_RIGUNIT()
 	if (Context.State == EControlRigState::Init)
@@ -17,7 +17,7 @@ UE_RigUnit_NoiseFloat_IMPLEMENT_MULTIPLEX(void, Execute, const FRigUnitContext&
 	Time = Time + Speed * Context.DeltaTime;
 }
 
-UE_RigUnit_NoiseVector_IMPLEMENT_MULTIPLEX(void, Execute, const FRigUnitContext& Context)
+UE_RigUnit_NoiseVector_IMPLEMENT_RIGVM(void, Execute, const FRigUnitContext& Context)
 {
     DECLARE_SCOPE_HIERARCHICAL_COUNTER_RIGUNIT()
 	if (Context.State == EControlRigState::Init)
diff --git a/Engine/Plugins/Experimental/ControlRig/Source/ControlRig/Private/Units/Math/RigUnit_Noise.h b/Engine/Plugins/Experimental/ControlRig/Source/ControlRig/Private/Units/Math/RigUnit_Noise.h
--- a/Engine/Plugins/Experimental/ControlRig/Source/ControlRig/Private/Units/Math/RigUnit_Noise.h
+++ b/Engine/Plugins/Experimental/ControlRig/Source/ControlRig/Private/Units/Math/RigUnit_Noise.h
@@ -17,7 +17,7 @@ struct FRigUnit_NoiseFloat : public FRigUnit_MathBase
 		Frequency = Maximum = 1.f;
 	}
 
-	MULTIPLEX_METHOD()
+	RIGVM_METHOD()
 	virtual void Execute(const FRigUnitContext& Context) override;
 
 	UPROPERTY(meta = (Input))
@@ -56,7 +56,7 @@ struct FRigUnit_NoiseVector : public FRigUnit_MathBase
 		Maximum = 1.f;
 	}
 
-	MULTIPLEX_METHOD()
+	RIGVM_METHOD()
 	virtual void Execute(const FRigUnitContext& Context) override;
 
 	UPROPERTY(meta = (Input))
diff --git a/Engine/Plugins/Experimental/ControlRig/Source/ControlRig/Private/Units/Math/RigUnit_Random.cpp b/Engine/Plugins/Experimental/ControlRig/Source/ControlRig/Private/Units/Math/RigUnit_Random.cpp
--- a/Engine/Plugins/Experimental/ControlRig/Source/ControlRig/Private/Units/Math/RigUnit_Random.cpp
+++ b/Engine/Plugins/Experimental/ControlRig/Source/ControlRig/Private/Units/Math/RigUnit_Random.cpp
@@ -15,7 +15,7 @@ float FRigUnit_Random_Helper(int32& Seed)
 	return FPlatformMath::Fractional(Result.f);
 }
 
-UE_RigUnit_RandomFloat_IMPLEMENT_MULTIPLEX(void, Execute, const FRigUnitContext& Context)
+UE_RigUnit_RandomFloat_IMPLEMENT_RIGVM(void, Execute, const FRigUnitContext& Context)
 {
     DECLARE_SCOPE_HIERARCHICAL_COUNTER_RIGUNIT()
 	if (Context.State == EControlRigState::Init)
@@ -38,7 +38,7 @@ UE_RigUnit_RandomFloat_IMPLEMENT_MULTIPLEX(void, Execute, const FRigUnitContext&
 	LastResult = Result;
 }
 
-UE_RigUnit_RandomVector_IMPLEMENT_MULTIPLEX(void, Execute, const FRigUnitContext& Context)
+UE_RigUnit_RandomVector_IMPLEMENT_RIGVM(void, Execute, const FRigUnitContext& Context)
 {
     DECLARE_SCOPE_HIERARCHICAL_COUNTER_RIGUNIT()
 	if (Context.State == EControlRigState::Init)
diff --git a/Engine/Plugins/Experimental/ControlRig/Source/ControlRig/Private/Units/Math/RigUnit_Random.h b/Engine/Plugins/Experimental/ControlRig/Source/ControlRig/Private/Units/Math/RigUnit_Random.h
--- a/Engine/Plugins/Experimental/ControlRig/Source/ControlRig/Private/Units/Math/RigUnit_Random.h
+++ b/Engine/Plugins/Experimental/ControlRig/Source/ControlRig/Private/Units/Math/RigUnit_Random.h
@@ -19,7 +19,7 @@ struct FRigUnit_RandomFloat : public FRigUnit_MathBase
 		TimeLeft = 0.f;
 	}
 
-	MULTIPLEX_METHOD()
+	RIGVM_METHOD()
 	virtual void Execute(const FRigUnitContext& Context) override;
 
 	UPROPERTY(meta = (Input, Constant))
@@ -63,7 +63,7 @@ struct FRigUnit_RandomVector: public FRigUnit_MathBase
 		Result = LastResult = FVector::ZeroVector;
 	}
 
-	MULTIPLEX_METHOD()
+	RIGVM_METHOD()
 	virtual void Execute(const FRigUnitContext& Context) override;
 
 	UPROPERTY(meta = (Input))
diff --git a/Engine/Plugins/Experimental/ControlRig/Source/ControlRig/Private/Units/Simulation/RigUnit_Accumulate.cpp b/Engine/Plugins/Experimental/ControlRig/Source/ControlRig/Private/Units/Simulation/RigUnit_Accumulate.cpp
--- a/Engine/Plugins/Experimental/ControlRig/Source/ControlRig/Private/Units/Simulation/RigUnit_Accumulate.cpp
+++ b/Engine/Plugins/Experimental/ControlRig/Source/ControlRig/Private/Units/Simulation/RigUnit_Accumulate.cpp
@@ -3,7 +3,7 @@
 #include "Units/Simulation/RigUnit_Accumulate.h"
 #include "Units/RigUnitContext.h"
 
-UE_RigUnit_AccumulateFloatAdd_IMPLEMENT_MULTIPLEX(void, Execute, const FRigUnitContext& Context)
+UE_RigUnit_AccumulateFloatAdd_IMPLEMENT_RIGVM(void, Execute, const FRigUnitContext& Context)
 {
     DECLARE_SCOPE_HIERARCHICAL_COUNTER_RIGUNIT()
 	if (Context.State == EControlRigState::Init)
@@ -20,7 +20,7 @@ UE_RigUnit_AccumulateFloatAdd_IMPLEMENT_MULTIPLEX(void, Execute, const FRigUnitC
 	}
 }
 
-UE_RigUnit_AccumulateVectorAdd_IMPLEMENT_MULTIPLEX(void, Execute, const FRigUnitContext& Context)
+UE_RigUnit_AccumulateVectorAdd_IMPLEMENT_RIGVM(void, Execute, const FRigUnitContext& Context)
 {
     DECLARE_SCOPE_HIERARCHICAL_COUNTER_RIGUNIT()
 	if (Context.State == EControlRigState::Init)
@@ -37,7 +37,7 @@ UE_RigUnit_AccumulateVectorAdd_IMPLEMENT_MULTIPLEX(void, Execute, const FRigUnit
 	}
 }
 
-UE_RigUnit_AccumulateFloatMul_IMPLEMENT_MULTIPLEX(void, Execute, const FRigUnitContext& Context)
+UE_RigUnit_AccumulateFloatMul_IMPLEMENT_RIGVM(void, Execute, const FRigUnitContext& Context)
 {
     DECLARE_SCOPE_HIERARCHICAL_COUNTER_RIGUNIT()
 	if (Context.State == EControlRigState::Init)
@@ -55,7 +55,7 @@ UE_RigUnit_AccumulateFloatMul_IMPLEMENT_MULTIPLEX(void, Execute, const FRigUnitC
 	}
 }
 
-UE_RigUnit_AccumulateVectorMul_IMPLEMENT_MULTIPLEX(void, Execute, const FRigUnitContext& Context)
+UE_RigUnit_AccumulateVectorMul_IMPLEMENT_RIGVM(void, Execute, const FRigUnitContext& Context)
 {
     DECLARE_SCOPE_HIERARCHICAL_COUNTER_RIGUNIT()
 	if (Context.State == EControlRigState::Init)
@@ -73,7 +73,7 @@ UE_RigUnit_AccumulateVectorMul_IMPLEMENT_MULTIPLEX(void, Execute, const FRigUnit
 	}
 }
 
-UE_RigUnit_AccumulateQuatMul_IMPLEMENT_MULTIPLEX(void, Execute, const FRigUnitContext& Context)
+UE_RigUnit_AccumulateQuatMul_IMPLEMENT_RIGVM(void, Execute, const FRigUnitContext& Context)
 {
     DECLARE_SCOPE_HIERARCHICAL_COUNTER_RIGUNIT()
 	if (Context.State == EControlRigState::Init)
@@ -98,7 +98,7 @@ UE_RigUnit_AccumulateQuatMul_IMPLEMENT_MULTIPLEX(void, Execute, const FRigUnitCo
 	}
 }
 
-UE_RigUnit_AccumulateTransformMul_IMPLEMENT_MULTIPLEX(void, Execute, const FRigUnitContext& Context)
+UE_RigUnit_AccumulateTransformMul_IMPLEMENT_RIGVM(void, Execute, const FRigUnitContext& Context)
 {
     DECLARE_SCOPE_HIERARCHICAL_COUNTER_RIGUNIT()
 	if (Context.State == EControlRigState::Init)
@@ -125,7 +125,7 @@ UE_RigUnit_AccumulateTransformMul_IMPLEMENT_MULTIPLEX(void, Execute, const FRigU
 	}
 }
 
-UE_RigUnit_AccumulateFloatLerp_IMPLEMENT_MULTIPLEX(void, Execute, const FRigUnitContext& Context)
+UE_RigUnit_AccumulateFloatLerp_IMPLEMENT_RIGVM(void, Execute, const FRigUnitContext& Context)
 {
     DECLARE_SCOPE_HIERARCHICAL_COUNTER_RIGUNIT()
 	if (Context.State == EControlRigState::Init)
@@ -138,7 +138,7 @@ UE_RigUnit_AccumulateFloatLerp_IMPLEMENT_MULTIPLEX(void, Execute, const FRigUnit
 	}
 }
 
-UE_RigUnit_AccumulateVectorLerp_IMPLEMENT_MULTIPLEX(void, Execute, const FRigUnitContext& Context)
+UE_RigUnit_AccumulateVectorLerp_IMPLEMENT_RIGVM(void, Execute, const FRigUnitContext& Context)
 {
     DECLARE_SCOPE_HIERARCHICAL_COUNTER_RIGUNIT()
 	if (Context.State == EControlRigState::Init)
@@ -151,7 +151,7 @@ UE_RigUnit_AccumulateVectorLerp_IMPLEMENT_MULTIPLEX(void, Execute, const FRigUni
 	}
 }
 
-UE_RigUnit_AccumulateQuatLerp_IMPLEMENT_MULTIPLEX(void, Execute, const FRigUnitContext& Context)
+UE_RigUnit_AccumulateQuatLerp_IMPLEMENT_RIGVM(void, Execute, const FRigUnitContext& Context)
 {
     DECLARE_SCOPE_HIERARCHICAL_COUNTER_RIGUNIT()
 	if (Context.State == EControlRigState::Init)
@@ -164,7 +164,7 @@ UE_RigUnit_AccumulateQuatLerp_IMPLEMENT_MULTIPLEX(void, Execute, const FRigUnitC
 	}
 }
 
-UE_RigUnit_AccumulateTransformLerp_IMPLEMENT_MULTIPLEX(void, Execute, const FRigUnitContext& Context)
+UE_RigUnit_AccumulateTransformLerp_IMPLEMENT_RIGVM(void, Execute, const FRigUnitContext& Context)
 {
     DECLARE_SCOPE_HIERARCHICAL_COUNTER_RIGUNIT()
 	if (Context.State == EControlRigState::Init)
@@ -181,7 +181,7 @@ UE_RigUnit_AccumulateTransformLerp_IMPLEMENT_MULTIPLEX(void, Execute, const FRig
 	}
 }
 
-UE_RigUnit_AccumulateFloatRange_IMPLEMENT_MULTIPLEX(void, Execute, const FRigUnitContext& Context)
+UE_RigUnit_AccumulateFloatRange_IMPLEMENT_RIGVM(void, Execute, const FRigUnitContext& Context)
 {
     DECLARE_SCOPE_HIERARCHICAL_COUNTER_RIGUNIT()
 	if (Context.State == EControlRigState::Init)
@@ -195,7 +195,7 @@ UE_RigUnit_AccumulateFloatRange_IMPLEMENT_MULTIPLEX(void, Execute, const FRigUni
 	}
 }
 
-UE_RigUnit_AccumulateVectorRange_IMPLEMENT_MULTIPLEX(void, Execute, const FRigUnitContext& Context)
+UE_RigUnit_AccumulateVectorRange_IMPLEMENT_RIGVM(void, Execute, const FRigUnitContext& Context)
 {
     DECLARE_SCOPE_HIERARCHICAL_COUNTER_RIGUNIT()
 	if (Context.State == EControlRigState::Init)
diff --git a/Engine/Plugins/Experimental/ControlRig/Source/ControlRig/Private/Units/Simulation/RigUnit_Accumulate.h b/Engine/Plugins/Experimental/ControlRig/Source/ControlRig/Private/Units/Simulation/RigUnit_Accumulate.h
--- a/Engine/Plugins/Experimental/ControlRig/Source/ControlRig/Private/Units/Simulation/RigUnit_Accumulate.h
+++ b/Engine/Plugins/Experimental/ControlRig/Source/ControlRig/Private/Units/Simulation/RigUnit_Accumulate.h
@@ -19,7 +19,7 @@ struct FRigUnit_AccumulateFloatAdd : public FRigUnit_SimBase
 		bIntegrateDeltaTime = false;
 	}
 
-	MULTIPLEX_METHOD()
+	RIGVM_METHOD()
 	virtual void Execute(const FRigUnitContext& Context) override;
 
 	UPROPERTY(meta=(Input))
@@ -52,7 +52,7 @@ struct FRigUnit_AccumulateVectorAdd : public FRigUnit_SimBase
 		bIntegrateDeltaTime = false;
 	}
 
-	MULTIPLEX_METHOD()
+	RIGVM_METHOD()
 	virtual void Execute(const FRigUnitContext& Context) override;
 
 	UPROPERTY(meta = (Input))
@@ -85,7 +85,7 @@ struct FRigUnit_AccumulateFloatMul : public FRigUnit_SimBase
 		bIntegrateDeltaTime = false;
 	}
 
-	MULTIPLEX_METHOD()
+	RIGVM_METHOD()
 	virtual void Execute(const FRigUnitContext& Context) override;
 
 	UPROPERTY(meta=(Input))
@@ -118,7 +118,7 @@ struct FRigUnit_AccumulateVectorMul : public FRigUnit_SimBase
 		bIntegrateDeltaTime = false;
 	}
 
-	MULTIPLEX_METHOD()
+	RIGVM_METHOD()
 	virtual void Execute(const FRigUnitContext& Context) override;
 
 	UPROPERTY(meta = (Input))
@@ -151,7 +151,7 @@ struct FRigUnit_AccumulateQuatMul : public FRigUnit_SimBase
 		bIntegrateDeltaTime = bFlipOrder = false;
 	}
 
-	MULTIPLEX_METHOD()
+	RIGVM_METHOD()
 	virtual void Execute(const FRigUnitContext& Context) override;
 
 	UPROPERTY(meta=(Input))
@@ -187,7 +187,7 @@ struct FRigUnit_AccumulateTransformMul : public FRigUnit_SimBase
 		bIntegrateDeltaTime = bFlipOrder = false;
 	}
 
-	MULTIPLEX_METHOD()
+	RIGVM_METHOD()
 	virtual void Execute(const FRigUnitContext& Context) override;
 
 	UPROPERTY(meta = (Input))
@@ -223,7 +223,7 @@ struct FRigUnit_AccumulateFloatLerp : public FRigUnit_SimBase
 		bIntegrateDeltaTime = false;
 	}
 
-	MULTIPLEX_METHOD()
+	RIGVM_METHOD()
 	virtual void Execute(const FRigUnitContext& Context) override;
 
 	UPROPERTY(meta=(Input))
@@ -260,7 +260,7 @@ struct FRigUnit_AccumulateVectorLerp : public FRigUnit_SimBase
 		bIntegrateDeltaTime = false;
 	}
 
-	MULTIPLEX_METHOD()
+	RIGVM_METHOD()
 	virtual void Execute(const FRigUnitContext& Context) override;
 
 	UPROPERTY(meta = (Input))
@@ -297,7 +297,7 @@ struct FRigUnit_AccumulateQuatLerp : public FRigUnit_SimBase
 		bIntegrateDeltaTime = false;
 	}
 
-	MULTIPLEX_METHOD()
+	RIGVM_METHOD()
 	virtual void Execute(const FRigUnitContext& Context) override;
 
 	UPROPERTY(meta=(Input))
@@ -334,7 +334,7 @@ struct FRigUnit_AccumulateTransformLerp : public FRigUnit_SimBase
 		bIntegrateDeltaTime = false;
 	}
 
-	MULTIPLEX_METHOD()
+	RIGVM_METHOD()
 	virtual void Execute(const FRigUnitContext& Context) override;
 
 	UPROPERTY(meta = (Input))
@@ -369,7 +369,7 @@ struct FRigUnit_AccumulateFloatRange : public FRigUnit_SimBase
 		Value = Minimum = Maximum = AccumulatedMinimum = AccumulatedMaximum = 0.f;
 	}
 
-	MULTIPLEX_METHOD()
+	RIGVM_METHOD()
 	virtual void Execute(const FRigUnitContext& Context) override;
 
 	UPROPERTY(meta = (Input))
@@ -401,7 +401,7 @@ struct FRigUnit_AccumulateVectorRange : public FRigUnit_SimBase
 		Value = Minimum = Maximum = AccumulatedMinimum = AccumulatedMaximum = FVector::ZeroVector;
 	}
 
-	MULTIPLEX_METHOD()
+	RIGVM_METHOD()
 	virtual void Execute(const FRigUnitContext& Context) override;
 
 	UPROPERTY(meta = (Input))
diff --git a/Engine/Plugins/Experimental/ControlRig/Source/ControlRig/Private/Units/Simulation/RigUnit_DeltaFromPrevious.cpp b/Engine/Plugins/Experimental/ControlRig/Source/ControlRig/Private/Units/Simulation/RigUnit_DeltaFromPrevious.cpp
--- a/Engine/Plugins/Experimental/ControlRig/Source/ControlRig/Private/Units/Simulation/RigUnit_DeltaFromPrevious.cpp
+++ b/Engine/Plugins/Experimental/ControlRig/Source/ControlRig/Private/Units/Simulation/RigUnit_DeltaFromPrevious.cpp
@@ -3,7 +3,7 @@
 #include "Units/Simulation/RigUnit_DeltaFromPrevious.h"
 #include "Units/RigUnitContext.h"
 
-UE_RigUnit_DeltaFromPreviousFloat_IMPLEMENT_MULTIPLEX(void, Execute, const FRigUnitContext& Context)
+UE_RigUnit_DeltaFromPreviousFloat_IMPLEMENT_RIGVM(void, Execute, const FRigUnitContext& Context)
 {
     DECLARE_SCOPE_HIERARCHICAL_COUNTER_RIGUNIT()
 	if (Context.State == EControlRigState::Init)
@@ -18,7 +18,7 @@ UE_RigUnit_DeltaFromPreviousFloat_IMPLEMENT_MULTIPLEX(void, Execute, const FRigU
 	Cache = Value;
 }
 
-UE_RigUnit_DeltaFromPreviousVector_IMPLEMENT_MULTIPLEX(void, Execute, const FRigUnitContext& Context)
+UE_RigUnit_DeltaFromPreviousVector_IMPLEMENT_RIGVM(void, Execute, const FRigUnitContext& Context)
 {
     DECLARE_SCOPE_HIERARCHICAL_COUNTER_RIGUNIT()
 	if (Context.State == EControlRigState::Init)
@@ -33,7 +33,7 @@ UE_RigUnit_DeltaFromPreviousVector_IMPLEMENT_MULTIPLEX(void, Execute, const FRig
 	Cache = Value;
 }
 
-UE_RigUnit_DeltaFromPreviousQuat_IMPLEMENT_MULTIPLEX(void, Execute, const FRigUnitContext& Context)
+UE_RigUnit_DeltaFromPreviousQuat_IMPLEMENT_RIGVM(void, Execute, const FRigUnitContext& Context)
 {
     DECLARE_SCOPE_HIERARCHICAL_COUNTER_RIGUNIT()
 	if (Context.State == EControlRigState::Init)
@@ -48,7 +48,7 @@ UE_RigUnit_DeltaFromPreviousQuat_IMPLEMENT_MULTIPLEX(void, Execute, const FRigUn
 	Cache = Value;
 }
 
-UE_RigUnit_DeltaFromPreviousTransform_IMPLEMENT_MULTIPLEX(void, Execute, const FRigUnitContext& Context)
+UE_RigUnit_DeltaFromPreviousTransform_IMPLEMENT_RIGVM(void, Execute, const FRigUnitContext& Context)
 {
     DECLARE_SCOPE_HIERARCHICAL_COUNTER_RIGUNIT()
 	if (Context.State == EControlRigState::Init)
diff --git a/Engine/Plugins/Experimental/ControlRig/Source/ControlRig/Private/Units/Simulation/RigUnit_DeltaFromPrevious.h b/Engine/Plugins/Experimental/ControlRig/Source/ControlRig/Private/Units/Simulation/RigUnit_DeltaFromPrevious.h
--- a/Engine/Plugins/Experimental/ControlRig/Source/ControlRig/Private/Units/Simulation/RigUnit_DeltaFromPrevious.h
+++ b/Engine/Plugins/Experimental/ControlRig/Source/ControlRig/Private/Units/Simulation/RigUnit_DeltaFromPrevious.h
@@ -18,7 +18,7 @@ struct FRigUnit_DeltaFromPreviousFloat : public FRigUnit_SimBase
 		Value = Delta = PreviousValue = Cache = 0.f;
 	}
 
-	MULTIPLEX_METHOD()
+	RIGVM_METHOD()
 	virtual void Execute(const FRigUnitContext& Context) override;
 
 	UPROPERTY(meta = (Input))
@@ -47,7 +47,7 @@ struct FRigUnit_DeltaFromPreviousVector : public FRigUnit_SimBase
 		Value = Delta = PreviousValue = Cache = FVector::ZeroVector;
 	}
 
-	MULTIPLEX_METHOD()
+	RIGVM_METHOD()
 	virtual void Execute(const FRigUnitContext& Context) override;
 
 	UPROPERTY(meta = (Input))
@@ -76,7 +76,7 @@ struct FRigUnit_DeltaFromPreviousQuat : public FRigUnit_SimBase
 		Value = Delta = PreviousValue = Cache = FQuat::Identity;
 	}
 
-	MULTIPLEX_METHOD()
+	RIGVM_METHOD()
 	virtual void Execute(const FRigUnitContext& Context) override;
 
 	UPROPERTY(meta = (Input))
@@ -105,7 +105,7 @@ struct FRigUnit_DeltaFromPreviousTransform : public FRigUnit_SimBase
 		Value = Delta = PreviousValue = Cache = FTransform::Identity;
 	}
 
-	MULTIPLEX_METHOD()
+	RIGVM_METHOD()
 	virtual void Execute(const FRigUnitContext& Context) override;
 
 	UPROPERTY(meta = (Input))
diff --git a/Engine/Plugins/Experimental/ControlRig/Source/ControlRig/Private/Units/Simulation/RigUnit_Kalman.cpp b/Engine/Plugins/Experimental/ControlRig/Source/ControlRig/Private/Units/Simulation/RigUnit_Kalman.cpp
--- a/Engine/Plugins/Experimental/ControlRig/Source/ControlRig/Private/Units/Simulation/RigUnit_Kalman.cpp
+++ b/Engine/Plugins/Experimental/ControlRig/Source/ControlRig/Private/Units/Simulation/RigUnit_Kalman.cpp
@@ -3,7 +3,7 @@
 #include "Units/Simulation/RigUnit_Kalman.h"
 #include "Units/RigUnitContext.h"
 
-UE_RigUnit_KalmanFloat_IMPLEMENT_MULTIPLEX(void, Execute, const FRigUnitContext& Context)
+UE_RigUnit_KalmanFloat_IMPLEMENT_RIGVM(void, Execute, const FRigUnitContext& Context)
 {
     DECLARE_SCOPE_HIERARCHICAL_COUNTER_RIGUNIT()
 	if(BufferSize <= 0)
@@ -52,7 +52,7 @@ UE_RigUnit_KalmanFloat_IMPLEMENT_MULTIPLEX(void, Execute, const FRigUnitContext&
 	}
 }
 
-UE_RigUnit_KalmanVector_IMPLEMENT_MULTIPLEX(void, Execute, const FRigUnitContext& Context)
+UE_RigUnit_KalmanVector_IMPLEMENT_RIGVM(void, Execute, const FRigUnitContext& Context)
 {
     DECLARE_SCOPE_HIERARCHICAL_COUNTER_RIGUNIT()
 	if(BufferSize <= 0)
@@ -101,7 +101,7 @@ UE_RigUnit_KalmanVector_IMPLEMENT_MULTIPLEX(void, Execute, const FRigUnitContext
 	}
 }
 
-UE_RigUnit_KalmanTransform_IMPLEMENT_MULTIPLEX(void, Execute, const FRigUnitContext& Context)
+UE_RigUnit_KalmanTransform_IMPLEMENT_RIGVM(void, Execute, const FRigUnitContext& Context)
 {
     DECLARE_SCOPE_HIERARCHICAL_COUNTER_RIGUNIT()
 	if(BufferSize <= 0)
diff --git a/Engine/Plugins/Experimental/ControlRig/Source/ControlRig/Private/Units/Simulation/RigUnit_Kalman.h b/Engine/Plugins/Experimental/ControlRig/Source/ControlRig/Private/Units/Simulation/RigUnit_Kalman.h
--- a/Engine/Plugins/Experimental/ControlRig/Source/ControlRig/Private/Units/Simulation/RigUnit_Kalman.h
+++ b/Engine/Plugins/Experimental/ControlRig/Source/ControlRig/Private/Units/Simulation/RigUnit_Kalman.h
@@ -21,7 +21,7 @@ struct FRigUnit_KalmanFloat : public FRigUnit_SimBase
 		LastInsertIndex = 0;
 	}
 
-	MULTIPLEX_METHOD()
+	RIGVM_METHOD()
 	virtual void Execute(const FRigUnitContext& Context) override;
 
 	UPROPERTY(meta = (Input))
@@ -56,7 +56,7 @@ struct FRigUnit_KalmanVector : public FRigUnit_SimBase
 		LastInsertIndex = 0;
 	}
 
-	MULTIPLEX_METHOD()
+	RIGVM_METHOD()
 	virtual void Execute(const FRigUnitContext& Context) override;
 
 	UPROPERTY(meta = (Input))
@@ -91,7 +91,7 @@ struct FRigUnit_KalmanTransform : public FRigUnit_SimBase
 		LastInsertIndex = 0;
 	}
 
-	MULTIPLEX_METHOD()
+	RIGVM_METHOD()
 	virtual void Execute(const FRigUnitContext& Context) override;
 
 	UPROPERTY(meta = (Input))
diff --git a/Engine/Plugins/Experimental/ControlRig/Source/ControlRig/Private/Units/Simulation/RigUnit_PointSimulation.cpp b/Engine/Plugins/Experimental/ControlRig/Source/ControlRig/Private/Units/Simulation/RigUnit_PointSimulation.cpp
--- a/Engine/Plugins/Experimental/ControlRig/Source/ControlRig/Private/Units/Simulation/RigUnit_PointSimulation.cpp
+++ b/Engine/Plugins/Experimental/ControlRig/Source/ControlRig/Private/Units/Simulation/RigUnit_PointSimulation.cpp
@@ -3,7 +3,7 @@
 #include "Units/Simulation/RigUnit_PointSimulation.h"
 #include "Units/RigUnitContext.h"
 
-UE_RigUnit_PointSimulation_IMPLEMENT_MULTIPLEX(void, Execute, const FRigUnitContext& Context)
+UE_RigUnit_PointSimulation_IMPLEMENT_RIGVM(void, Execute, const FRigUnitContext& Context)
 {
     DECLARE_SCOPE_HIERARCHICAL_COUNTER_RIGUNIT()
 	FRigBoneHierarchy* Hierarchy = ExecuteContext.GetBones();
diff --git a/Engine/Plugins/Experimental/ControlRig/Source/ControlRig/Private/Units/Simulation/RigUnit_PointSimulation.h b/Engine/Plugins/Experimental/ControlRig/Source/ControlRig/Private/Units/Simulation/RigUnit_PointSimulation.h
--- a/Engine/Plugins/Experimental/ControlRig/Source/ControlRig/Private/Units/Simulation/RigUnit_PointSimulation.h
+++ b/Engine/Plugins/Experimental/ControlRig/Source/ControlRig/Private/Units/Simulation/RigUnit_PointSimulation.h
@@ -135,7 +135,7 @@ struct FRigUnit_PointSimulation : public FRigUnit_SimBaseMutable
 		Bezier = FCRFourPointBezier();
 	}
 
-	MULTIPLEX_METHOD()
+	RIGVM_METHOD()
 	virtual void Execute(const FRigUnitContext& Context) override;
 
 	/** The points to simulate */
diff --git a/Engine/Plugins/Experimental/ControlRig/Source/ControlRig/Private/Units/Simulation/RigUnit_TimeOffset.cpp b/Engine/Plugins/Experimental/ControlRig/Source/ControlRig/Private/Units/Simulation/RigUnit_TimeOffset.cpp
--- a/Engine/Plugins/Experimental/ControlRig/Source/ControlRig/Private/Units/Simulation/RigUnit_TimeOffset.cpp
+++ b/Engine/Plugins/Experimental/ControlRig/Source/ControlRig/Private/Units/Simulation/RigUnit_TimeOffset.cpp
@@ -3,7 +3,7 @@
 #include "Units/Simulation/RigUnit_TimeOffset.h"
 #include "Units/RigUnitContext.h"
 
-UE_RigUnit_TimeOffsetFloat_IMPLEMENT_MULTIPLEX(void, Execute, const FRigUnitContext& Context)
+UE_RigUnit_TimeOffsetFloat_IMPLEMENT_RIGVM(void, Execute, const FRigUnitContext& Context)
 {
     DECLARE_SCOPE_HIERARCHICAL_COUNTER_RIGUNIT()
 	Result = Value;
@@ -108,7 +108,7 @@ UE_RigUnit_TimeOffsetFloat_IMPLEMENT_MULTIPLEX(void, Execute, const FRigUnitCont
 	}
 }
 
-UE_RigUnit_TimeOffsetVector_IMPLEMENT_MULTIPLEX(void, Execute, const FRigUnitContext& Context)
+UE_RigUnit_TimeOffsetVector_IMPLEMENT_RIGVM(void, Execute, const FRigUnitContext& Context)
 {
     DECLARE_SCOPE_HIERARCHICAL_COUNTER_RIGUNIT()
 	Result = Value;
@@ -213,7 +213,7 @@ UE_RigUnit_TimeOffsetVector_IMPLEMENT_MULTIPLEX(void, Execute, const FRigUnitCon
 	}
 }
 
-UE_RigUnit_TimeOffsetTransform_IMPLEMENT_MULTIPLEX(void, Execute, const FRigUnitContext& Context)
+UE_RigUnit_TimeOffsetTransform_IMPLEMENT_RIGVM(void, Execute, const FRigUnitContext& Context)
 {
     DECLARE_SCOPE_HIERARCHICAL_COUNTER_RIGUNIT()
 	Result = Value;
diff --git a/Engine/Plugins/Experimental/ControlRig/Source/ControlRig/Private/Units/Simulation/RigUnit_TimeOffset.h b/Engine/Plugins/Experimental/ControlRig/Source/ControlRig/Private/Units/Simulation/RigUnit_TimeOffset.h
--- a/Engine/Plugins/Experimental/ControlRig/Source/ControlRig/Private/Units/Simulation/RigUnit_TimeOffset.h
+++ b/Engine/Plugins/Experimental/ControlRig/Source/ControlRig/Private/Units/Simulation/RigUnit_TimeOffset.h
@@ -23,7 +23,7 @@ struct FRigUnit_TimeOffsetFloat : public FRigUnit_SimBase
 		UpperBound = 0;
 	}
 
-	MULTIPLEX_METHOD()
+	RIGVM_METHOD()
 	virtual void Execute(const FRigUnitContext& Context) override;
 	
 	/** The value to record */
@@ -76,7 +76,7 @@ struct FRigUnit_TimeOffsetVector : public FRigUnit_SimBase
 		UpperBound = 0;
 	}
 
-	MULTIPLEX_METHOD()
+	RIGVM_METHOD()
 	virtual void Execute(const FRigUnitContext& Context) override;
 	
 	/** The value to record */
@@ -129,7 +129,7 @@ struct FRigUnit_TimeOffsetTransform : public FRigUnit_SimBase
 		UpperBound = 0;
 	}
 
-	MULTIPLEX_METHOD()
+	RIGVM_METHOD()
 	virtual void Execute(const FRigUnitContext& Context) override;
 	
 	/** The value to record */
diff --git a/Engine/Plugins/Experimental/ControlRig/Source/ControlRig/Private/Units/Simulation/RigUnit_Timeline.cpp b/Engine/Plugins/Experimental/ControlRig/Source/ControlRig/Private/Units/Simulation/RigUnit_Timeline.cpp
--- a/Engine/Plugins/Experimental/ControlRig/Source/ControlRig/Private/Units/Simulation/RigUnit_Timeline.cpp
+++ b/Engine/Plugins/Experimental/ControlRig/Source/ControlRig/Private/Units/Simulation/RigUnit_Timeline.cpp
@@ -3,7 +3,7 @@
 #include "Units/Simulation/RigUnit_Timeline.h"
 #include "Units/RigUnitContext.h"
 
-UE_RigUnit_Timeline_IMPLEMENT_MULTIPLEX(void, Execute, const FRigUnitContext& Context)
+UE_RigUnit_Timeline_IMPLEMENT_RIGVM(void, Execute, const FRigUnitContext& Context)
 {
     DECLARE_SCOPE_HIERARCHICAL_COUNTER_RIGUNIT()
 	if (Context.State == EControlRigState::Init)
diff --git a/Engine/Plugins/Experimental/ControlRig/Source/ControlRig/Private/Units/Simulation/RigUnit_Timeline.h b/Engine/Plugins/Experimental/ControlRig/Source/ControlRig/Private/Units/Simulation/RigUnit_Timeline.h
--- a/Engine/Plugins/Experimental/ControlRig/Source/ControlRig/Private/Units/Simulation/RigUnit_Timeline.h
+++ b/Engine/Plugins/Experimental/ControlRig/Source/ControlRig/Private/Units/Simulation/RigUnit_Timeline.h
@@ -19,7 +19,7 @@ struct FRigUnit_Timeline : public FRigUnit_SimBase
 		Time = AccumulatedValue = 0.f;
 	}
 
-	MULTIPLEX_METHOD()
+	RIGVM_METHOD()
 	virtual void Execute(const FRigUnitContext& Context) override;
 
 	UPROPERTY(meta = (Input))
diff --git a/Engine/Plugins/Experimental/ControlRig/Source/ControlRig/Private/Units/Simulation/RigUnit_Verlet.cpp b/Engine/Plugins/Experimental/ControlRig/Source/ControlRig/Private/Units/Simulation/RigUnit_Verlet.cpp
--- a/Engine/Plugins/Experimental/ControlRig/Source/ControlRig/Private/Units/Simulation/RigUnit_Verlet.cpp
+++ b/Engine/Plugins/Experimental/ControlRig/Source/ControlRig/Private/Units/Simulation/RigUnit_Verlet.cpp
@@ -3,7 +3,7 @@
 #include "Units/Simulation/RigUnit_Verlet.h"
 #include "Units/RigUnitContext.h"
 
-UE_RigUnit_VerletIntegrateVector_IMPLEMENT_MULTIPLEX(void, Execute, const FRigUnitContext& Context)
+UE_RigUnit_VerletIntegrateVector_IMPLEMENT_RIGVM(void, Execute, const FRigUnitContext& Context)
 {
     DECLARE_SCOPE_HIERARCHICAL_COUNTER_RIGUNIT()
 	if (Context.State == EControlRigState::Init)
diff --git a/Engine/Plugins/Experimental/ControlRig/Source/ControlRig/Private/Units/Simulation/RigUnit_Verlet.h b/Engine/Plugins/Experimental/ControlRig/Source/ControlRig/Private/Units/Simulation/RigUnit_Verlet.h
--- a/Engine/Plugins/Experimental/ControlRig/Source/ControlRig/Private/Units/Simulation/RigUnit_Verlet.h
+++ b/Engine/Plugins/Experimental/ControlRig/Source/ControlRig/Private/Units/Simulation/RigUnit_Verlet.h
@@ -23,7 +23,7 @@ struct FRigUnit_VerletIntegrateVector : public FRigUnit_SimBase
 		Point = FCRSimPoint();
 	}
 
-	MULTIPLEX_METHOD()
+	RIGVM_METHOD()
 	virtual void Execute(const FRigUnitContext& Context) override;
 
 	UPROPERTY(meta = (Input))
diff --git a/Engine/Plugins/Experimental/ControlRig/Source/ControlRig/Public/Units/Control/RigUnit_Control.h b/Engine/Plugins/Experimental/ControlRig/Source/ControlRig/Public/Units/Control/RigUnit_Control.h
--- a/Engine/Plugins/Experimental/ControlRig/Source/ControlRig/Public/Units/Control/RigUnit_Control.h
+++ b/Engine/Plugins/Experimental/ControlRig/Source/ControlRig/Public/Units/Control/RigUnit_Control.h
@@ -21,7 +21,7 @@ struct CONTROLRIG_API FRigUnit_Control : public FRigUnit
 	{
 	}
 
-	MULTIPLEX_METHOD()
+	RIGVM_METHOD()
 	virtual void Execute(const FRigUnitContext& Context) override;
 
 	/** Combine Transform and Base to make the resultant transform */
diff --git a/Engine/Plugins/Experimental/ControlRig/Source/ControlRig/Public/Units/RigUnit.h b/Engine/Plugins/Experimental/ControlRig/Source/ControlRig/Public/Units/RigUnit.h
--- a/Engine/Plugins/Experimental/ControlRig/Source/ControlRig/Public/Units/RigUnit.h
+++ b/Engine/Plugins/Experimental/ControlRig/Source/ControlRig/Public/Units/RigUnit.h
@@ -3,7 +3,7 @@
 #pragma once
 
 #include "ControlRigDefines.h"
-#include "MultiplexRegistry.h"
+#include "RigVMRegistry.h"
 #include "RigUnit.generated.h"
 
 struct FRigUnitContext;
diff --git a/Engine/Source/Programs/UnrealHeaderTool/Private/CodeGenerator.cpp b/Engine/Source/Programs/UnrealHeaderTool/Private/CodeGenerator.cpp
--- a/Engine/Source/Programs/UnrealHeaderTool/Private/CodeGenerator.cpp
+++ b/Engine/Source/Programs/UnrealHeaderTool/Private/CodeGenerator.cpp
@@ -3299,13 +3299,13 @@ void FNativeClassHeaderGenerator::ExportGeneratedStructBodyMacros(FOutputDevice&
 
 	const FString SingletonName = GetSingletonName(Struct);
 
-	FString MultiplexMacroPrefix(StructNameCPP);
-	MultiplexMacroPrefix = FString::Printf(TEXT("UE_%s"), *MultiplexMacroPrefix.Mid(1));
+	FString RigVMMacroPrefix(StructNameCPP);
+	RigVMMacroPrefix = FString::Printf(TEXT("UE_%s"), *RigVMMacroPrefix.Mid(1));
 
 	const TCHAR* Comma = TEXT(", ");
 	TArray<FString> MemberNames, MemberDeclarations, MemberCasts, MemberCastProlog, MemberMulticastProlog;
-	const TArray<FMultiplexMethodInfo>* StructMultiplexMethods = FHeaderParser::StructMultiplexMethods.Find(Struct);
-	if(StructMultiplexMethods)
+	const TArray<FRigVMMethodInfo>* StructRigVMMethods = FHeaderParser::StructRigVMMethods.Find(Struct);
+	if(StructRigVMMethods)
 	{
 		const TCHAR* FText = TEXT("F");
 		const TCHAR* TText = TEXT("T");
@@ -3324,9 +3324,9 @@ void FNativeClassHeaderGenerator::ExportGeneratedStructBodyMacros(FOutputDevice&
 		{
 			FString Name = Prop->GetName();
 			FString AddressName = FString::Printf(TEXT("%s_Address"), *Name);
-			MemberMulticastProlog.Add(FString::Printf(TEXT("const FMultiplexArgument& %s = MultiplexAddresses[%d];"), *AddressName, MemberIndex++));
+			MemberMulticastProlog.Add(FString::Printf(TEXT("const FRigVMArgument& %s = RigVMAddresses[%d];"), *AddressName, MemberIndex++));
 		}
-		MemberMulticastProlog.Insert(FString::Printf(TEXT("ensure(MultiplexAddresses.Num() == %d);"), MemberMulticastProlog.Num()), 0);
+		MemberMulticastProlog.Insert(FString::Printf(TEXT("ensure(RigVMAddresses.Num() == %d);"), MemberMulticastProlog.Num()), 0);
 		MemberMulticastProlog.Add(FString());
 
 		for (UProperty* Prop : TFieldRange<UProperty>(Struct))
@@ -3400,11 +3400,11 @@ void FNativeClassHeaderGenerator::ExportGeneratedStructBodyMacros(FOutputDevice&
 
 			MemberDeclarations.Add(FString::Printf(TEXT("%s %s"), *ParameterCPPType, *Name));
 			FString AddressName = FString::Printf(TEXT("%s_Address"), *Name);
-			MemberMulticastProlog.Add(FString::Printf(TEXT("%s %s = MultiplexStorage[%s.StorageType()]->%s<%s>(%s.Index());"), *ExtractedMulticastCPPType, *Name, *AddressName, *MulticastGetter, *MultiCastCPPType, *AddressName));
+			MemberMulticastProlog.Add(FString::Printf(TEXT("%s %s = RigVMStorage[%s.StorageType()]->%s<%s>(%s.Index());"), *ExtractedMulticastCPPType, *Name, *AddressName, *MulticastGetter, *MultiCastCPPType, *AddressName));
 		}
 
 		OutGeneratedHeaderText.Log(TEXT("\n"));
-		OutGeneratedHeaderText.Logf(TEXT("#define %s_IMPLEMENT_MULTIPLEX(ReturnType, FunctionName, ...) \\\n"), *MultiplexMacroPrefix);
+		OutGeneratedHeaderText.Logf(TEXT("#define %s_IMPLEMENT_RIGVM(ReturnType, FunctionName, ...) \\\n"), *RigVMMacroPrefix);
 		OutGeneratedHeaderText.Logf(TEXT("\tReturnType %s::Static##FunctionName(\\\r\n\t\t%s, ##__VA_ARGS__)\n"), StructNameCPP, *FString::Join(MemberDeclarations, TEXT(",\\\r\n\t\t")));
 		OutGeneratedHeaderText.Log(TEXT("\n"));
 	}
@@ -3424,20 +3424,20 @@ void FNativeClassHeaderGenerator::ExportGeneratedStructBodyMacros(FOutputDevice&
 		const FString StaticClassLine = FString::Printf(TEXT("\t%sstatic class UScriptStruct* StaticStruct();\r\n"), *RequiredAPI);
 		const FString PrivatePropertiesOffset = PrivatePropertiesOffsetGetters(Struct, StructNameCPP);
 		
-		FString MultiplexMethodsDeclarations;
-		if (StructMultiplexMethods)
+		FString RigVMMethodsDeclarations;
+		if (StructRigVMMethods)
 		{
-			for (const FMultiplexMethodInfo& Info : (*StructMultiplexMethods))
+			for (const FRigVMMethodInfo& Info : (*StructRigVMMethods))
 			{
 				FString ParamsSuffix = Info.Params.IsEmpty() ? TEXT("") : FString::Printf(TEXT(", %s"), *Info.Params);
-				MultiplexMethodsDeclarations += FString::Printf(TEXT("\tstatic %s Static%s(\r\n\t\t%s%s);\r\n"), *Info.ReturnType, *Info.Name, *FString::Join(MemberDeclarations, TEXT(",\r\n\t\t")), *ParamsSuffix);
-				MultiplexMethodsDeclarations += FString::Printf(TEXT("\tstatic %s Multiplex%s(const TArrayView<FMultiplexArgument>& MultiplexAddresses, FMultiplexStorage** MultiplexStorage, const TArrayView<void*>& AdditionalArgs);\r\n"), *Info.ReturnType, *Info.Name);
+				RigVMMethodsDeclarations += FString::Printf(TEXT("\tstatic %s Static%s(\r\n\t\t%s%s);\r\n"), *Info.ReturnType, *Info.Name, *FString::Join(MemberDeclarations, TEXT(",\r\n\t\t")), *ParamsSuffix);
+				RigVMMethodsDeclarations += FString::Printf(TEXT("\tstatic %s RigVM%s(const TArrayView<FRigVMArgument>& RigVMAddresses, FRigVMStorage** RigVMStorage, const TArrayView<void*>& AdditionalArgs);\r\n"), *Info.ReturnType, *Info.Name);
 			}
 		}
 
 		const FString SuperTypedef = BaseStruct ? FString::Printf(TEXT("\ttypedef %s Super;\r\n"), NameLookupCPP.GetNameCPP(BaseStruct)) : FString();
 
-		const FString CombinedLine = FriendLine + StaticClassLine + MultiplexMethodsDeclarations + PrivatePropertiesOffset + SuperTypedef;
+		const FString CombinedLine = FriendLine + StaticClassLine + RigVMMethodsDeclarations + PrivatePropertiesOffset + SuperTypedef;
 		const FString MacroName = SourceFile.GetGeneratedBodyMacroName(Struct->StructMacroDeclaredLineNumber);
 
 		const FString Macroized = Macroize(*MacroName, *CombinedLine);
@@ -3469,11 +3469,11 @@ void FNativeClassHeaderGenerator::ExportGeneratedStructBodyMacros(FOutputDevice&
 		Out.Logf(TEXT("\t\tSingleton = GetStaticStruct(%s, %s, TEXT(\"%s\"), sizeof(%s), %s());\r\n"),
 			*SingletonName.LeftChop(2), *OuterName, *ActualStructName, StructNameCPP, *GetHashName);
 
-		if (StructMultiplexMethods)
+		if (StructRigVMMethods)
 		{
-			for (const FMultiplexMethodInfo& Info : (*StructMultiplexMethods))
+			for (const FRigVMMethodInfo& Info : (*StructRigVMMethods))
 			{
-				Out.Logf(TEXT("\t\tFMultiplexRegistry::Get().Register(TEXT(\"%s::%s\"), &%s::Multiplex%s);\r\n"),
+				Out.Logf(TEXT("\t\tFRigVMRegistry::Get().Register(TEXT(\"%s::%s\"), &%s::RigVM%s);\r\n"),
 					StructNameCPP, *Info.Name, StructNameCPP, *Info.Name);
 			}
 		}
@@ -3650,9 +3650,9 @@ void FNativeClassHeaderGenerator::ExportGeneratedStructBodyMacros(FOutputDevice&
 	Out.Log(GeneratedStructRegisterFunctionText);
 	Out.Logf(TEXT("\tuint32 %s() { return %uU; }\r\n"), *HashFuncName, StructHash);
 
-	if (StructMultiplexMethods)
+	if (StructRigVMMethods)
 	{
-		for (const FMultiplexMethodInfo& Info : (*StructMultiplexMethods))
+		for (const FRigVMMethodInfo& Info : (*StructRigVMMethods))
 		{
 			Out.Log(TEXT("\r\n"));
 
@@ -3731,7 +3731,7 @@ void FNativeClassHeaderGenerator::ExportGeneratedStructBodyMacros(FOutputDevice&
 				ParamDeclarationSuffix = FString::Printf(TEXT(", %s"), *Info.Params);
 			}
 
-			Out.Logf(TEXT("%s %s::Multiplex%s(const TArrayView<FMultiplexArgument>& MultiplexAddresses, FMultiplexStorage** MultiplexStorage, const TArrayView<void*>& AdditionalArgs)\r\n"), *Info.ReturnType, StructNameCPP, *Info.Name);
+			Out.Logf(TEXT("%s %s::RigVM%s(const TArrayView<FRigVMArgument>& RigVMAddresses, FRigVMStorage** RigVMStorage, const TArrayView<void*>& AdditionalArgs)\r\n"), *Info.ReturnType, StructNameCPP, *Info.Name);
 			Out.Log(TEXT("{\r\n"));
 
 			if (ParameterExtractionProlog.Num() > 0)
diff --git a/Engine/Source/Programs/UnrealHeaderTool/Private/HeaderParser.cpp b/Engine/Source/Programs/UnrealHeaderTool/Private/HeaderParser.cpp
--- a/Engine/Source/Programs/UnrealHeaderTool/Private/HeaderParser.cpp
+++ b/Engine/Source/Programs/UnrealHeaderTool/Private/HeaderParser.cpp
@@ -72,7 +72,7 @@ static const FName NAME_DocumentationPolicy(TEXT("DocumentationPolicy"));
 EGeneratedCodeVersion FHeaderParser::DefaultGeneratedCodeVersion = EGeneratedCodeVersion::V1;
 TArray<FString> FHeaderParser::StructsWithNoPrefix;
 TArray<FString> FHeaderParser::StructsWithTPrefix;
-TMap<UStruct*, TArray<FMultiplexMethodInfo>> FHeaderParser::StructMultiplexMethods;
+TMap<UStruct*, TArray<FRigVMMethodInfo>> FHeaderParser::StructRigVMMethods;
 TArray<FString> FHeaderParser::DelegateParameterCountStrings;
 TMap<FString, FString> FHeaderParser::TypeRedirectMap;
 TArray<FString> FHeaderParser::PropertyCPPTypesRequiringUIRanges = { TEXT("float"), TEXT("double") };
@@ -2495,9 +2495,9 @@ UScriptStruct* FHeaderParser::CompileStructDeclaration(FClasses& AllClasses)
 		{
 			FError::Throwf(TEXT("USTRUCTs cannot contain UFUNCTIONs."));
 		}
-		else if (Token.Matches(TEXT("MULTIPLEX_METHOD"), ESearchCase::CaseSensitive))
+		else if (Token.Matches(TEXT("RIGVM_METHOD"), ESearchCase::CaseSensitive))
 		{
-			CompileMultiplexMethodDeclaration(AllClasses, Struct);
+			CompileRigVMMethodDeclaration(AllClasses, Struct);
 		}
 		else if (Token.Matches(TEXT("GENERATED_USTRUCT_BODY")) || Token.Matches(TEXT("GENERATED_BODY")))
 		{
@@ -3008,8 +3008,8 @@ void FHeaderParser::FixupDelegateProperties( FClasses& AllClasses, UStruct* Stru
 	CheckDocumentationPolicyForStruct(Struct, MetaData);
 
 	// check if the struct has multi plex methods, and if so perform some additional checks
-	TArray<FMultiplexMethodInfo>* MultiplexMethods = StructMultiplexMethods.Find(Struct);
-	if (MultiplexMethods)
+	TArray<FRigVMMethodInfo>* RigVMMethods = StructRigVMMethods.Find(Struct);
+	if (RigVMMethods)
 	{
 		const TCHAR* InputText = TEXT("Input");
 		const TCHAR* OutputText = TEXT("Output");
@@ -3031,7 +3031,7 @@ void FHeaderParser::FixupDelegateProperties( FClasses& AllClasses, UStruct* Stru
 				// we only support arrays - no maps or similar data structures
 				if (MemberCPPType != TArrayText)
 				{
-					UE_LOG_ERROR_UHT(TEXT("Multiplex Struct '%s' - Member '%s' type '%s' not supported by Multiplex."), *Struct->GetName(), *PropName, *MemberCPPType);
+					UE_LOG_ERROR_UHT(TEXT("RigVM Struct '%s' - Member '%s' type '%s' not supported by RigVM."), *Struct->GetName(), *PropName, *MemberCPPType);
 				}
 				else
 				{
@@ -3042,7 +3042,7 @@ void FHeaderParser::FixupDelegateProperties( FClasses& AllClasses, UStruct* Stru
 
 					if (!bIsConstant && !Prop->HasMetaData(MaxArraySizeText))
 					{
-						UE_LOG_ERROR_UHT(TEXT("Multiplex Struct '%s' - Member '%s' requires the 'MaxArraySize' meta tag."), *Struct->GetName(), *PropName);
+						UE_LOG_ERROR_UHT(TEXT("RigVM Struct '%s' - Member '%s' requires the 'MaxArraySize' meta tag."), *Struct->GetName(), *PropName);
 					}
 				}
 			}
@@ -6207,11 +6207,11 @@ void FHeaderParser::CompileInterfaceDeclaration(FClasses& AllClasses)
 	PushNest(ENestType::Interface, InterfaceClass);
 }
 
-void FHeaderParser::CompileMultiplexMethodDeclaration(FClasses& AllClasses, UStruct* Struct)
+void FHeaderParser::CompileRigVMMethodDeclaration(FClasses& AllClasses, UStruct* Struct)
 {
 	if (!MatchSymbol(TEXT("(")))
 	{
-		FError::Throwf(TEXT("Bad MULTIPLEX_METHOD definition"));
+		FError::Throwf(TEXT("Bad RIGVM_METHOD definition"));
 	}
 
 	// find the next close brace
@@ -6249,7 +6249,7 @@ void FHeaderParser::CompileMultiplexMethodDeclaration(FClasses& AllClasses, UStr
 
 	if (!MatchSymbol(TEXT("(")))
 	{
-		FError::Throwf(TEXT("Bad MULTIPLEX_METHOD definition"));
+		FError::Throwf(TEXT("Bad RIGVM_METHOD definition"));
 	}
 
 	TArray<FString> ParamsContent;
@@ -6271,12 +6271,12 @@ void FHeaderParser::CompileMultiplexMethodDeclaration(FClasses& AllClasses, UStr
 		}
 	}
 
-	FMultiplexMethodInfo Info;
+	FRigVMMethodInfo Info;
 	Info.ReturnType = ReturnTypeToken.Identifier;
 	Info.Name = NameToken.Identifier;
 	Info.Params = FString::Join(ParamsContent, TEXT(" "));
 
-	TArray<FMultiplexMethodInfo>& Infos = StructMultiplexMethods.FindOrAdd(Struct);
+	TArray<FRigVMMethodInfo>& Infos = StructRigVMMethods.FindOrAdd(Struct);
 	Infos.Add(Info);
 }
 
diff --git a/Engine/Source/Programs/UnrealHeaderTool/Private/HeaderParser.h b/Engine/Source/Programs/UnrealHeaderTool/Private/HeaderParser.h
--- a/Engine/Source/Programs/UnrealHeaderTool/Private/HeaderParser.h
+++ b/Engine/Source/Programs/UnrealHeaderTool/Private/HeaderParser.h
@@ -117,7 +117,7 @@ struct FIndexRange
 	int32 Count;
 };
 
-struct FMultiplexMethodInfo
+struct FRigVMMethodInfo
 {
 	FString ReturnType;
 	FString Name;
@@ -427,7 +427,7 @@ protected:
 	TMap<int32, FString> RPCsNeedingHookup;
 
 	// List of all multiplex methods defined on structs
-	static TMap<UStruct*, TArray<FMultiplexMethodInfo>> StructMultiplexMethods;
+	static TMap<UStruct*, TArray<FRigVMMethodInfo>> StructRigVMMethods;
 
 	// Constructor.
 	explicit FHeaderParser(FFeedbackContext* InWarn, const FManifestModule& InModule);
@@ -572,7 +572,7 @@ protected:
 	void CompileFunctionDeclaration(FClasses& AllClasses);
 	void CompileVariableDeclaration (FClasses& AllClasses, UStruct* Struct);
 	void CompileInterfaceDeclaration(FClasses& AllClasses);
-	void CompileMultiplexMethodDeclaration(FClasses& AllClasses, UStruct* Struct);
+	void CompileRigVMMethodDeclaration(FClasses& AllClasses, UStruct* Struct);
 
 	FClass* ParseInterfaceNameDeclaration(FClasses& AllClasses, FString& DeclaredInterfaceName, FString& RequiredAPIMacroIfPresent);
 	bool TryParseIInterfaceClass(FClasses& AllClasses);
diff --git a/Engine/Source/Programs/UnrealHeaderTool/Resources/UHTDebugging/StaticMethodStructHeader.h b/Engine/Source/Programs/UnrealHeaderTool/Resources/UHTDebugging/StaticMethodStructHeader.h
--- a/Engine/Source/Programs/UnrealHeaderTool/Resources/UHTDebugging/StaticMethodStructHeader.h
+++ b/Engine/Source/Programs/UnrealHeaderTool/Resources/UHTDebugging/StaticMethodStructHeader.h
@@ -21,13 +21,13 @@ struct FStaticMethodStruct : public FStaticMethodStructBase
 {
 	GENERATED_BODY()
 
-	MULTIPLEX_METHOD()
+	RIGVM_METHOD()
 	void Clear();
 
-	MULTIPLEX_METHOD()
+	RIGVM_METHOD()
 	virtual void Execute(bool bAdditionalFlag = false, const FString& InString = TEXT("")) override;
 
-	MULTIPLEX_METHOD()
+	RIGVM_METHOD()
 	void Compute(float TestFloat);
 
 	UPROPERTY(meta = (Input))
diff --git a/Engine/Source/Runtime/CoreUObject/Public/UObject/ObjectMacros.h b/Engine/Source/Runtime/CoreUObject/Public/UObject/ObjectMacros.h
--- a/Engine/Source/Runtime/CoreUObject/Public/UObject/ObjectMacros.h
+++ b/Engine/Source/Runtime/CoreUObject/Public/UObject/ObjectMacros.h
@@ -581,7 +581,7 @@ struct COREUOBJECT_API FReferencerInformationList
 #define UPARAM(...)
 #define UENUM(...)
 #define UDELEGATE(...)
-#define MULTIPLEX_METHOD(...)
+#define RIGVM_METHOD(...)
 
 // This pair of macros is used to help implement GENERATED_BODY() and GENERATED_USTRUCT_BODY()
 #define BODY_MACRO_COMBINE_INNER(A,B,C,D) A##B##C##D
diff --git a/Engine/Source/Runtime/RigVM/Private/MultiplexByteCode.cpp b/Engine/Source/Runtime/RigVM/Private/MultiplexByteCode.cpp
deleted file mode 100644
diff --git a/Engine/Source/Runtime/RigVM/Private/MultiplexRegistry.cpp b/Engine/Source/Runtime/RigVM/Private/MultiplexRegistry.cpp
deleted file mode 100644
diff --git a/Engine/Source/Runtime/RigVM/Private/MultiplexStorage.cpp b/Engine/Source/Runtime/RigVM/Private/MultiplexStorage.cpp
deleted file mode 100644
diff --git a/Engine/Source/Runtime/RigVM/Private/MultiplexVM.cpp b/Engine/Source/Runtime/RigVM/Private/MultiplexVM.cpp
deleted file mode 100644
diff --git a/Engine/Source/Runtime/RigVM/Private/RigVM.cpp b/Engine/Source/Runtime/RigVM/Private/RigVM.cpp
new file mode 100644
--- /dev/null
+++ b/Engine/Source/Runtime/RigVM/Private/RigVM.cpp
@@ -0,0 +1,152 @@
+// Copyright 1998-2019 Epic Games, Inc. All Rights Reserved.
+
+#include "RigVM.h"
+
+void URigVM::Reset()
+{
+	Literals.Reset();
+	WorkState.Reset();
+	FunctionNames.Reset();
+	Functions.Reset();
+	ByteCode.Reset();
+	Instructions.Reset();
+}
+
+int32 URigVM::AddRigVMFunction(UScriptStruct* InRigVMStruct, const FName& InFunctionName)
+{
+	check(InRigVMStruct);
+	FString FunctionKey = FString::Printf(TEXT("F%s::%s"), *InRigVMStruct->GetName(), *InFunctionName.ToString());
+	int32 FunctionIndex = FunctionNames.Find(FunctionKey);
+	if (FunctionIndex != INDEX_NONE)
+	{
+		return FunctionIndex;
+	}
+
+	FRigVMFunctionPtr Function = FRigVMRegistry::Get().Find(*FunctionKey);
+	if (Function == nullptr)
+	{
+		return INDEX_NONE;
+	}
+
+	FunctionNames.Add(FunctionKey);
+	return Functions.Add(Function);
+}
+
+void URigVM::ResolveFunctionsIfRequired()
+{
+	if (Functions.Num() != FunctionNames.Num())
+	{
+		Functions.Reset();
+		Functions.SetNumZeroed(FunctionNames.Num());
+
+		for (int32 FunctionIndex = 0; FunctionIndex < FunctionNames.Num(); FunctionIndex)
+		{
+			Functions[FunctionIndex] = FRigVMRegistry::Get().Find(*FunctionNames[FunctionIndex]);
+		}
+	}
+}
+
+void URigVM::RefreshInstructionsIfRequired()
+{
+	if (Instructions.Num() == 0)
+	{
+		Instructions = ByteCode.GetTable();
+	}
+}
+
+bool URigVM::Execute(FRigVMStorage** Storage, TArrayView<void*> AdditionalArgs)
+{
+	ResolveFunctionsIfRequired();
+	RefreshInstructionsIfRequired();
+
+	if (Instructions.Num() == 0)
+	{
+		return true;
+	}
+
+	FRigVMStorage* LocalStorage[] = { &WorkState, &Literals }; 
+	if (Storage == nullptr)
+	{
+		Storage = LocalStorage;
+	}
+
+	int32 InstructionIndex = 0;
+	while (Instructions[InstructionIndex].OpCode != ERigVMOpCode::Exit)
+	{
+		switch (Instructions[InstructionIndex].OpCode)
+		{
+			case ERigVMOpCode::Copy:
+			{
+				const FRigVMCopyOp& Op = ByteCode.GetOpAt<FRigVMCopyOp>(Instructions[InstructionIndex]);
+				Storage[0]->Copy(Op.Source.Index(), Op.Target.Index(), Storage[Op.Source.StorageType()], Op.SourceOffset, Op.TargetOffset, Op.NumBytes);
+				InstructionIndex++;
+				break;
+			}
+			case ERigVMOpCode::Increment:
+			{
+				const FRigVMIncrementOp& Op = ByteCode.GetOpAt<FRigVMIncrementOp>(Instructions[InstructionIndex]);
+				Storage[0]->GetRef<int32>(Op.Arg.Index())++;
+				InstructionIndex++;
+				break;
+			}
+			case ERigVMOpCode::Decrement:
+			{
+				const FRigVMDecrementOp& Op = ByteCode.GetOpAt<FRigVMDecrementOp>(Instructions[InstructionIndex]);
+				Storage[0]->GetRef<int32>(Op.Arg.Index())--;
+				InstructionIndex++;
+				break;
+			}
+			case ERigVMOpCode::Equals:
+			{
+				const FRigVMEqualsOp& Op = ByteCode.GetOpAt<FRigVMEqualsOp>(Instructions[InstructionIndex]);
+				ensureMsgf(false, TEXT("TO BE IMPLEMENTED"));
+				break;
+			}
+			case ERigVMOpCode::NotEquals:
+			{
+				const FRigVMNotEqualsOp& Op = ByteCode.GetOpAt<FRigVMNotEqualsOp>(Instructions[InstructionIndex]);
+				ensureMsgf(false, TEXT("TO BE IMPLEMENTED"));
+				break;
+			}
+			case ERigVMOpCode::Jump:
+			{
+				const FRigVMJumpOp& Op = ByteCode.GetOpAt<FRigVMJumpOp>(Instructions[InstructionIndex]);
+				ensureMsgf(false, TEXT("TO BE IMPLEMENTED"));
+				break;
+			}
+			case ERigVMOpCode::JumpIfTrue:
+			{
+				const FRigVMJumpIfTrueOp& Op = ByteCode.GetOpAt<FRigVMJumpIfTrueOp>(Instructions[InstructionIndex]);
+				ensureMsgf(false, TEXT("TO BE IMPLEMENTED"));
+				break;
+			}
+			case ERigVMOpCode::JumpIfFalse:
+			{
+				const FRigVMJumpIfFalseOp& Op = ByteCode.GetOpAt<FRigVMJumpIfFalseOp>(Instructions[InstructionIndex]);
+				ensureMsgf(false, TEXT("TO BE IMPLEMENTED"));
+				break;
+			}
+			case ERigVMOpCode::Execute:
+			{
+				const FRigVMExecuteOp& Op = ByteCode.GetOpAt<FRigVMExecuteOp>(Instructions[InstructionIndex]);
+				TArrayView<FRigVMArgument> Args = ByteCode.GetArgumentsForExecuteOp(Instructions[InstructionIndex]);
+				(*Functions[Op.FunctionIndex])(Args, Storage, AdditionalArgs);
+				InstructionIndex++;
+				break;
+			}
+			case ERigVMOpCode::Exit:
+			case ERigVMOpCode::Invalid:
+			{
+				ensure(false);
+				return false;
+			}
+		}
+
+		if (!Instructions.IsValidIndex(InstructionIndex))
+		{
+			break;
+		}
+	}
+
+	return true;
+}
diff --git a/Engine/Source/Runtime/RigVM/Private/RigVMByteCode.cpp b/Engine/Source/Runtime/RigVM/Private/RigVMByteCode.cpp
new file mode 100644
--- /dev/null
+++ b/Engine/Source/Runtime/RigVM/Private/RigVMByteCode.cpp
@@ -0,0 +1,161 @@
+// Copyright 1998-2019 Epic Games, Inc. All Rights Reserved.
+
+#include "RigVMByteCode.h"
+
+FRigVMByteCodeTable::FRigVMByteCodeTable()
+{
+}
+
+FRigVMByteCodeTable::FRigVMByteCodeTable(const FRigVMByteCode& InByteCode)
+{
+	uint64 Address = 0;
+	while (Address < InByteCode.Num())
+	{
+		ERigVMOpCode OpCode = InByteCode.GetOpCodeAt(Address);
+		Entries.Add(FRigVMByteCodeTableEntry(OpCode, Address));
+		Address += InByteCode.GetOpNumBytesAt(Address);
+	}
+}
+
+void FRigVMByteCodeTable::Reset()
+{
+	Entries.Reset();
+}
+
+FRigVMByteCode::FRigVMByteCode()
+{
+}
+
+void FRigVMByteCode::Reset()
+{
+	ByteCode.Reset();
+}
+
+uint64 FRigVMByteCode::Num() const
+{
+	return (uint64)ByteCode.Num();
+}
+
+uint64 FRigVMByteCode::GetOpNumBytesAt(uint64 InByteCodeIndex, bool bIncludeArguments) const
+{
+	ERigVMOpCode OpCode = GetOpCodeAt(InByteCodeIndex);
+	switch (OpCode)
+	{
+		case ERigVMOpCode::Copy:
+		{
+			return (uint64)sizeof(FRigVMCopyOp);
+		}
+		case ERigVMOpCode::Increment:
+		{
+			return (uint64)sizeof(FRigVMIncrementOp);
+		}
+		case ERigVMOpCode::Decrement:
+		{
+			return (uint64)sizeof(FRigVMDecrementOp);
+		}
+		case ERigVMOpCode::Equals:
+		{
+			return (uint64)sizeof(FRigVMEqualsOp);
+		}
+		case ERigVMOpCode::NotEquals:
+		{
+			return (uint64)sizeof(FRigVMNotEqualsOp);
+		}
+		case ERigVMOpCode::Jump:
+		{
+			return (uint64)sizeof(FRigVMJumpOp);
+		}
+		case ERigVMOpCode::JumpIfTrue:
+		{
+			return (uint64)sizeof(FRigVMJumpIfTrueOp);
+		}
+		case ERigVMOpCode::JumpIfFalse:
+		{
+			return (uint64)sizeof(FRigVMJumpIfFalseOp);
+		}
+		case ERigVMOpCode::Execute:
+		{
+			uint64 NumBytes = (uint64)sizeof(FRigVMExecuteOp);
+			if(bIncludeArguments)
+			{
+				const FRigVMExecuteOp& ExecuteOp = GetOpAt<FRigVMExecuteOp>(InByteCodeIndex);
+				NumBytes += (uint64)ExecuteOp.ArgumentCount * (uint64)sizeof(FRigVMArgument);
+			}
+			return NumBytes;
+		}
+		case ERigVMOpCode::Exit:
+		{
+			return (uint64)sizeof(FRigVMExitOp);
+		}
+		case ERigVMOpCode::Invalid:
+		{
+			ensure(false);
+			return 0;
+		}
+	}
+	return 0;
+}
+
+uint64 FRigVMByteCode::AddCopyOp(const FRigVMArgument& InSource, const FRigVMArgument& InTarget, int32 InSourceOffset, int32 InTargetOffset, int32 InNumBytes)
+{
+	FRigVMCopyOp Op(InSource, InTarget, InSourceOffset, InTargetOffset, InNumBytes);
+	return AddOp(Op);
+}
+
+uint64 FRigVMByteCode::AddIncrementOp(const FRigVMArgument& InArg)
+{
+	FRigVMIncrementOp Op(InArg);
+	return AddOp(Op);
+}
+
+uint64 FRigVMByteCode::AddDecrementOp(const FRigVMArgument& InArg)
+{
+	FRigVMDecrementOp Op(InArg);
+	return AddOp(Op);
+}
+
+uint64 FRigVMByteCode::AddEqualsOp(const FRigVMArgument& InA, const FRigVMArgument& InB, const FRigVMArgument& InResult)
+{
+	FRigVMEqualsOp Op(InA, InB, InResult);
+	return AddOp(Op);
+}
+
+uint64 FRigVMByteCode::AddNotEqualsOp(const FRigVMArgument& InA, const FRigVMArgument& InB, const FRigVMArgument& InResult)
+{
+	FRigVMNotEqualsOp Op(InA, InB, InResult);
+	return AddOp(Op);
+}
+
+uint64 FRigVMByteCode::AddJumpOp(uint64 InByteCodeIndex)
+{
+	FRigVMJumpOp Op(InByteCodeIndex);
+	return AddOp(Op);
+}
+
+uint64 FRigVMByteCode::AddJumpIfTrueOp(uint64 InByteCodeIndex, const FRigVMArgument& InCondition)
+{
+	FRigVMJumpIfTrueOp Op(InByteCodeIndex, InCondition);
+	return AddOp(Op);
+}
+
+uint64 FRigVMByteCode::AddJumpIfFalseOp(uint64 InByteCodeIndex, const FRigVMArgument& InCondition)
+{
+	FRigVMJumpIfFalseOp Op(InByteCodeIndex, InCondition);
+	return AddOp(Op);
+}
+
+uint64 FRigVMByteCode::AddExecuteOp(uint16 InFunctionIndex, const TArrayView<FRigVMArgument>& InArguments)
+{
+	FRigVMExecuteOp Op(InFunctionIndex, (uint16)InArguments.Num());
+	uint64 OpAddress = AddOp(Op);
+	uint64 ArgumentsAddress = (uint64)ByteCode.AddUninitialized(sizeof(FRigVMArgument) * InArguments.Num());
+	FMemory::Memcpy(ByteCode.GetData() + ArgumentsAddress, InArguments.GetData(), sizeof(FRigVMArgument) * InArguments.Num());
+	return OpAddress;
+}
+
+uint64 FRigVMByteCode::AddExitOp()
+{
+	FRigVMExitOp Op;
+	return AddOp(Op);
+}
+
diff --git a/Engine/Source/Runtime/RigVM/Private/RigVMRegistry.cpp b/Engine/Source/Runtime/RigVM/Private/RigVMRegistry.cpp
new file mode 100644
--- /dev/null
+++ b/Engine/Source/Runtime/RigVM/Private/RigVMRegistry.cpp
@@ -0,0 +1,37 @@
+// Copyright 1998-2019 Epic Games, Inc. All Rights Reserved.
+
+#include "RigVMRegistry.h"
+#include "UObject/UObjectIterator.h"
+
+FRigVMRegistry FRigVMRegistry::s_RigVMRegistry;
+
+FRigVMRegistry& FRigVMRegistry::Get()
+{
+	return s_RigVMRegistry;
+}
+
+void FRigVMRegistry::Refresh()
+{
+}
+
+void FRigVMRegistry::Register(const TCHAR* InName, FRigVMFunctionPtr InFunctionPtr)
+{
+	if (Find(InName) != nullptr)
+	{
+		return;
+	}
+	Functions.Add(FRigVMFunction(InName, InFunctionPtr));
+}
+
+FRigVMFunctionPtr FRigVMRegistry::Find(const TCHAR* InName) const
+{
+	for (const FRigVMFunction& Function : Functions)
+	{
+		if (FCString::Strcmp(Function.Name, InName) == 0)
+		{
+			return Function.FunctionPtr;
+		}
+	}
+
+	return nullptr;
+}
diff --git a/Engine/Source/Runtime/RigVM/Private/RigVMStorage.cpp b/Engine/Source/Runtime/RigVM/Private/RigVMStorage.cpp
new file mode 100644
--- /dev/null
+++ b/Engine/Source/Runtime/RigVM/Private/RigVMStorage.cpp
@@ -0,0 +1,523 @@
+// Copyright 1998-2019 Epic Games, Inc. All Rights Reserved.
+
+#include "RigVMStorage.h"
+
+FRigVMStorage::FRigVMStorage(bool bInUseNames)
+	:bUseNameMap(bInUseNames)
+{
+}
+
+FRigVMStorage::FRigVMStorage(const FRigVMStorage& Other)
+{
+	*this = Other;
+}
+
+FRigVMStorage::~FRigVMStorage()
+{
+	Reset();
+}
+
+FRigVMStorage& FRigVMStorage::operator= (const FRigVMStorage &InOther)
+{
+	Reset();
+
+	bUseNameMap = InOther.bUseNameMap;
+	Data.Append(InOther.Data);
+	Addresses.Append(InOther.Addresses);
+	ScriptStructs.Append(InOther.ScriptStructs);
+
+	UpdateAddresses();
+
+	for (int32 Index = 0; Index < Addresses.Num(); Index++)
+	{
+		Construct(Index);
+		Copy(Index, Index, &InOther);
+	}
+
+	return *this;
+}
+
+void FRigVMStorage::Reset()
+{
+	for (int32 Index = 0; Index < Addresses.Num(); Index++)
+	{
+		Destroy(Index);
+	}
+
+	Data.Reset();
+	Addresses.Reset();
+	ScriptStructs.Reset();
+	NameMap.Reset();
+}
+
+bool FRigVMStorage::Copy(
+	int32 InSourceAddressIndex,
+	int32 InTargetAddressIndex,
+	const FRigVMStorage* InSourceStorage,
+	int32 InSourceByteOffset,
+	int32 InTargetByteOffset,
+	int32 InNumBytes)
+{
+	if (InSourceStorage == nullptr)
+	{
+		InSourceStorage = this;
+	}
+
+	ensure(InSourceStorage->Addresses.IsValidIndex(InSourceAddressIndex));
+	ensure(Addresses.IsValidIndex(InTargetAddressIndex));
+
+	if (InSourceAddressIndex == InTargetAddressIndex && InSourceByteOffset == InTargetByteOffset && this == InSourceStorage)
+	{
+		return false;
+	}
+
+	const FRigVMAddress& Source = InSourceStorage->Addresses[InSourceAddressIndex];
+	const FRigVMAddress& Target = Addresses[InTargetAddressIndex];
+
+	int32 SourceStartByte = Source.FirstByte();
+	int32 SourceNumBytes = Source.NumBytes(false);
+	if(InSourceByteOffset != INDEX_NONE)
+	{
+		ensure(InNumBytes != INDEX_NONE);
+		ensure(InSourceByteOffset >= 0 && InSourceByteOffset < Source.NumBytes(false));
+		ensure(InNumBytes > 0 && InSourceByteOffset + InNumBytes <= Source.NumBytes(false));
+
+		SourceStartByte += InSourceByteOffset;
+		SourceNumBytes = InNumBytes;
+	}
+
+	int32 TargetStartByte = Target.FirstByte();
+	int32 TargetNumBytes = Target.NumBytes(false);
+	if(InTargetByteOffset != INDEX_NONE)
+	{
+		ensure(InNumBytes != INDEX_NONE);
+		ensure(InTargetByteOffset >= 0 && InTargetByteOffset < Target.NumBytes(false));
+		ensure(InNumBytes > 0 && InTargetByteOffset + InNumBytes <= Target.NumBytes(false));
+
+		TargetStartByte += InTargetByteOffset;
+		TargetNumBytes = InNumBytes;
+	}
+
+	if (SourceNumBytes != TargetNumBytes)
+	{
+		return false;
+	}
+
+	switch (Target.Type)
+	{
+		case ERigVMAddressType::Plain:
+		{
+			FMemory::Memcpy(&Data[TargetStartByte], &InSourceStorage->Data[SourceStartByte], TargetNumBytes);
+			break;
+		}
+		case ERigVMAddressType::Struct:
+		{
+			UScriptStruct* ScriptStruct = GetScriptStruct(InTargetAddressIndex);
+			int32 NumStructs = TargetNumBytes / ScriptStruct->GetStructureSize();
+			ScriptStruct->CopyScriptStruct(&Data[TargetStartByte], &InSourceStorage->Data[SourceStartByte], NumStructs);
+			break;
+		}
+		case ERigVMAddressType::Name:
+		{
+			int32 NumNames = TargetNumBytes / sizeof(FName);
+			TArrayView<FName> TargetNames((FName*)&Data[TargetStartByte], NumNames);
+			TArrayView<FName> SourceNames((FName*)&InSourceStorage->Data[SourceStartByte], NumNames);
+			for (int32 Index = 0; Index < NumNames; Index++)
+			{
+				TargetNames[Index] = SourceNames[Index];
+			}
+			break;
+		}
+		case ERigVMAddressType::String:
+		{
+			int32 NumStrings = TargetNumBytes / sizeof(FString);
+			TArrayView<FString> TargetStrings((FString*)&Data[TargetStartByte], NumStrings);
+			TArrayView<FString> SourceStrings((FString*)&InSourceStorage->Data[SourceStartByte], NumStrings);
+			for (int32 Index = 0; Index < NumStrings; Index++)
+			{
+				TargetStrings[Index] = SourceStrings[Index];
+			}
+			break;
+		}
+		case ERigVMAddressType::Invalid:
+		{
+			return false;
+		}
+	}
+
+	return true;
+}
+
+bool FRigVMStorage::Copy(
+	const FName& InSourceName,
+	const FName& InTargetName,
+	const FRigVMStorage* InSourceStorage,
+	int32 InSourceByteOffset,
+	int32 InTargetByteOffset,
+	int32 InNumBytes)
+{
+	ensure(bUseNameMap);
+
+	int32 SourceAddressIndex = GetIndex(InSourceName);
+	int32 TargetAddressIndex = GetIndex(InTargetName);
+
+	if(SourceAddressIndex == INDEX_NONE || TargetAddressIndex == INDEX_NONE)
+	{
+		return false;
+	}
+
+	return Copy(SourceAddressIndex, TargetAddressIndex, InSourceStorage, InSourceByteOffset, InTargetByteOffset, InNumBytes);
+}
+
+int32 FRigVMStorage::Allocate(const FName& InNewName, int32 InElementSize, int32 InCount, const void* InDataPtr, bool bUpdateAddresses)
+{
+	FName Name = InNewName;
+	if (bUseNameMap && InNewName == NAME_None)
+	{
+		const TCHAR* AddressPrefix = TEXT("Address");
+		int32 AddressSuffix = 0;
+		do
+		{
+			Name = FName(*FString::Printf(TEXT("%s_%d"), AddressPrefix, AddressSuffix++));
+		} while (!IsNameAvailable(Name));
+	}
+
+	ensure(InElementSize > 0 && InCount > 0);
+
+	if (bUseNameMap)
+	{
+		if (!IsNameAvailable(Name))
+		{
+			return INDEX_NONE;
+		}
+	}
+
+	FRigVMAddress NewAddress;
+	NewAddress.ByteIndex = Data.Num();
+	if (bUseNameMap)
+	{
+		NewAddress.Name = Name;
+	}
+	NewAddress.ElementSize = InElementSize;
+	NewAddress.ElementCount = InCount;
+	NewAddress.Type = ERigVMAddressType::Plain;
+
+	Data.AddZeroed(NewAddress.NumBytes());
+
+	if (InDataPtr != nullptr)
+	{
+		FMemory::Memcpy(&Data[NewAddress.FirstByte()], InDataPtr, NewAddress.NumBytes());
+	}
+
+	int32 AddressIndex = Addresses.Num();
+	Addresses.Add(NewAddress);
+
+	if (bUpdateAddresses)
+	{
+		UpdateAddresses();
+	}
+	return AddressIndex;
+}
+
+int32 FRigVMStorage::Allocate(int32 InElementSize, int32 InCount, const void* InDataPtr, bool bUpdateAddresses)
+{
+	return Allocate(NAME_None, InElementSize, InCount, InDataPtr, bUpdateAddresses);
+}
+
+bool FRigVMStorage::Construct(int32 InAddressIndex, int32 InElementIndex)
+{
+	ensure(Addresses.IsValidIndex(InAddressIndex));
+
+	const FRigVMAddress& Address = Addresses[InAddressIndex];
+	switch (Address.Type)
+	{
+		case ERigVMAddressType::Struct:
+		{
+			void* DataPtr = (void*)&Data[InElementIndex == INDEX_NONE ? Address.FirstByte() : Address.FirstByte() + InElementIndex * Address.ElementSize];
+			int32 Count = InElementIndex == INDEX_NONE ? Address.ElementCount : 1;
+
+			UScriptStruct* ScriptStruct = GetScriptStruct(InAddressIndex);
+			ScriptStruct->InitializeStruct(DataPtr, Count);
+			break;
+		}
+		case ERigVMAddressType::String:
+		{
+			FString* DataPtr = (FString*)&Data[InElementIndex == INDEX_NONE ? Address.FirstByte() : Address.FirstByte() + InElementIndex * Address.ElementSize];
+			int32 Count = InElementIndex == INDEX_NONE ? Address.ElementCount : 1;
+
+			FMemory::Memzero(DataPtr, Count * Address.ElementSize);
+			for (int32 Index = 0; Index < Count; Index++)
+			{
+				DataPtr[Index] = FString();
+			}
+			break;
+		}
+		case ERigVMAddressType::Name:
+		{
+			FName* DataPtr = (FName*)&Data[InElementIndex == INDEX_NONE ? Address.FirstByte() : Address.FirstByte() + InElementIndex * Address.ElementSize];
+			int32 Count = InElementIndex == INDEX_NONE ? Address.ElementCount : 1;
+
+			FMemory::Memzero(DataPtr, Count * Address.ElementSize);
+			for (int32 Index = 0; Index < Count; Index++)
+			{
+				DataPtr[Index] = FName();
+			}
+			break;
+		}
+		default:
+		{
+			return false;
+		}
+	}
+
+	return true;
+}
+
+bool FRigVMStorage::Destroy(int32 InAddressIndex, int32 InElementIndex)
+{
+	ensure(Addresses.IsValidIndex(InAddressIndex));
+
+	const FRigVMAddress& Address = Addresses[InAddressIndex];
+	switch (Address.Type)
+	{
+		case ERigVMAddressType::Struct:
+		{
+			void* DataPtr = (void*)&Data[InElementIndex == INDEX_NONE ? Address.FirstByte() : Address.FirstByte() + InElementIndex * Address.ElementSize];
+			int32 Count = InElementIndex == INDEX_NONE ? Address.ElementCount : 1;
+
+			UScriptStruct* ScriptStruct = GetScriptStruct(InAddressIndex);
+			ScriptStruct->DestroyStruct(DataPtr, Count);
+			break;
+		}
+		case ERigVMAddressType::String:
+		{
+			FString* DataPtr = (FString*)&Data[InElementIndex == INDEX_NONE ? Address.FirstByte() : Address.FirstByte() + InElementIndex * Address.ElementSize];
+			int32 Count = InElementIndex == INDEX_NONE ? Address.ElementCount : 1;
+
+			for (int32 Index = 0; Index < Count; Index++)
+			{
+				DataPtr[Index].Empty();
+			}
+			break;
+		}
+		case ERigVMAddressType::Name:
+		{
+			FName* DataPtr = (FName*)&Data[InElementIndex == INDEX_NONE ? Address.FirstByte() : Address.FirstByte() + InElementIndex * Address.ElementSize];
+			int32 Count = InElementIndex == INDEX_NONE ? Address.ElementCount : 1;
+
+			for (int32 Index = 0; Index < Count; Index++)
+			{
+				DataPtr[Index] = FName();
+			}
+			break;
+		}
+		default:
+		{
+			return false;
+		}
+	}
+
+	return true;
+}
+
+bool FRigVMStorage::Remove(int32 InAddressIndex)
+{
+	if (InAddressIndex < 0 || InAddressIndex >= Addresses.Num())
+	{
+		return false;
+	}
+
+	Destroy(InAddressIndex);
+
+	FRigVMAddress AddressToRemove = Addresses[InAddressIndex];
+	Data.RemoveAt(AddressToRemove.ByteIndex, AddressToRemove.NumBytes());
+	Addresses.RemoveAt(InAddressIndex);
+
+	for (int32 Index = InAddressIndex; Index < Addresses.Num(); Index++)
+	{
+		Addresses[Index].ByteIndex -= AddressToRemove.NumBytes();
+	}
+
+	UpdateAddresses();
+	return true;
+}
+
+bool FRigVMStorage::Remove(const FName& InAddressName)
+{
+	ensure(bUseNameMap);
+	return Remove(GetIndex(InAddressName));
+}
+
+FName FRigVMStorage::Rename(int32 InAddressIndex, const FName& InNewName)
+{
+	if (Addresses[InAddressIndex].Name == InNewName)
+	{
+		return Addresses[InAddressIndex].Name;
+	}
+
+	if (!IsNameAvailable(InNewName))
+	{
+		return Addresses[InAddressIndex].Name;
+	}
+
+	Addresses[InAddressIndex].Name = InNewName;
+	UpdateAddresses();
+
+	return InNewName;
+}
+
+FName FRigVMStorage::Rename(const FName& InOldName, const FName& InNewName)
+{
+	ensure(bUseNameMap);
+
+	int32 AddressIndex = GetIndex(InOldName);
+	if (AddressIndex == INDEX_NONE)
+	{
+		return NAME_None;
+	}
+
+	return Rename(AddressIndex, InNewName);
+}
+
+bool FRigVMStorage::Resize(int32 InAddressIndex, int32 InNewElementCount)
+{
+	if (InNewElementCount <= 0)
+	{
+		return Remove(InAddressIndex);
+	}
+
+	if (Addresses[InAddressIndex].ElementCount == InNewElementCount)
+	{
+		return false;
+	}
+
+	FRigVMAddress& Address = Addresses[InAddressIndex];
+
+	if (Address.ElementCount > InNewElementCount) // shrink
+	{
+		int32 ElementsToRemove = Address.ElementCount - InNewElementCount;
+		int32 NumBytesToRemove = Address.ElementSize * ElementsToRemove;
+		int32 FirstByteToRemove = Address.FirstByte() + Address.ElementSize * InNewElementCount;
+
+		for (int32 ElementIndex = InNewElementCount; ElementIndex < Address.ElementCount; ElementIndex++)
+		{
+			Destroy(InAddressIndex, ElementIndex);
+		}
+
+		Data.RemoveAt(FirstByteToRemove, NumBytesToRemove);
+		Address.ElementCount = InNewElementCount;
+
+		for (int32 AddressIndex = InAddressIndex + 1; AddressIndex < Addresses.Num(); AddressIndex++)
+		{
+			Addresses[AddressIndex].ByteIndex -= NumBytesToRemove;
+		}
+	}
+	else // grow
+	{
+		int32 OldElementCount = Address.ElementCount;
+		int32 ElementsToAdd = InNewElementCount - Address.ElementCount;
+		int32 NumBytesToAdd = Address.ElementSize * ElementsToAdd;
+		int32 FirstByteToAdd = Address.FirstByte() + Address.ElementSize * Address.ElementCount;
+
+		Data.InsertZeroed(FirstByteToAdd, NumBytesToAdd);
+		Address.ElementCount = InNewElementCount;
+
+		for (int32 ElementIndex = OldElementCount; ElementIndex < InNewElementCount; ElementIndex++)
+		{
+			Construct(InAddressIndex, ElementIndex);
+		}
+
+		for (int32 AddressIndex = InAddressIndex + 1; AddressIndex < Addresses.Num(); AddressIndex++)
+		{
+			Addresses[AddressIndex].ByteIndex += NumBytesToAdd;
+		}
+	}
+
+	UpdateAddresses();
+	return true;
+}
+
+bool FRigVMStorage::Resize(const FName& InAddressName, int32 InNewElementCount)
+{
+	ensure(bUseNameMap);
+
+	int32 AddressIndex = GetIndex(InAddressName);
+	if (AddressIndex == INDEX_NONE)
+	{
+		return false;
+	}
+
+	return Resize(AddressIndex, InNewElementCount);
+}
+
+void FRigVMStorage::UpdateAddresses()
+{
+	int32 AlignmentShift = 0;
+	for (int32 AddressIndex = 0; AddressIndex < Addresses.Num(); AddressIndex++)
+	{
+		FRigVMAddress& Address = Addresses[AddressIndex];
+		Address.ByteIndex += AlignmentShift;
+
+		UScriptStruct* ScriptStruct = GetScriptStruct(AddressIndex);
+		if (ScriptStruct != nullptr)
+		{
+			UScriptStruct::ICppStructOps* TheCppStructOps = ScriptStruct->GetCppStructOps();
+			if (TheCppStructOps != NULL)
+			{
+				if (!TheCppStructOps->HasZeroConstructor())
+				{
+					void* Pointer = GetData(AddressIndex);
+
+					if (Address.AlignmentBytes > 0)
+					{
+						if (!IsAligned(Pointer, TheCppStructOps->GetAlignment()))
+						{
+							Data.RemoveAt(Address.ByteIndex, Address.AlignmentBytes);
+							AlignmentShift -= Address.AlignmentBytes;
+							Address.AlignmentBytes = 0;
+							Pointer = GetData(AddressIndex);
+						}
+					}
+
+					while (!IsAligned(Pointer, TheCppStructOps->GetAlignment()))
+					{
+						Data.InsertZeroed(Address.ByteIndex, 1);
+						Address.AlignmentBytes++;
+						AlignmentShift++;
+						Pointer = GetData(AddressIndex);
+					}
+				}
+			}
+		}
+	}
+
+	for (int32 AddressIndex = 0; AddressIndex < Addresses.Num(); AddressIndex++)
+	{
+		Addresses[AddressIndex].Pointer = GetData(AddressIndex);
+	}
+
+	if (bUseNameMap)
+	{
+		NameMap.Reset();
+		for (int32 Index = 0; Index < Addresses.Num(); Index++)
+		{
+			NameMap.Add(Addresses[Index].Name, Index);
+		}
+	}
+}
+
+void FRigVMStorage::FillWithZeroes(int32 InAddressIndex)
+{
+	ensure(Addresses.IsValidIndex(InAddressIndex));
+	FMemory::Memzero(GetData(InAddressIndex), Addresses[InAddressIndex].NumBytes());
+}
+
+int32 FRigVMStorage::FindOrAddScriptStruct(UScriptStruct* InScriptStruct)
+{
+	int32 StructIndex = INDEX_NONE;
+	if (ScriptStructs.Find(InScriptStruct, StructIndex))
+	{
+		return StructIndex;
+	}
+	return ScriptStructs.Add(InScriptStruct);
+}
diff --git a/Engine/Source/Runtime/RigVM/Public/MultiplexByteCode.h b/Engine/Source/Runtime/RigVM/Public/MultiplexByteCode.h
deleted file mode 100644
diff --git a/Engine/Source/Runtime/RigVM/Public/MultiplexRegistry.h b/Engine/Source/Runtime/RigVM/Public/MultiplexRegistry.h
deleted file mode 100644
diff --git a/Engine/Source/Runtime/RigVM/Public/MultiplexStorage.h b/Engine/Source/Runtime/RigVM/Public/MultiplexStorage.h
deleted file mode 100644
diff --git a/Engine/Source/Runtime/RigVM/Public/MultiplexVM.h b/Engine/Source/Runtime/RigVM/Public/MultiplexVM.h
deleted file mode 100644
diff --git a/Engine/Source/Runtime/RigVM/Public/RigVM.h b/Engine/Source/Runtime/RigVM/Public/RigVM.h
new file mode 100644
--- /dev/null
+++ b/Engine/Source/Runtime/RigVM/Public/RigVM.h
@@ -0,0 +1,46 @@
+// Copyright 1998-2019 Epic Games, Inc. All Rights Reserved.
+
+#pragma once
+
+#include "CoreMinimal.h"
+#include "RigVMStorage.h"
+#include "RigVMRegistry.h"
+#include "RigVMByteCode.h"
+#include "RigVM.generated.h"
+
+UCLASS()
+class RIGVM_API URigVM : public UObject
+{
+	GENERATED_BODY()
+
+public:
+
+	void Reset();
+
+	bool Execute(FRigVMStorage** Storage = nullptr, TArrayView<void*> AdditionalArgs = TArrayView<void*>());
+
+	UFUNCTION()
+	int32 AddRigVMFunction(UScriptStruct* InRigVMStruct, const FName& InFunctionName);
+
+	UPROPERTY()
+	FRigVMStorage Literals;
+
+	UPROPERTY()
+	FRigVMStorage WorkState;
+
+	UPROPERTY()
+	FRigVMByteCode ByteCode;
+
+	UPROPERTY()
+	FRigVMByteCodeTable Instructions;
+
+private:
+
+	void ResolveFunctionsIfRequired();
+	void RefreshInstructionsIfRequired();
+
+	UPROPERTY()
+	TArray<FString> FunctionNames;
+
+	TArray<FRigVMFunctionPtr> Functions;
+};
diff --git a/Engine/Source/Runtime/RigVM/Public/RigVMByteCode.h b/Engine/Source/Runtime/RigVM/Public/RigVMByteCode.h
new file mode 100644
--- /dev/null
+++ b/Engine/Source/Runtime/RigVM/Public/RigVMByteCode.h
@@ -0,0 +1,375 @@
+// Copyright 1998-2019 Epic Games, Inc. All Rights Reserved.
+
+#pragma once
+
+#include "CoreMinimal.h"
+#include "RigVMRegistry.h"
+#include "RigVMByteCode.generated.h"
+
+struct FRigVMByteCode;
+
+UENUM()
+enum class ERigVMOpCode : uint8
+{
+	Copy,
+	Increment,
+	Decrement,
+	Equals,
+	NotEquals,
+	Jump,
+	JumpIfTrue,
+	JumpIfFalse,
+	Execute,
+	Exit,
+	Invalid
+};
+
+struct RIGVM_API FRigVMBaseOp
+{
+	FRigVMBaseOp(ERigVMOpCode InOpCode = ERigVMOpCode::Invalid)
+	: OpCode(InOpCode)
+	{
+	}
+
+	ERigVMOpCode OpCode;
+};
+
+struct RIGVM_API FRigVMCopyOp : public FRigVMBaseOp
+{
+	FRigVMCopyOp()
+	: FRigVMBaseOp(ERigVMOpCode::Copy)
+	, Source()
+	, Target()
+	, SourceOffset(INDEX_NONE)
+	, TargetOffset(INDEX_NONE)
+	, NumBytes(INDEX_NONE)
+	{
+	}
+
+	FRigVMCopyOp(
+		FRigVMArgument InSource,
+		FRigVMArgument InTarget,
+		int32 InSourceOffset = INDEX_NONE,
+		int32 InTargetOffset = INDEX_NONE,
+		int32 InNumBytes = INDEX_NONE
+	)
+		: FRigVMBaseOp(ERigVMOpCode::Copy)
+		, Source(InSource)
+		, Target(InTarget)
+		, SourceOffset(InSourceOffset)
+		, TargetOffset(InTargetOffset)
+		, NumBytes(InNumBytes)
+	{
+	}
+
+	FRigVMArgument Source;
+	FRigVMArgument Target;
+	int32 SourceOffset;
+	int32 TargetOffset;
+	int32 NumBytes;
+};
+
+struct RIGVM_API FRigVMIncrementOp : public FRigVMBaseOp
+{
+	FRigVMIncrementOp()
+	: FRigVMBaseOp(ERigVMOpCode::Increment)
+	, Arg()
+	{
+	}
+
+	FRigVMIncrementOp(FRigVMArgument InArg)
+	: FRigVMBaseOp(ERigVMOpCode::Increment)
+	, Arg(InArg)
+	{
+	}
+
+	FRigVMArgument Arg;
+};
+
+struct RIGVM_API FRigVMDecrementOp : public FRigVMBaseOp
+{
+	FRigVMDecrementOp()
+	: FRigVMBaseOp(ERigVMOpCode::Decrement)
+	, Arg()
+	{
+	}
+
+	FRigVMDecrementOp(FRigVMArgument InArg)
+	: FRigVMBaseOp(ERigVMOpCode::Decrement)
+	, Arg(InArg)
+	{
+	}
+
+	FRigVMArgument Arg;
+};
+
+struct RIGVM_API FRigVMEqualsOp : public FRigVMBaseOp
+{
+	FRigVMEqualsOp()
+	: FRigVMBaseOp(ERigVMOpCode::Equals)
+	, A()
+	, B()
+	, Result()
+	{
+	}
+
+	FRigVMEqualsOp(
+		FRigVMArgument InA,
+		FRigVMArgument InB,
+		FRigVMArgument InResult
+	)
+	: FRigVMBaseOp(ERigVMOpCode::Equals)
+	, A(InA)
+	, B(InB)
+	, Result(InResult)
+	{
+	}
+
+	FRigVMArgument A;
+	FRigVMArgument B;
+	FRigVMArgument Result;
+};
+
+struct RIGVM_API FRigVMNotEqualsOp : public FRigVMBaseOp
+{
+	FRigVMNotEqualsOp()
+	: FRigVMBaseOp(ERigVMOpCode::NotEquals)
+	, A()
+	, B()
+	, Result()
+	{
+	}
+
+	FRigVMNotEqualsOp(
+		FRigVMArgument InA,
+		FRigVMArgument InB,
+		FRigVMArgument InResult
+	)
+	: FRigVMBaseOp(ERigVMOpCode::NotEquals)
+	, A(InA)
+	, B(InB)
+	, Result(InResult)
+	{
+	}
+
+	FRigVMArgument A;
+	FRigVMArgument B;
+	FRigVMArgument Result;
+};
+
+struct RIGVM_API FRigVMJumpOp : public FRigVMBaseOp
+{
+	FRigVMJumpOp()
+	: FRigVMBaseOp(ERigVMOpCode::Jump)
+	, InstructionIndex(INDEX_NONE)
+	{
+	}
+
+	FRigVMJumpOp(int32 InInstructionIndex)
+	: FRigVMBaseOp(ERigVMOpCode::Jump)
+	, InstructionIndex(InInstructionIndex)
+	{
+	}
+
+	int32 InstructionIndex;
+};
+
+struct RIGVM_API FRigVMJumpIfTrueOp : public FRigVMBaseOp
+{
+	FRigVMJumpIfTrueOp()
+	: FRigVMBaseOp(ERigVMOpCode::JumpIfTrue)
+	, InstructionIndex(INDEX_NONE)
+	, Condition()
+	{
+	}
+
+	FRigVMJumpIfTrueOp(int32 InInstructionIndex, FRigVMArgument InCondition)
+	: FRigVMBaseOp(ERigVMOpCode::JumpIfTrue)
+	, InstructionIndex(InInstructionIndex)
+	, Condition(InCondition)
+	{
+	}
+
+	int32 InstructionIndex;
+	FRigVMArgument Condition;
+};
+
+struct RIGVM_API FRigVMJumpIfFalseOp : public FRigVMBaseOp
+{
+	FRigVMJumpIfFalseOp()
+	: FRigVMBaseOp(ERigVMOpCode::JumpIfFalse)
+	, InstructionIndex(INDEX_NONE)
+	, Condition()
+	{
+	}
+
+	FRigVMJumpIfFalseOp(int32 InInstructionIndex, FRigVMArgument InCondition)
+	: FRigVMBaseOp(ERigVMOpCode::JumpIfFalse)
+	, InstructionIndex(InInstructionIndex)
+	, Condition(InCondition)
+	{
+	}
+
+	int32 InstructionIndex;
+	FRigVMArgument Condition;
+};
+
+struct RIGVM_API FRigVMExecuteOp : public FRigVMBaseOp
+{
+	FRigVMExecuteOp()
+	: FRigVMBaseOp(ERigVMOpCode::Execute)
+	, FunctionIndex(INDEX_NONE)
+	, ArgumentCount(0)
+	{
+	}
+
+	FRigVMExecuteOp(uint16 InFunctionIndex, uint16 InArgumentCount)
+	: FRigVMBaseOp(ERigVMOpCode::Execute)
+	, FunctionIndex(InFunctionIndex)
+	, ArgumentCount(InArgumentCount)
+	{
+	}
+
+	uint16 FunctionIndex;
+	uint16 ArgumentCount;
+};
+
+struct RIGVM_API FRigVMExitOp : public FRigVMBaseOp
+{
+	FRigVMExitOp()
+	: FRigVMBaseOp(ERigVMOpCode::Exit)
+	{
+	}
+};
+
+USTRUCT()
+struct RIGVM_API FRigVMByteCodeTableEntry
+{
+	GENERATED_USTRUCT_BODY()
+
+	FRigVMByteCodeTableEntry(ERigVMOpCode InOpCode = ERigVMOpCode::Invalid, uint64 InByteCodeIndex = UINT64_MAX)
+		: OpCode(InOpCode)
+		, ByteCodeIndex(InByteCodeIndex)
+	{
+	}
+
+	UPROPERTY()
+	ERigVMOpCode OpCode;
+
+	UPROPERTY()
+	uint64 ByteCodeIndex;
+};
+
+USTRUCT()
+struct RIGVM_API FRigVMByteCodeTable
+{
+	GENERATED_USTRUCT_BODY()
+
+public:
+
+	FRigVMByteCodeTable();
+
+	void Reset();
+	FORCEINLINE bool IsValidIndex(int32 InIndex) const { return Entries.IsValidIndex(InIndex); }
+	FORCEINLINE int32 Num() const { return Entries.Num(); }
+	FORCEINLINE const FRigVMByteCodeTableEntry& operator[](int32 InIndex) const { return Entries[InIndex]; }
+
+private:
+
+	// hide utility constructor
+	FRigVMByteCodeTable(const FRigVMByteCode& InByteCode);
+
+	UPROPERTY()
+	TArray<FRigVMByteCodeTableEntry> Entries;
+
+	friend struct FRigVMByteCode;
+};
+
+USTRUCT()
+struct RIGVM_API FRigVMByteCode
+{
+	GENERATED_USTRUCT_BODY()
+
+public:
+
+	FRigVMByteCode();
+
+	void Reset();
+	uint64 Num() const;
+
+	uint64 AddCopyOp(const FRigVMArgument& InSource, const FRigVMArgument& InTarget, int32 InSourceOffset = INDEX_NONE, int32 InTargetOffset = INDEX_NONE, int32 InNumBytes = INDEX_NONE);
+	uint64 AddIncrementOp(const FRigVMArgument& InArg);
+	uint64 AddDecrementOp(const FRigVMArgument& InArg);
+	uint64 AddEqualsOp(const FRigVMArgument& InA, const FRigVMArgument& InB, const FRigVMArgument& InResult);
+	uint64 AddNotEqualsOp(const FRigVMArgument& InA, const FRigVMArgument& InB, const FRigVMArgument& InResult);
+	uint64 AddJumpOp(uint64 InByteCodeIndex);
+	uint64 AddJumpIfTrueOp(uint64 InByteCodeIndex, const FRigVMArgument& InCondition);
+	uint64 AddJumpIfFalseOp(uint64 InByteCodeIndex, const FRigVMArgument& InCondition);
+	uint64 AddExecuteOp(uint16 InFunctionIndex, const TArrayView<FRigVMArgument>& InArguments);
+	uint64 AddExitOp();
+
+	FORCEINLINE FRigVMByteCodeTable GetTable() const
+	{
+		return FRigVMByteCodeTable(*this);
+	}
+
+	FORCEINLINE ERigVMOpCode GetOpCodeAt(uint64 InByteCodeIndex) const
+	{
+		ensure(InByteCodeIndex >= 0 && InByteCodeIndex < ByteCode.Num());
+		return (ERigVMOpCode)ByteCode[InByteCodeIndex];
+	}
+
+	uint64 GetOpNumBytesAt(uint64 InByteCodeIndex, bool bIncludeArguments = true) const;
+
+	template<class OpType>
+	FORCEINLINE const OpType& GetOpAt(uint64 InByteCodeIndex) const
+	{
+		ensure(InByteCodeIndex >= 0 && InByteCodeIndex <= ByteCode.Num() - sizeof(OpType));
+		return *(const OpType*)(ByteCode.GetData() + InByteCodeIndex);
+	}
+
+	template<class OpType>
+	FORCEINLINE const OpType& GetOpAt(const FRigVMByteCodeTableEntry& InEntry) const
+	{
+		ensure(OpType().OpCode == InEntry.OpCode);
+		return GetOpAt<OpType>(InEntry.ByteCodeIndex);
+	}
+
+	FORCEINLINE TArrayView<FRigVMArgument> GetArgumentsAt(uint64 InByteCodeIndex, uint16 InArgumentCount) const
+	{
+		ensure(InByteCodeIndex >= 0 && InByteCodeIndex <= ByteCode.Num() - sizeof(FRigVMArgument) * InArgumentCount);
+		return TArrayView<FRigVMArgument>((FRigVMArgument*)(ByteCode.GetData() + InByteCodeIndex), InArgumentCount);
+	}
+
+	FORCEINLINE TArrayView<FRigVMArgument> GetArgumentsForExecuteOp(uint64 InByteCodeIndex) const
+	{
+		const FRigVMExecuteOp& ExecuteOp = GetOpAt<FRigVMExecuteOp>(InByteCodeIndex);
+		return GetArgumentsAt(InByteCodeIndex + sizeof(FRigVMExecuteOp), ExecuteOp.ArgumentCount);
+	}
+
+	FORCEINLINE TArrayView<FRigVMArgument> GetArgumentsForExecuteOp(const FRigVMByteCodeTableEntry& InEntry) const
+	{
+		const FRigVMExecuteOp& ExecuteOp = GetOpAt<FRigVMExecuteOp>(InEntry);
+		return GetArgumentsAt(InEntry.ByteCodeIndex + sizeof(FRigVMExecuteOp), ExecuteOp.ArgumentCount);
+	}
+
+	FORCEINLINE const TArrayView<uint8> GetByteCode() const
+	{
+		const uint8* Data = ByteCode.GetData();
+		return TArrayView<uint8>((uint8*)Data, ByteCode.Num());
+	}
+
+private:
+
+	template<class OpType>
+	FORCEINLINE uint64 AddOp(const OpType& InOp)
+	{
+		uint64 Address = (uint64)ByteCode.AddUninitialized(sizeof(OpType));
+		FMemory::Memcpy(ByteCode.GetData() + Address, &InOp, sizeof(OpType));
+		return Address;
+	}
+
+	// storage for all functions
+	UPROPERTY()
+	TArray<uint8> ByteCode;
+};
diff --git a/Engine/Source/Runtime/RigVM/Public/RigVMRegistry.h b/Engine/Source/Runtime/RigVM/Public/RigVMRegistry.h
new file mode 100644
--- /dev/null
+++ b/Engine/Source/Runtime/RigVM/Public/RigVMRegistry.h
@@ -0,0 +1,51 @@
+// Copyright 1998-2019 Epic Games, Inc. All Rights Reserved.
+
+#pragma once
+
+#include "CoreMinimal.h"
+#include "RigVMStorage.h"
+
+typedef void (*FRigVMFunctionPtr)(const TArrayView<FRigVMArgument>&, FRigVMStorage**, const TArrayView<void*>&);
+
+struct RIGVM_API FRigVMFunction
+{
+	const TCHAR* Name;
+	FRigVMFunctionPtr FunctionPtr;
+
+	FRigVMFunction()
+		: Name(nullptr)
+		, FunctionPtr(nullptr)
+	{
+	}
+
+	FRigVMFunction(const TCHAR* InName, FRigVMFunctionPtr InFunctionPtr)
+		: Name(InName)
+		, FunctionPtr(InFunctionPtr)
+	{
+	}
+};
+
+struct RIGVM_API FRigVMRegistry
+{
+public:
+	static FRigVMRegistry& Get();
+
+	void Register(const TCHAR* InName, FRigVMFunctionPtr InFunctionPtr);
+	void Refresh();
+
+	FRigVMFunctionPtr Find(const TCHAR* InName) const;
+
+private:
+
+	// disable default constructor
+	FRigVMRegistry() {}
+	// disable copy constructor
+	FRigVMRegistry(const FRigVMRegistry&) {}
+	// disable assignment operator
+	FRigVMRegistry& operator= (const FRigVMRegistry &InOther) { return *this; }
+
+	// storage for all functions
+	TArray<FRigVMFunction> Functions;
+
+	static FRigVMRegistry s_RigVMRegistry;
+};
diff --git a/Engine/Source/Runtime/RigVM/Public/RigVMStorage.h b/Engine/Source/Runtime/RigVM/Public/RigVMStorage.h
new file mode 100644
--- /dev/null
+++ b/Engine/Source/Runtime/RigVM/Public/RigVMStorage.h
@@ -0,0 +1,547 @@
+// Copyright 1998-2019 Epic Games, Inc. All Rights Reserved.
+
+#pragma once
+
+#include "CoreMinimal.h"
+#include "UObject/ObjectMacros.h"
+#include "RigVMStorage.generated.h"
+
+struct RIGVM_API FRigVMArgument
+{
+public:
+
+	FRigVMArgument()
+		:Address(INDEX_NONE)
+	{
+	}
+
+	FRigVMArgument(int32 InAddress)
+		:Address(InAddress)
+	{
+	}
+
+	static FRigVMArgument MakeArgument(int32 InAddress)
+	{
+		return FRigVMArgument(InAddress);
+	}
+
+	static FRigVMArgument MakeLiteral(int32 InAddress)
+	{
+		return FRigVMArgument(-1 - InAddress);
+	}
+
+	FORCEINLINE bool operator ==(int32 InOther) const { return Address == InOther; }
+
+	FORCEINLINE bool IsLiteral() const { return Address < 0; }
+	FORCEINLINE int32 StorageType() const { return Address < 0 ? 1 : 0; }
+	FORCEINLINE int32 Index() const { return Address < 0 ? -(Address + 1) : Address; }
+	FORCEINLINE int32 NativeAddress() const { return Address; }
+
+private:
+	int32 Address;
+};
+
+UENUM()
+enum class ERigVMAddressType : uint8
+{
+	Plain,
+	String,
+	Name,
+	Struct,
+	Invalid
+};
+
+USTRUCT()
+struct RIGVM_API FRigVMAddress
+{
+	GENERATED_USTRUCT_BODY()
+
+		FRigVMAddress()
+		: Type(ERigVMAddressType::Invalid)
+		, Pointer(nullptr)
+		, ByteIndex(INDEX_NONE)
+		, ElementSize(0)
+		, ElementCount(0)
+		, AlignmentBytes(0)
+		, Name(NAME_None)
+		, ScriptStructIndex(INDEX_NONE)
+	{
+	}
+
+	UPROPERTY()
+	ERigVMAddressType Type;
+
+	void* Pointer;
+
+	UPROPERTY()
+	int32 ByteIndex;
+
+	UPROPERTY()
+	int32 ElementSize;
+
+	UPROPERTY()
+	int32 ElementCount;
+
+	UPROPERTY()
+	int32 AlignmentBytes;
+
+	UPROPERTY()
+	FName Name;
+
+	UPROPERTY()
+	int32 ScriptStructIndex;
+
+	FORCEINLINE int32 FirstByte() const { return ByteIndex + AlignmentBytes; }
+	FORCEINLINE bool IsArray() const { return ElementCount > 1; }
+	FORCEINLINE bool IsPlain() const { return ScriptStructIndex == INDEX_NONE;  }
+	FORCEINLINE int32 NumBytes(bool bIncludeAlignment = true) const { return ElementCount * ElementSize + (bIncludeAlignment ? AlignmentBytes : 0); }
+
+	template<class T>
+	FORCEINLINE const T* Get() const
+	{
+		ensure(ElementCount > 0);
+		return (const T*)Pointer;
+	}
+
+	template<class T>
+	FORCEINLINE const T& GetRef() const
+	{
+		return *Get<T>();
+	}
+
+	template<class T>
+	FORCEINLINE T* Get()
+	{
+		ensure(ElementCount > 0);
+		return (T*)Pointer;
+	}
+
+	template<class T>
+	FORCEINLINE T& GetRef()
+	{
+		return *Get<T>();
+	}
+
+	template<class T>
+	FORCEINLINE TArrayView<T> GetArray()
+	{
+		ensure(ElementCount > 0);
+		return TArrayView<T>((T*)Pointer, ElementCount);
+	}
+};
+
+USTRUCT()
+struct RIGVM_API FRigVMStorage
+{
+	GENERATED_USTRUCT_BODY()
+	
+public:
+
+	FRigVMStorage(bool bInUseNames = true);
+	FRigVMStorage(const FRigVMStorage& Other);
+	~FRigVMStorage();
+
+	FRigVMStorage& operator= (const FRigVMStorage &InOther);
+
+	FORCEINLINE bool SupportsNames() const { return bUseNameMap;  }
+	FORCEINLINE int32 Num() const { return Addresses.Num(); }
+	FORCEINLINE const FRigVMAddress& operator[](int32 InIndex) const { return Addresses[InIndex]; }
+	FORCEINLINE FRigVMAddress& operator[](int32 InIndex) { return Addresses[InIndex]; }
+	FORCEINLINE const FRigVMAddress& operator[](const FRigVMArgument& InArg) const { return Addresses[InArg.Index()]; }
+	FORCEINLINE FRigVMAddress& operator[](const FRigVMArgument& InArg) { return Addresses[InArg.Index()]; }
+	FORCEINLINE const FRigVMAddress& operator[](const FName& InName) const { return Addresses[GetIndex(InName)]; }
+	FORCEINLINE FRigVMAddress& operator[](const FName& InName) { return Addresses[GetIndex(InName)]; }
+
+	FORCEINLINE TArray<FRigVMAddress>::RangedForIteratorType      begin() { return Addresses.begin(); }
+	FORCEINLINE TArray<FRigVMAddress>::RangedForConstIteratorType begin() const { return Addresses.begin(); }
+	FORCEINLINE TArray<FRigVMAddress>::RangedForIteratorType      end() { return Addresses.end(); }
+	FORCEINLINE TArray<FRigVMAddress>::RangedForConstIteratorType end() const { return Addresses.end(); }
+
+	FORCEINLINE const void* GetData(int32 InAddressIndex) const
+	{
+		ensure(Addresses.IsValidIndex(InAddressIndex));
+		const FRigVMAddress& Address = Addresses[InAddressIndex];
+		ensure(Address.ElementCount > 0);
+		return (const void*)&Data[Address.FirstByte()];
+	}
+
+	FORCEINLINE void* GetData(int32 InAddressIndex)
+	{
+		ensure(Addresses.IsValidIndex(InAddressIndex));
+		const FRigVMAddress& Address = Addresses[InAddressIndex];
+		ensure(Address.ElementCount > 0);
+		return (void*)&Data[Address.FirstByte()];
+	}
+
+	template<class T>
+	FORCEINLINE const T* Get(int32 InAddressIndex) const
+	{
+		ensure(Addresses.IsValidIndex(InAddressIndex));
+		const FRigVMAddress& Address = Addresses[InAddressIndex];
+		ensure(Address.ElementCount > 0);
+		return (const T*)&Data[Address.FirstByte()];
+	}
+
+	template<class T>
+	FORCEINLINE const T* Get(const FRigVMArgument& InArgument) const
+	{
+		return Get<T>(InArgument.Index());
+	}
+
+	template<class T>
+	FORCEINLINE const T& GetRef(int32 InAddressIndex) const
+	{
+		return *Get<T>(InAddressIndex);
+	}
+
+	template<class T>
+	FORCEINLINE const T& GetRef(const FRigVMArgument& InArgument) const
+	{
+		return GetRef<T>(InArgument.Index());
+	}
+
+	template<class T>
+	FORCEINLINE T* Get(int32 InAddressIndex)
+	{
+		ensure(Addresses.IsValidIndex(InAddressIndex));
+		const FRigVMAddress& Address = Addresses[InAddressIndex];
+		ensure(Address.ElementCount > 0);
+		return (T*)&Data[Address.FirstByte()];
+	}
+
+	template<class T>
+	FORCEINLINE T* Get(const FRigVMArgument& InArgument)
+	{
+		return Get<T>(InArgument.Index());
+	}
+
+	template<class T>
+	FORCEINLINE T& GetRef(int32 InAddressIndex)
+	{
+		return *Get<T>(InAddressIndex);
+	}
+
+	template<class T>
+	FORCEINLINE T& GetRef(const FRigVMArgument& InArgument)
+	{
+		return GetRef<T>(InArgument.Index());
+	}
+
+	template<class T>
+	FORCEINLINE TArrayView<T> GetArray(int32 InAddressIndex)
+	{
+		ensure(Addresses.IsValidIndex(InAddressIndex));
+		const FRigVMAddress& Address = Addresses[InAddressIndex];
+		ensure(Address.ElementCount > 0);
+		return TArrayView<T>((T*)&Data[Address.FirstByte()], Address.ElementCount);
+	}
+	
+	template<class T>
+	FORCEINLINE TArrayView<T> GetArray(const FRigVMArgument& InArgument)
+	{
+		return GetArray<T>(InArgument.Index());
+	}
+
+	FORCEINLINE UScriptStruct* GetScriptStruct(int32 InAddressIndex) const
+	{
+		ensure(Addresses.IsValidIndex(InAddressIndex));
+		const FRigVMAddress& Address = Addresses[InAddressIndex];
+		if (Address.ScriptStructIndex != INDEX_NONE)
+		{
+			ensure(ScriptStructs.IsValidIndex(Address.ScriptStructIndex));
+			return ScriptStructs[Address.ScriptStructIndex];
+		}
+		return nullptr;
+	}
+
+	bool Copy(
+		int32 InSourceAddressIndex,
+		int32 InTargetAddressIndex,
+		const FRigVMStorage* InSourceStorage = nullptr,
+		int32 InSourceByteOffset = INDEX_NONE,
+		int32 InTargetByteOffset = INDEX_NONE,
+		int32 InNumBytes = INDEX_NONE);
+
+	bool Copy(
+		const FName& InSourceName,
+		const FName& InTargetName,
+		const FRigVMStorage* InSourceStorage = nullptr,
+		int32 InSourceByteOffset = INDEX_NONE,
+		int32 InTargetByteOffset = INDEX_NONE,
+		int32 InNumBytes = INDEX_NONE);
+
+	FORCEINLINE int32 GetIndex(const FName& InName) const
+	{
+		if (!bUseNameMap)
+		{
+			return INDEX_NONE;
+		}
+
+		if (NameMap.Num() != Addresses.Num())
+		{
+			for (int32 Index = 0; Index < Addresses.Num(); Index++)
+			{
+				if (Addresses[Index].Name == InName)
+				{
+					return Index;
+				}
+			}
+		}
+		else
+		{
+			const int32* Index = NameMap.Find(InName);
+			if (Index != nullptr)
+			{
+				return *Index;
+			}
+		}
+
+		return INDEX_NONE;
+	}
+
+	FORCEINLINE bool IsNameAvailable(const FName& InPotentialNewName) const\
+	{
+		if (!bUseNameMap)
+		{
+			return false;
+		}
+		return GetIndex(InPotentialNewName) == INDEX_NONE;
+	}
+
+	void Reset();
+
+	FORCEINLINE int32 AddPlainArray(const FName& InNewName, int32 InElementSize, int32 InCount, const void* InDataPtr = nullptr)
+	{
+		return Allocate(InNewName, InElementSize, InCount, InDataPtr);
+	}
+
+	template<class T>
+	FORCEINLINE int32 AddPlainArray(const FName& InNewName, int32 InCount, const T* InDataPtr = nullptr)
+	{
+		return Allocate(InNewName, sizeof(T), InCount, (const void*)InDataPtr);
+	}
+
+	template<class T>
+	FORCEINLINE int32 AddPlainArray(const FName& InNewName, const TArray<T>& InArray)
+	{
+		return AddPlainArray<T>(InNewName, InArray.Num(), InArray.GetData());
+	}
+
+	template<class T>
+	FORCEINLINE int32 AddPlainArray(const TArray<T>& InArray)
+	{
+		return AddPlainArray<T>(NAME_None, InArray);
+	}
+
+	FORCEINLINE int32 AddPlain(const FName& InNewName, int32 InElementSize, const void* InValuePtr)
+	{
+		return AddPlainArray(InNewName, InElementSize, 1, InValuePtr);
+	}
+
+	FORCEINLINE int32 AddPlain(int32 InElementSize, const void* InValuePtr)
+	{
+		return AddPlain(NAME_None, InElementSize, InValuePtr);
+	}
+
+	template<class T>
+	FORCEINLINE int32 AddPlain(const FName& InNewName, const T& InValue)
+	{
+		return AddPlainArray<T>(InNewName, 1, &InValue);
+	}
+
+	template<class T>
+	FORCEINLINE int32 AddPlain(const T& InValue)
+	{
+		return AddPlain<T>(NAME_None, InValue);
+	}
+
+	FORCEINLINE int32 AddNameArray(const FName& InNewName, int32 InCount, const FName* InDataPtr = nullptr)
+	{
+		int32 Address = Allocate(InNewName, sizeof(FName), InCount, nullptr);
+		Addresses[Address].Type = ERigVMAddressType::Name;
+
+		Construct(Address);
+
+		if(InDataPtr)
+		{
+			FName* DataPtr = (FName*)GetData(Address);
+			for (int32 Index = 0; Index < InCount; Index++)
+			{
+				DataPtr[Index] = InDataPtr[Index];
+			}
+		}
+
+		return Address;
+	}
+
+	FORCEINLINE int32 AddNameArray(const FName& InNewName, const TArray<FName>& InArray)
+	{
+		return AddNameArray(InNewName, InArray.Num(), InArray.GetData());
+	}
+
+	FORCEINLINE int32 AddNameArray(const TArray<FName>& InArray)
+	{
+		return AddNameArray(NAME_None, InArray);
+	}
+
+	FORCEINLINE int32 AddName(const FName& InNewName, const FName& InValue)
+	{
+		return AddNameArray(InNewName, 1, &InValue);
+	}
+
+	FORCEINLINE int32 AddName(const FName& InValue)
+	{
+		return AddName(NAME_None, InValue);
+	}
+
+	FORCEINLINE int32 AddStringArray(const FName& InNewName, int32 InCount, const FString* InDataPtr = nullptr)
+	{
+		int32 Address = Allocate(InNewName, sizeof(FString), InCount, nullptr);
+		Addresses[Address].Type = ERigVMAddressType::String;
+
+		Construct(Address);
+
+		if(InDataPtr)
+		{
+			FString* DataPtr = (FString*)GetData(Address);
+			for (int32 Index = 0; Index < InCount; Index++)
+			{
+				DataPtr[Index] = InDataPtr[Index];
+			}
+		}
+
+		return Address;
+	}
+
+	FORCEINLINE int32 AddStringArray(const FName& InNewName, const TArray<FString>& InArray)
+	{
+		return AddStringArray(InNewName, InArray.Num(), InArray.GetData());
+	}
+
+	FORCEINLINE int32 AddStringArray(const TArray<FString>& InArray)
+	{
+		return AddStringArray(NAME_None, InArray);
+	}
+
+	FORCEINLINE int32 AddString(const FName& InNewName, const FString& InValue)
+	{
+		return AddStringArray(InNewName, 1, &InValue);
+	}
+
+	FORCEINLINE int32 AddString(const FString& InValue)
+	{
+		return AddString(NAME_None, InValue);
+	}
+
+
+	FORCEINLINE int32 AddStructArray(const FName& InNewName, UScriptStruct* InScriptStruct, int32 InCount, const void* InDataPtr = nullptr)
+	{
+		int32 Address = Allocate(InNewName, InScriptStruct->GetStructureSize(), InCount, nullptr, false);
+		if (Address == INDEX_NONE)
+		{
+			return INDEX_NONE;
+		}
+
+		Addresses[Address].Type = ERigVMAddressType::Struct;
+		Addresses[Address].ScriptStructIndex = FindOrAddScriptStruct(InScriptStruct);
+
+		UpdateAddresses();
+
+		// construct the content
+		Construct(Address);
+
+		// copy values from the provided data
+		if (InDataPtr != nullptr)
+		{
+			InScriptStruct->CopyScriptStruct(GetData(Address), InDataPtr, InCount);
+		}
+
+		return Address;
+	}
+
+	FORCEINLINE int32 AddStructArray(UScriptStruct* InScriptStruct, int32 InCount, const void* InDataPtr = nullptr)
+	{
+		return AddStructArray(NAME_None, InScriptStruct, InCount, InDataPtr);
+	}
+
+	template<class T>
+	FORCEINLINE int32 AddStructArray(const FName& InNewName, int32 InCount, const T* InDataPtr = nullptr)
+	{
+		// if you are hitting this - you might need to use AddPlainArray instead!
+		UScriptStruct* Struct = T::StaticStruct();
+		if (Struct == nullptr)
+		{
+			return INDEX_NONE;
+		}
+	
+		return AddStructArray(InNewName, Struct, InCount, InDataPtr);
+	}
+
+	template<class T>
+	FORCEINLINE int32 AddStructArray(const FName& InNewName, const TArray<T>& InArray)
+	{
+		return AddStructArray<T>(InNewName, InArray.Num(), InArray.GetData());
+	}
+
+	template<class T>
+	FORCEINLINE int32 AddStructArray(const TArray<T>& InArray)
+	{
+		return AddStructArray<T>(NAME_None, InArray);
+	}
+
+	FORCEINLINE int32 AddStruct(const FName& InNewName, UScriptStruct* InScriptStruct, const void* InValuePtr)
+	{
+		return AddStructArray(InNewName, InScriptStruct, 1, InValuePtr);
+	}
+
+	FORCEINLINE int32 AddStruct(UScriptStruct* InScriptStruct, const void* InValuePtr)
+	{
+		return AddStruct(NAME_None, InScriptStruct, InValuePtr);
+	}
+
+	template<class T>
+	FORCEINLINE int32 AddStruct(const FName& InNewName, const T& InValue)
+	{
+		return AddStructArray<T>(InNewName, 1, &InValue);
+	}
+
+	template<class T>
+	FORCEINLINE int32 AddStruct(const T& InValue)
+	{
+		return AddStruct<T>(NAME_None, InValue);
+	}
+
+	bool Remove(int32 InAddressIndex);
+	bool Remove(const FName& InAddressName);
+	FName Rename(int32 InAddressIndex, const FName& InNewName);
+	FName Rename(const FName& InOldName, const FName& InNewName);
+	bool Resize(int32 InAddressIndex, int32 InNewElementCount);
+	bool Resize(const FName& InAddressName, int32 InNewElementCount);
+
+	void UpdateAddresses();
+
+private:
+
+	int32 Allocate(const FName& InNewName, int32 InElementSize, int32 InCount, const void* InDataPtr = nullptr, bool bUpdateAddresses = true);
+	int32 Allocate(int32 InElementSize, int32 InCount, const void* InDataPtr = nullptr, bool bUpdateAddresses = true);
+	bool Construct(int32 InAddressIndex, int32 InElementIndex = INDEX_NONE);
+	bool Destroy(int32 InAddressIndex, int32 InElementIndex = INDEX_NONE);
+
+	void FillWithZeroes(int32 InAddressIndex);
+	int32 FindOrAddScriptStruct(UScriptStruct* InScriptStruct);
+
+	UPROPERTY()
+	bool bUseNameMap;
+
+	UPROPERTY()
+	TArray<FRigVMAddress> Addresses;
+
+	UPROPERTY()
+	TArray<uint8> Data;
+
+	UPROPERTY()
+	TArray<UScriptStruct*> ScriptStructs;
+
+	UPROPERTY(transient)
+	TMap<FName, int32> NameMap;
+};
