RigVM: Renaming Address to Register

diff --git a/Engine/Source/Runtime/RigVM/Private/RigVMByteCode.cpp b/Engine/Source/Runtime/RigVM/Private/RigVMByteCode.cpp
--- a/Engine/Source/Runtime/RigVM/Private/RigVMByteCode.cpp
+++ b/Engine/Source/Runtime/RigVM/Private/RigVMByteCode.cpp
@@ -8,12 +8,12 @@ FRigVMByteCodeTable::FRigVMByteCodeTable()
 
 FRigVMByteCodeTable::FRigVMByteCodeTable(const FRigVMByteCode& InByteCode)
 {
-	uint64 Address = 0;
-	while (Address < InByteCode.Num())
+	uint64 ByteIndex = 0;
+	while (ByteIndex < InByteCode.Num())
 	{
-		ERigVMOpCode OpCode = InByteCode.GetOpCodeAt(Address);
-		Entries.Add(FRigVMByteCodeTableEntry(OpCode, Address));
-		Address += InByteCode.GetOpNumBytesAt(Address);
+		ERigVMOpCode OpCode = InByteCode.GetOpCodeAt(ByteIndex);
+		Entries.Add(FRigVMByteCodeTableEntry(OpCode, ByteIndex));
+		ByteIndex += InByteCode.GetOpNumBytesAt(ByteIndex);
 	}
 }
 
@@ -147,10 +147,10 @@ uint64 FRigVMByteCode::AddJumpIfFalseOp(uint64 InByteCodeIndex, const FRigVMArgu
 uint64 FRigVMByteCode::AddExecuteOp(uint16 InFunctionIndex, const TArrayView<FRigVMArgument>& InArguments)
 {
 	FRigVMExecuteOp Op(InFunctionIndex, (uint16)InArguments.Num());
-	uint64 OpAddress = AddOp(Op);
-	uint64 ArgumentsAddress = (uint64)ByteCode.AddUninitialized(sizeof(FRigVMArgument) * InArguments.Num());
-	FMemory::Memcpy(ByteCode.GetData() + ArgumentsAddress, InArguments.GetData(), sizeof(FRigVMArgument) * InArguments.Num());
-	return OpAddress;
+	uint64 OpByteIndex = AddOp(Op);
+	uint64 ArgumentsByteIndex = (uint64)ByteCode.AddUninitialized(sizeof(FRigVMArgument) * InArguments.Num());
+	FMemory::Memcpy(ByteCode.GetData() + ArgumentsByteIndex, InArguments.GetData(), sizeof(FRigVMArgument) * InArguments.Num());
+	return OpByteIndex;
 }
 
 uint64 FRigVMByteCode::AddExitOp()
diff --git a/Engine/Source/Runtime/RigVM/Private/RigVMStorage.cpp b/Engine/Source/Runtime/RigVM/Private/RigVMStorage.cpp
--- a/Engine/Source/Runtime/RigVM/Private/RigVMStorage.cpp
+++ b/Engine/Source/Runtime/RigVM/Private/RigVMStorage.cpp
@@ -24,12 +24,12 @@ FRigVMStorage& FRigVMStorage::operator= (const FRigVMStorage &InOther)
 
 	bUseNameMap = InOther.bUseNameMap;
 	Data.Append(InOther.Data);
-	Addresses.Append(InOther.Addresses);
+	Registers.Append(InOther.Registers);
 	ScriptStructs.Append(InOther.ScriptStructs);
 
-	UpdateAddresses();
+	UpdateRegisters();
 
-	for (int32 Index = 0; Index < Addresses.Num(); Index++)
+	for (int32 Index = 0; Index < Registers.Num(); Index++)
 	{
 		Construct(Index);
 		Copy(Index, Index, &InOther);
@@ -40,20 +40,20 @@ FRigVMStorage& FRigVMStorage::operator= (const FRigVMStorage &InOther)
 
 void FRigVMStorage::Reset()
 {
-	for (int32 Index = 0; Index < Addresses.Num(); Index++)
+	for (int32 Index = 0; Index < Registers.Num(); Index++)
 	{
 		Destroy(Index);
 	}
 
 	Data.Reset();
-	Addresses.Reset();
+	Registers.Reset();
 	ScriptStructs.Reset();
 	NameMap.Reset();
 }
 
 bool FRigVMStorage::Copy(
-	int32 InSourceAddressIndex,
-	int32 InTargetAddressIndex,
+	int32 InSourceRegisterIndex,
+	int32 InTargetRegisterIndex,
 	const FRigVMStorage* InSourceStorage,
 	int32 InSourceByteOffset,
 	int32 InTargetByteOffset,
@@ -64,16 +64,16 @@ bool FRigVMStorage::Copy(
 		InSourceStorage = this;
 	}
 
-	ensure(InSourceStorage->Addresses.IsValidIndex(InSourceAddressIndex));
-	ensure(Addresses.IsValidIndex(InTargetAddressIndex));
+	ensure(InSourceStorage->Registers.IsValidIndex(InSourceRegisterIndex));
+	ensure(Registers.IsValidIndex(InTargetRegisterIndex));
 
-	if (InSourceAddressIndex == InTargetAddressIndex && InSourceByteOffset == InTargetByteOffset && this == InSourceStorage)
+	if (InSourceRegisterIndex == InTargetRegisterIndex && InSourceByteOffset == InTargetByteOffset && this == InSourceStorage)
 	{
 		return false;
 	}
 
-	const FRigVMAddress& Source = InSourceStorage->Addresses[InSourceAddressIndex];
-	const FRigVMAddress& Target = Addresses[InTargetAddressIndex];
+	const FRigVMRegister& Source = InSourceStorage->Registers[InSourceRegisterIndex];
+	const FRigVMRegister& Target = Registers[InTargetRegisterIndex];
 
 	int32 SourceStartByte = Source.FirstByte();
 	int32 SourceNumBytes = Source.NumBytes(false);
@@ -106,19 +106,19 @@ bool FRigVMStorage::Copy(
 
 	switch (Target.Type)
 	{
-		case ERigVMAddressType::Plain:
+		case ERigVMRegisterType::Plain:
 		{
 			FMemory::Memcpy(&Data[TargetStartByte], &InSourceStorage->Data[SourceStartByte], TargetNumBytes);
 			break;
 		}
-		case ERigVMAddressType::Struct:
+		case ERigVMRegisterType::Struct:
 		{
-			UScriptStruct* ScriptStruct = GetScriptStruct(InTargetAddressIndex);
+			UScriptStruct* ScriptStruct = GetScriptStruct(InTargetRegisterIndex);
 			int32 NumStructs = TargetNumBytes / ScriptStruct->GetStructureSize();
 			ScriptStruct->CopyScriptStruct(&Data[TargetStartByte], &InSourceStorage->Data[SourceStartByte], NumStructs);
 			break;
 		}
-		case ERigVMAddressType::Name:
+		case ERigVMRegisterType::Name:
 		{
 			int32 NumNames = TargetNumBytes / sizeof(FName);
 			TArrayView<FName> TargetNames((FName*)&Data[TargetStartByte], NumNames);
@@ -129,7 +129,7 @@ bool FRigVMStorage::Copy(
 			}
 			break;
 		}
-		case ERigVMAddressType::String:
+		case ERigVMRegisterType::String:
 		{
 			int32 NumStrings = TargetNumBytes / sizeof(FString);
 			TArrayView<FString> TargetStrings((FString*)&Data[TargetStartByte], NumStrings);
@@ -140,7 +140,7 @@ bool FRigVMStorage::Copy(
 			}
 			break;
 		}
-		case ERigVMAddressType::Invalid:
+		case ERigVMRegisterType::Invalid:
 		{
 			return false;
 		}
@@ -159,27 +159,27 @@ bool FRigVMStorage::Copy(
 {
 	ensure(bUseNameMap);
 
-	int32 SourceAddressIndex = GetIndex(InSourceName);
-	int32 TargetAddressIndex = GetIndex(InTargetName);
+	int32 SourceRegisterIndex = GetIndex(InSourceName);
+	int32 TargetRegisterIndex = GetIndex(InTargetName);
 
-	if(SourceAddressIndex == INDEX_NONE || TargetAddressIndex == INDEX_NONE)
+	if(SourceRegisterIndex == INDEX_NONE || TargetRegisterIndex == INDEX_NONE)
 	{
 		return false;
 	}
 
-	return Copy(SourceAddressIndex, TargetAddressIndex, InSourceStorage, InSourceByteOffset, InTargetByteOffset, InNumBytes);
+	return Copy(SourceRegisterIndex, TargetRegisterIndex, InSourceStorage, InSourceByteOffset, InTargetByteOffset, InNumBytes);
 }
 
-int32 FRigVMStorage::Allocate(const FName& InNewName, int32 InElementSize, int32 InCount, const void* InDataPtr, bool bUpdateAddresses)
+int32 FRigVMStorage::Allocate(const FName& InNewName, int32 InElementSize, int32 InCount, const void* InDataPtr, bool bUpdateRegisters)
 {
 	FName Name = InNewName;
 	if (bUseNameMap && InNewName == NAME_None)
 	{
-		const TCHAR* AddressPrefix = TEXT("Address");
-		int32 AddressSuffix = 0;
+		const TCHAR* RegisterPrefix = TEXT("Register");
+		int32 RegisterSuffix = 0;
 		do
 		{
-			Name = FName(*FString::Printf(TEXT("%s_%d"), AddressPrefix, AddressSuffix++));
+			Name = FName(*FString::Printf(TEXT("%s_%d"), RegisterPrefix, RegisterSuffix++));
 		} while (!IsNameAvailable(Name));
 	}
 
@@ -193,72 +193,72 @@ int32 FRigVMStorage::Allocate(const FName& InNewName, int32 InElementSize, int32
 		}
 	}
 
-	FRigVMAddress NewAddress;
-	NewAddress.ByteIndex = Data.Num();
+	FRigVMRegister NewRegister;
+	NewRegister.ByteIndex = Data.Num();
 	if (bUseNameMap)
 	{
-		NewAddress.Name = Name;
+		NewRegister.Name = Name;
 	}
-	NewAddress.ElementSize = InElementSize;
-	NewAddress.ElementCount = InCount;
-	NewAddress.Type = ERigVMAddressType::Plain;
+	NewRegister.ElementSize = InElementSize;
+	NewRegister.ElementCount = InCount;
+	NewRegister.Type = ERigVMRegisterType::Plain;
 
-	Data.AddZeroed(NewAddress.NumBytes());
+	Data.AddZeroed(NewRegister.NumBytes());
 
 	if (InDataPtr != nullptr)
 	{
-		FMemory::Memcpy(&Data[NewAddress.FirstByte()], InDataPtr, NewAddress.NumBytes());
+		FMemory::Memcpy(&Data[NewRegister.FirstByte()], InDataPtr, NewRegister.NumBytes());
 	}
 
-	int32 AddressIndex = Addresses.Num();
-	Addresses.Add(NewAddress);
+	int32 RegisterIndex = Registers.Num();
+	Registers.Add(NewRegister);
 
-	if (bUpdateAddresses)
+	if (bUpdateRegisters)
 	{
-		UpdateAddresses();
+		UpdateRegisters();
 	}
-	return AddressIndex;
+	return RegisterIndex;
 }
 
-int32 FRigVMStorage::Allocate(int32 InElementSize, int32 InCount, const void* InDataPtr, bool bUpdateAddresses)
+int32 FRigVMStorage::Allocate(int32 InElementSize, int32 InCount, const void* InDataPtr, bool bUpdateRegisters)
 {
-	return Allocate(NAME_None, InElementSize, InCount, InDataPtr, bUpdateAddresses);
+	return Allocate(NAME_None, InElementSize, InCount, InDataPtr, bUpdateRegisters);
 }
 
-bool FRigVMStorage::Construct(int32 InAddressIndex, int32 InElementIndex)
+bool FRigVMStorage::Construct(int32 InRegisterIndex, int32 InElementIndex)
 {
-	ensure(Addresses.IsValidIndex(InAddressIndex));
+	ensure(Registers.IsValidIndex(InRegisterIndex));
 
-	const FRigVMAddress& Address = Addresses[InAddressIndex];
-	switch (Address.Type)
+	const FRigVMRegister& Register = Registers[InRegisterIndex];
+	switch (Register.Type)
 	{
-		case ERigVMAddressType::Struct:
+		case ERigVMRegisterType::Struct:
 		{
-			void* DataPtr = (void*)&Data[InElementIndex == INDEX_NONE ? Address.FirstByte() : Address.FirstByte() + InElementIndex * Address.ElementSize];
-			int32 Count = InElementIndex == INDEX_NONE ? Address.ElementCount : 1;
+			void* DataPtr = (void*)&Data[InElementIndex == INDEX_NONE ? Register.FirstByte() : Register.FirstByte() + InElementIndex * Register.ElementSize];
+			int32 Count = InElementIndex == INDEX_NONE ? Register.ElementCount : 1;
 
-			UScriptStruct* ScriptStruct = GetScriptStruct(InAddressIndex);
+			UScriptStruct* ScriptStruct = GetScriptStruct(InRegisterIndex);
 			ScriptStruct->InitializeStruct(DataPtr, Count);
 			break;
 		}
-		case ERigVMAddressType::String:
+		case ERigVMRegisterType::String:
 		{
-			FString* DataPtr = (FString*)&Data[InElementIndex == INDEX_NONE ? Address.FirstByte() : Address.FirstByte() + InElementIndex * Address.ElementSize];
-			int32 Count = InElementIndex == INDEX_NONE ? Address.ElementCount : 1;
+			FString* DataPtr = (FString*)&Data[InElementIndex == INDEX_NONE ? Register.FirstByte() : Register.FirstByte() + InElementIndex * Register.ElementSize];
+			int32 Count = InElementIndex == INDEX_NONE ? Register.ElementCount : 1;
 
-			FMemory::Memzero(DataPtr, Count * Address.ElementSize);
+			FMemory::Memzero(DataPtr, Count * Register.ElementSize);
 			for (int32 Index = 0; Index < Count; Index++)
 			{
 				DataPtr[Index] = FString();
 			}
 			break;
 		}
-		case ERigVMAddressType::Name:
+		case ERigVMRegisterType::Name:
 		{
-			FName* DataPtr = (FName*)&Data[InElementIndex == INDEX_NONE ? Address.FirstByte() : Address.FirstByte() + InElementIndex * Address.ElementSize];
-			int32 Count = InElementIndex == INDEX_NONE ? Address.ElementCount : 1;
+			FName* DataPtr = (FName*)&Data[InElementIndex == INDEX_NONE ? Register.FirstByte() : Register.FirstByte() + InElementIndex * Register.ElementSize];
+			int32 Count = InElementIndex == INDEX_NONE ? Register.ElementCount : 1;
 
-			FMemory::Memzero(DataPtr, Count * Address.ElementSize);
+			FMemory::Memzero(DataPtr, Count * Register.ElementSize);
 			for (int32 Index = 0; Index < Count; Index++)
 			{
 				DataPtr[Index] = FName();
@@ -274,26 +274,26 @@ bool FRigVMStorage::Construct(int32 InAddressIndex, int32 InElementIndex)
 	return true;
 }
 
-bool FRigVMStorage::Destroy(int32 InAddressIndex, int32 InElementIndex)
+bool FRigVMStorage::Destroy(int32 InRegisterIndex, int32 InElementIndex)
 {
-	ensure(Addresses.IsValidIndex(InAddressIndex));
+	ensure(Registers.IsValidIndex(InRegisterIndex));
 
-	const FRigVMAddress& Address = Addresses[InAddressIndex];
-	switch (Address.Type)
+	const FRigVMRegister& Register = Registers[InRegisterIndex];
+	switch (Register.Type)
 	{
-		case ERigVMAddressType::Struct:
+		case ERigVMRegisterType::Struct:
 		{
-			void* DataPtr = (void*)&Data[InElementIndex == INDEX_NONE ? Address.FirstByte() : Address.FirstByte() + InElementIndex * Address.ElementSize];
-			int32 Count = InElementIndex == INDEX_NONE ? Address.ElementCount : 1;
+			void* DataPtr = (void*)&Data[InElementIndex == INDEX_NONE ? Register.FirstByte() : Register.FirstByte() + InElementIndex * Register.ElementSize];
+			int32 Count = InElementIndex == INDEX_NONE ? Register.ElementCount : 1;
 
-			UScriptStruct* ScriptStruct = GetScriptStruct(InAddressIndex);
+			UScriptStruct* ScriptStruct = GetScriptStruct(InRegisterIndex);
 			ScriptStruct->DestroyStruct(DataPtr, Count);
 			break;
 		}
-		case ERigVMAddressType::String:
+		case ERigVMRegisterType::String:
 		{
-			FString* DataPtr = (FString*)&Data[InElementIndex == INDEX_NONE ? Address.FirstByte() : Address.FirstByte() + InElementIndex * Address.ElementSize];
-			int32 Count = InElementIndex == INDEX_NONE ? Address.ElementCount : 1;
+			FString* DataPtr = (FString*)&Data[InElementIndex == INDEX_NONE ? Register.FirstByte() : Register.FirstByte() + InElementIndex * Register.ElementSize];
+			int32 Count = InElementIndex == INDEX_NONE ? Register.ElementCount : 1;
 
 			for (int32 Index = 0; Index < Count; Index++)
 			{
@@ -301,10 +301,10 @@ bool FRigVMStorage::Destroy(int32 InAddressIndex, int32 InElementIndex)
 			}
 			break;
 		}
-		case ERigVMAddressType::Name:
+		case ERigVMRegisterType::Name:
 		{
-			FName* DataPtr = (FName*)&Data[InElementIndex == INDEX_NONE ? Address.FirstByte() : Address.FirstByte() + InElementIndex * Address.ElementSize];
-			int32 Count = InElementIndex == INDEX_NONE ? Address.ElementCount : 1;
+			FName* DataPtr = (FName*)&Data[InElementIndex == INDEX_NONE ? Register.FirstByte() : Register.FirstByte() + InElementIndex * Register.ElementSize];
+			int32 Count = InElementIndex == INDEX_NONE ? Register.ElementCount : 1;
 
 			for (int32 Index = 0; Index < Count; Index++)
 			{
@@ -321,48 +321,48 @@ bool FRigVMStorage::Destroy(int32 InAddressIndex, int32 InElementIndex)
 	return true;
 }
 
-bool FRigVMStorage::Remove(int32 InAddressIndex)
+bool FRigVMStorage::Remove(int32 InRegisterIndex)
 {
-	if (InAddressIndex < 0 || InAddressIndex >= Addresses.Num())
+	if (InRegisterIndex < 0 || InRegisterIndex >= Registers.Num())
 	{
 		return false;
 	}
 
-	Destroy(InAddressIndex);
+	Destroy(InRegisterIndex);
 
-	FRigVMAddress AddressToRemove = Addresses[InAddressIndex];
-	Data.RemoveAt(AddressToRemove.ByteIndex, AddressToRemove.NumBytes());
-	Addresses.RemoveAt(InAddressIndex);
+	FRigVMRegister RegisterToRemove = Registers[InRegisterIndex];
+	Data.RemoveAt(RegisterToRemove.ByteIndex, RegisterToRemove.NumBytes());
+	Registers.RemoveAt(InRegisterIndex);
 
-	for (int32 Index = InAddressIndex; Index < Addresses.Num(); Index++)
+	for (int32 Index = InRegisterIndex; Index < Registers.Num(); Index++)
 	{
-		Addresses[Index].ByteIndex -= AddressToRemove.NumBytes();
+		Registers[Index].ByteIndex -= RegisterToRemove.NumBytes();
 	}
 
-	UpdateAddresses();
+	UpdateRegisters();
 	return true;
 }
 
-bool FRigVMStorage::Remove(const FName& InAddressName)
+bool FRigVMStorage::Remove(const FName& InRegisterName)
 {
 	ensure(bUseNameMap);
-	return Remove(GetIndex(InAddressName));
+	return Remove(GetIndex(InRegisterName));
 }
 
-FName FRigVMStorage::Rename(int32 InAddressIndex, const FName& InNewName)
+FName FRigVMStorage::Rename(int32 InRegisterIndex, const FName& InNewName)
 {
-	if (Addresses[InAddressIndex].Name == InNewName)
+	if (Registers[InRegisterIndex].Name == InNewName)
 	{
-		return Addresses[InAddressIndex].Name;
+		return Registers[InRegisterIndex].Name;
 	}
 
 	if (!IsNameAvailable(InNewName))
 	{
-		return Addresses[InAddressIndex].Name;
+		return Registers[InRegisterIndex].Name;
 	}
 
-	Addresses[InAddressIndex].Name = InNewName;
-	UpdateAddresses();
+	Registers[InRegisterIndex].Name = InNewName;
+	UpdateRegisters();
 
 	return InNewName;
 }
@@ -371,95 +371,95 @@ FName FRigVMStorage::Rename(const FName& InOldName, const FName& InNewName)
 {
 	ensure(bUseNameMap);
 
-	int32 AddressIndex = GetIndex(InOldName);
-	if (AddressIndex == INDEX_NONE)
+	int32 RegisterIndex = GetIndex(InOldName);
+	if (RegisterIndex == INDEX_NONE)
 	{
 		return NAME_None;
 	}
 
-	return Rename(AddressIndex, InNewName);
+	return Rename(RegisterIndex, InNewName);
 }
 
-bool FRigVMStorage::Resize(int32 InAddressIndex, int32 InNewElementCount)
+bool FRigVMStorage::Resize(int32 InRegisterIndex, int32 InNewElementCount)
 {
 	if (InNewElementCount <= 0)
 	{
-		return Remove(InAddressIndex);
+		return Remove(InRegisterIndex);
 	}
 
-	if (Addresses[InAddressIndex].ElementCount == InNewElementCount)
+	if (Registers[InRegisterIndex].ElementCount == InNewElementCount)
 	{
 		return false;
 	}
 
-	FRigVMAddress& Address = Addresses[InAddressIndex];
+	FRigVMRegister& Register = Registers[InRegisterIndex];
 
-	if (Address.ElementCount > InNewElementCount) // shrink
+	if (Register.ElementCount > InNewElementCount) // shrink
 	{
-		int32 ElementsToRemove = Address.ElementCount - InNewElementCount;
-		int32 NumBytesToRemove = Address.ElementSize * ElementsToRemove;
-		int32 FirstByteToRemove = Address.FirstByte() + Address.ElementSize * InNewElementCount;
+		int32 ElementsToRemove = Register.ElementCount - InNewElementCount;
+		int32 NumBytesToRemove = Register.ElementSize * ElementsToRemove;
+		int32 FirstByteToRemove = Register.FirstByte() + Register.ElementSize * InNewElementCount;
 
-		for (int32 ElementIndex = InNewElementCount; ElementIndex < Address.ElementCount; ElementIndex++)
+		for (int32 ElementIndex = InNewElementCount; ElementIndex < Register.ElementCount; ElementIndex++)
 		{
-			Destroy(InAddressIndex, ElementIndex);
+			Destroy(InRegisterIndex, ElementIndex);
 		}
 
 		Data.RemoveAt(FirstByteToRemove, NumBytesToRemove);
-		Address.ElementCount = InNewElementCount;
+		Register.ElementCount = InNewElementCount;
 
-		for (int32 AddressIndex = InAddressIndex + 1; AddressIndex < Addresses.Num(); AddressIndex++)
+		for (int32 RegisterIndex = InRegisterIndex + 1; RegisterIndex < Registers.Num(); RegisterIndex++)
 		{
-			Addresses[AddressIndex].ByteIndex -= NumBytesToRemove;
+			Registers[RegisterIndex].ByteIndex -= NumBytesToRemove;
 		}
 	}
 	else // grow
 	{
-		int32 OldElementCount = Address.ElementCount;
-		int32 ElementsToAdd = InNewElementCount - Address.ElementCount;
-		int32 NumBytesToAdd = Address.ElementSize * ElementsToAdd;
-		int32 FirstByteToAdd = Address.FirstByte() + Address.ElementSize * Address.ElementCount;
+		int32 OldElementCount = Register.ElementCount;
+		int32 ElementsToAdd = InNewElementCount - Register.ElementCount;
+		int32 NumBytesToAdd = Register.ElementSize * ElementsToAdd;
+		int32 FirstByteToAdd = Register.FirstByte() + Register.ElementSize * Register.ElementCount;
 
 		Data.InsertZeroed(FirstByteToAdd, NumBytesToAdd);
-		Address.ElementCount = InNewElementCount;
+		Register.ElementCount = InNewElementCount;
 
 		for (int32 ElementIndex = OldElementCount; ElementIndex < InNewElementCount; ElementIndex++)
 		{
-			Construct(InAddressIndex, ElementIndex);
+			Construct(InRegisterIndex, ElementIndex);
 		}
 
-		for (int32 AddressIndex = InAddressIndex + 1; AddressIndex < Addresses.Num(); AddressIndex++)
+		for (int32 RegisterIndex = InRegisterIndex + 1; RegisterIndex < Registers.Num(); RegisterIndex++)
 		{
-			Addresses[AddressIndex].ByteIndex += NumBytesToAdd;
+			Registers[RegisterIndex].ByteIndex += NumBytesToAdd;
 		}
 	}
 
-	UpdateAddresses();
+	UpdateRegisters();
 	return true;
 }
 
-bool FRigVMStorage::Resize(const FName& InAddressName, int32 InNewElementCount)
+bool FRigVMStorage::Resize(const FName& InRegisterName, int32 InNewElementCount)
 {
 	ensure(bUseNameMap);
 
-	int32 AddressIndex = GetIndex(InAddressName);
-	if (AddressIndex == INDEX_NONE)
+	int32 RegisterIndex = GetIndex(InRegisterName);
+	if (RegisterIndex == INDEX_NONE)
 	{
 		return false;
 	}
 
-	return Resize(AddressIndex, InNewElementCount);
+	return Resize(RegisterIndex, InNewElementCount);
 }
 
-void FRigVMStorage::UpdateAddresses()
+void FRigVMStorage::UpdateRegisters()
 {
 	int32 AlignmentShift = 0;
-	for (int32 AddressIndex = 0; AddressIndex < Addresses.Num(); AddressIndex++)
+	for (int32 RegisterIndex = 0; RegisterIndex < Registers.Num(); RegisterIndex++)
 	{
-		FRigVMAddress& Address = Addresses[AddressIndex];
-		Address.ByteIndex += AlignmentShift;
+		FRigVMRegister& Register = Registers[RegisterIndex];
+		Register.ByteIndex += AlignmentShift;
 
-		UScriptStruct* ScriptStruct = GetScriptStruct(AddressIndex);
+		UScriptStruct* ScriptStruct = GetScriptStruct(RegisterIndex);
 		if (ScriptStruct != nullptr)
 		{
 			UScriptStruct::ICppStructOps* TheCppStructOps = ScriptStruct->GetCppStructOps();
@@ -467,50 +467,50 @@ void FRigVMStorage::UpdateAddresses()
 			{
 				if (!TheCppStructOps->HasZeroConstructor())
 				{
-					void* Pointer = GetData(AddressIndex);
+					void* Pointer = GetData(RegisterIndex);
 
-					if (Address.AlignmentBytes > 0)
+					if (Register.AlignmentBytes > 0)
 					{
 						if (!IsAligned(Pointer, TheCppStructOps->GetAlignment()))
 						{
-							Data.RemoveAt(Address.ByteIndex, Address.AlignmentBytes);
-							AlignmentShift -= Address.AlignmentBytes;
-							Address.AlignmentBytes = 0;
-							Pointer = GetData(AddressIndex);
+							Data.RemoveAt(Register.ByteIndex, Register.AlignmentBytes);
+							AlignmentShift -= Register.AlignmentBytes;
+							Register.AlignmentBytes = 0;
+							Pointer = GetData(RegisterIndex);
 						}
 					}
 
 					while (!IsAligned(Pointer, TheCppStructOps->GetAlignment()))
 					{
-						Data.InsertZeroed(Address.ByteIndex, 1);
-						Address.AlignmentBytes++;
+						Data.InsertZeroed(Register.ByteIndex, 1);
+						Register.AlignmentBytes++;
 						AlignmentShift++;
-						Pointer = GetData(AddressIndex);
+						Pointer = GetData(RegisterIndex);
 					}
 				}
 			}
 		}
 	}
 
-	for (int32 AddressIndex = 0; AddressIndex < Addresses.Num(); AddressIndex++)
+	for (int32 RegisterIndex = 0; RegisterIndex < Registers.Num(); RegisterIndex++)
 	{
-		Addresses[AddressIndex].Pointer = GetData(AddressIndex);
+		Registers[RegisterIndex].Pointer = GetData(RegisterIndex);
 	}
 
 	if (bUseNameMap)
 	{
 		NameMap.Reset();
-		for (int32 Index = 0; Index < Addresses.Num(); Index++)
+		for (int32 Index = 0; Index < Registers.Num(); Index++)
 		{
-			NameMap.Add(Addresses[Index].Name, Index);
+			NameMap.Add(Registers[Index].Name, Index);
 		}
 	}
 }
 
-void FRigVMStorage::FillWithZeroes(int32 InAddressIndex)
+void FRigVMStorage::FillWithZeroes(int32 InRegisterIndex)
 {
-	ensure(Addresses.IsValidIndex(InAddressIndex));
-	FMemory::Memzero(GetData(InAddressIndex), Addresses[InAddressIndex].NumBytes());
+	ensure(Registers.IsValidIndex(InRegisterIndex));
+	FMemory::Memzero(GetData(InRegisterIndex), Registers[InRegisterIndex].NumBytes());
 }
 
 int32 FRigVMStorage::FindOrAddScriptStruct(UScriptStruct* InScriptStruct)
diff --git a/Engine/Source/Runtime/RigVM/Public/RigVMByteCode.h b/Engine/Source/Runtime/RigVM/Public/RigVMByteCode.h
--- a/Engine/Source/Runtime/RigVM/Public/RigVMByteCode.h
+++ b/Engine/Source/Runtime/RigVM/Public/RigVMByteCode.h
@@ -364,9 +364,9 @@ private:
 	template<class OpType>
 	FORCEINLINE uint64 AddOp(const OpType& InOp)
 	{
-		uint64 Address = (uint64)ByteCode.AddUninitialized(sizeof(OpType));
-		FMemory::Memcpy(ByteCode.GetData() + Address, &InOp, sizeof(OpType));
-		return Address;
+		uint64 ByteIndex = (uint64)ByteCode.AddUninitialized(sizeof(OpType));
+		FMemory::Memcpy(ByteCode.GetData() + ByteIndex, &InOp, sizeof(OpType));
+		return ByteIndex;
 	}
 
 	// storage for all functions
diff --git a/Engine/Source/Runtime/RigVM/Public/RigVMStorage.h b/Engine/Source/Runtime/RigVM/Public/RigVMStorage.h
--- a/Engine/Source/Runtime/RigVM/Public/RigVMStorage.h
+++ b/Engine/Source/Runtime/RigVM/Public/RigVMStorage.h
@@ -39,7 +39,7 @@ private:
 typedef TArrayView<FRigVMArgument> FRigVMArgumentArray;
 
 UENUM()
-enum class ERigVMAddressType : uint8
+enum class ERigVMRegisterType : uint8
 {
 	Plain,
 	String,
@@ -49,12 +49,12 @@ enum class ERigVMAddressType : uint8
 };
 
 USTRUCT()
-struct RIGVM_API FRigVMAddress
+struct RIGVM_API FRigVMRegister
 {
 	GENERATED_USTRUCT_BODY()
 
-		FRigVMAddress()
-		: Type(ERigVMAddressType::Invalid)
+		FRigVMRegister()
+		: Type(ERigVMRegisterType::Invalid)
 		, Pointer(nullptr)
 		, ByteIndex(INDEX_NONE)
 		, ElementSize(0)
@@ -66,7 +66,7 @@ struct RIGVM_API FRigVMAddress
 	}
 
 	UPROPERTY()
-	ERigVMAddressType Type;
+	ERigVMRegisterType Type;
 
 	void* Pointer;
 
@@ -127,7 +127,7 @@ struct RIGVM_API FRigVMAddress
 	}
 };
 
-typedef TArrayView<FRigVMAddress> FRigVMAddressArray;
+typedef TArrayView<FRigVMRegister> FRigVMRegisterArray;
 
 USTRUCT()
 struct RIGVM_API FRigVMStorage
@@ -145,48 +145,48 @@ public:
 	FORCEINLINE bool IsLiteralStorage() const { return bIsLiteralStorage;  }
 	FORCEINLINE void SetLiteralStorage(bool InIsLiteralStorage = true) { bIsLiteralStorage = InIsLiteralStorage; }
 	FORCEINLINE bool SupportsNames() const { return bUseNameMap;  }
-	FORCEINLINE int32 Num() const { return Addresses.Num(); }
-	FORCEINLINE const FRigVMAddress& operator[](int32 InIndex) const { return Addresses[InIndex]; }
-	FORCEINLINE FRigVMAddress& operator[](int32 InIndex) { return Addresses[InIndex]; }
-	FORCEINLINE const FRigVMAddress& operator[](const FRigVMArgument& InArg) const { return Addresses[InArg.GetRegisterIndex()]; }
-	FORCEINLINE FRigVMAddress& operator[](const FRigVMArgument& InArg) { return Addresses[InArg.GetRegisterIndex()]; }
-	FORCEINLINE const FRigVMAddress& operator[](const FName& InName) const { return Addresses[GetIndex(InName)]; }
-	FORCEINLINE FRigVMAddress& operator[](const FName& InName) { return Addresses[GetIndex(InName)]; }
+	FORCEINLINE int32 Num() const { return Registers.Num(); }
+	FORCEINLINE const FRigVMRegister& operator[](int32 InIndex) const { return Registers[InIndex]; }
+	FORCEINLINE FRigVMRegister& operator[](int32 InIndex) { return Registers[InIndex]; }
+	FORCEINLINE const FRigVMRegister& operator[](const FRigVMArgument& InArg) const { return Registers[InArg.GetRegisterIndex()]; }
+	FORCEINLINE FRigVMRegister& operator[](const FRigVMArgument& InArg) { return Registers[InArg.GetRegisterIndex()]; }
+	FORCEINLINE const FRigVMRegister& operator[](const FName& InName) const { return Registers[GetIndex(InName)]; }
+	FORCEINLINE FRigVMRegister& operator[](const FName& InName) { return Registers[GetIndex(InName)]; }
 
-	FORCEINLINE TArray<FRigVMAddress>::RangedForIteratorType      begin() { return Addresses.begin(); }
-	FORCEINLINE TArray<FRigVMAddress>::RangedForConstIteratorType begin() const { return Addresses.begin(); }
-	FORCEINLINE TArray<FRigVMAddress>::RangedForIteratorType      end() { return Addresses.end(); }
-	FORCEINLINE TArray<FRigVMAddress>::RangedForConstIteratorType end() const { return Addresses.end(); }
+	FORCEINLINE TArray<FRigVMRegister>::RangedForIteratorType      begin() { return Registers.begin(); }
+	FORCEINLINE TArray<FRigVMRegister>::RangedForConstIteratorType begin() const { return Registers.begin(); }
+	FORCEINLINE TArray<FRigVMRegister>::RangedForIteratorType      end() { return Registers.end(); }
+	FORCEINLINE TArray<FRigVMRegister>::RangedForConstIteratorType end() const { return Registers.end(); }
 
-	FORCEINLINE FRigVMArgument GetArgument(int32 InAddressIndex) const
+	FORCEINLINE FRigVMArgument GetArgument(int32 InRegisterIndex) const
 	{
-		ensure(Addresses.IsValidIndex(InAddressIndex));
-		return FRigVMArgument(IsLiteralStorage(), InAddressIndex, Addresses[InAddressIndex].ByteIndex);
+		ensure(Registers.IsValidIndex(InRegisterIndex));
+		return FRigVMArgument(IsLiteralStorage(), InRegisterIndex, Registers[InRegisterIndex].ByteIndex);
 	}
 
-	FORCEINLINE const void* GetData(int32 InAddressIndex) const
+	FORCEINLINE const void* GetData(int32 InRegisterIndex) const
 	{
-		ensure(Addresses.IsValidIndex(InAddressIndex));
-		const FRigVMAddress& Address = Addresses[InAddressIndex];
-		ensure(Address.ElementCount > 0);
-		return (const void*)&Data[Address.FirstByte()];
+		ensure(Registers.IsValidIndex(InRegisterIndex));
+		const FRigVMRegister& Register = Registers[InRegisterIndex];
+		ensure(Register.ElementCount > 0);
+		return (const void*)&Data[Register.FirstByte()];
 	}
 
-	FORCEINLINE void* GetData(int32 InAddressIndex)
+	FORCEINLINE void* GetData(int32 InRegisterIndex)
 	{
-		ensure(Addresses.IsValidIndex(InAddressIndex));
-		const FRigVMAddress& Address = Addresses[InAddressIndex];
-		ensure(Address.ElementCount > 0);
-		return (void*)&Data[Address.FirstByte()];
+		ensure(Registers.IsValidIndex(InRegisterIndex));
+		const FRigVMRegister& Register = Registers[InRegisterIndex];
+		ensure(Register.ElementCount > 0);
+		return (void*)&Data[Register.FirstByte()];
 	}
 
 	template<class T>
-	FORCEINLINE const T* Get(int32 InAddressIndex) const
+	FORCEINLINE const T* Get(int32 InRegisterIndex) const
 	{
-		ensure(Addresses.IsValidIndex(InAddressIndex));
-		const FRigVMAddress& Address = Addresses[InAddressIndex];
-		ensure(Address.ElementCount > 0);
-		return (const T*)&Data[Address.FirstByte()];
+		ensure(Registers.IsValidIndex(InRegisterIndex));
+		const FRigVMRegister& Register = Registers[InRegisterIndex];
+		ensure(Register.ElementCount > 0);
+		return (const T*)&Data[Register.FirstByte()];
 	}
 
 	template<class T>
@@ -196,9 +196,9 @@ public:
 	}
 
 	template<class T>
-	FORCEINLINE const T& GetRef(int32 InAddressIndex) const
+	FORCEINLINE const T& GetRef(int32 InRegisterIndex) const
 	{
-		return *Get<T>(InAddressIndex);
+		return *Get<T>(InRegisterIndex);
 	}
 
 	template<class T>
@@ -208,12 +208,12 @@ public:
 	}
 
 	template<class T>
-	FORCEINLINE T* Get(int32 InAddressIndex)
+	FORCEINLINE T* Get(int32 InRegisterIndex)
 	{
-		ensure(Addresses.IsValidIndex(InAddressIndex));
-		const FRigVMAddress& Address = Addresses[InAddressIndex];
-		ensure(Address.ElementCount > 0);
-		return (T*)&Data[Address.FirstByte()];
+		ensure(Registers.IsValidIndex(InRegisterIndex));
+		const FRigVMRegister& Register = Registers[InRegisterIndex];
+		ensure(Register.ElementCount > 0);
+		return (T*)&Data[Register.FirstByte()];
 	}
 
 	template<class T>
@@ -223,9 +223,9 @@ public:
 	}
 
 	template<class T>
-	FORCEINLINE T& GetRef(int32 InAddressIndex)
+	FORCEINLINE T& GetRef(int32 InRegisterIndex)
 	{
-		return *Get<T>(InAddressIndex);
+		return *Get<T>(InRegisterIndex);
 	}
 
 	template<class T>
@@ -235,12 +235,12 @@ public:
 	}
 
 	template<class T>
-	FORCEINLINE TArrayView<T> GetArray(int32 InAddressIndex)
+	FORCEINLINE TArrayView<T> GetArray(int32 InRegisterIndex)
 	{
-		ensure(Addresses.IsValidIndex(InAddressIndex));
-		const FRigVMAddress& Address = Addresses[InAddressIndex];
-		ensure(Address.ElementCount > 0);
-		return TArrayView<T>((T*)&Data[Address.FirstByte()], Address.ElementCount);
+		ensure(Registers.IsValidIndex(InRegisterIndex));
+		const FRigVMRegister& Register = Registers[InRegisterIndex];
+		ensure(Register.ElementCount > 0);
+		return TArrayView<T>((T*)&Data[Register.FirstByte()], Register.ElementCount);
 	}
 	
 	template<class T>
@@ -249,21 +249,21 @@ public:
 		return GetArray<T>(InArgument.GetRegisterIndex());
 	}
 
-	FORCEINLINE UScriptStruct* GetScriptStruct(int32 InAddressIndex) const
+	FORCEINLINE UScriptStruct* GetScriptStruct(int32 InRegisterIndex) const
 	{
-		ensure(Addresses.IsValidIndex(InAddressIndex));
-		const FRigVMAddress& Address = Addresses[InAddressIndex];
-		if (Address.ScriptStructIndex != INDEX_NONE)
+		ensure(Registers.IsValidIndex(InRegisterIndex));
+		const FRigVMRegister& Register = Registers[InRegisterIndex];
+		if (Register.ScriptStructIndex != INDEX_NONE)
 		{
-			ensure(ScriptStructs.IsValidIndex(Address.ScriptStructIndex));
-			return ScriptStructs[Address.ScriptStructIndex];
+			ensure(ScriptStructs.IsValidIndex(Register.ScriptStructIndex));
+			return ScriptStructs[Register.ScriptStructIndex];
 		}
 		return nullptr;
 	}
 
 	bool Copy(
-		int32 InSourceAddressIndex,
-		int32 InTargetAddressIndex,
+		int32 InSourceRegisterIndex,
+		int32 InTargetRegisterIndex,
 		const FRigVMStorage* InSourceStorage = nullptr,
 		int32 InSourceByteOffset = INDEX_NONE,
 		int32 InTargetByteOffset = INDEX_NONE,
@@ -284,11 +284,11 @@ public:
 			return INDEX_NONE;
 		}
 
-		if (NameMap.Num() != Addresses.Num())
+		if (NameMap.Num() != Registers.Num())
 		{
-			for (int32 Index = 0; Index < Addresses.Num(); Index++)
+			for (int32 Index = 0; Index < Registers.Num(); Index++)
 			{
-				if (Addresses[Index].Name == InName)
+				if (Registers[Index].Name == InName)
 				{
 					return Index;
 				}
@@ -364,21 +364,21 @@ public:
 
 	FORCEINLINE int32 AddNameArray(const FName& InNewName, int32 InCount, const FName* InDataPtr = nullptr)
 	{
-		int32 Address = Allocate(InNewName, sizeof(FName), InCount, nullptr);
-		Addresses[Address].Type = ERigVMAddressType::Name;
+		int32 Register = Allocate(InNewName, sizeof(FName), InCount, nullptr);
+		Registers[Register].Type = ERigVMRegisterType::Name;
 
-		Construct(Address);
+		Construct(Register);
 
 		if(InDataPtr)
 		{
-			FName* DataPtr = (FName*)GetData(Address);
+			FName* DataPtr = (FName*)GetData(Register);
 			for (int32 Index = 0; Index < InCount; Index++)
 			{
 				DataPtr[Index] = InDataPtr[Index];
 			}
 		}
 
-		return Address;
+		return Register;
 	}
 
 	FORCEINLINE int32 AddNameArray(const FName& InNewName, const TArray<FName>& InArray)
@@ -403,21 +403,21 @@ public:
 
 	FORCEINLINE int32 AddStringArray(const FName& InNewName, int32 InCount, const FString* InDataPtr = nullptr)
 	{
-		int32 Address = Allocate(InNewName, sizeof(FString), InCount, nullptr);
-		Addresses[Address].Type = ERigVMAddressType::String;
+		int32 Register = Allocate(InNewName, sizeof(FString), InCount, nullptr);
+		Registers[Register].Type = ERigVMRegisterType::String;
 
-		Construct(Address);
+		Construct(Register);
 
 		if(InDataPtr)
 		{
-			FString* DataPtr = (FString*)GetData(Address);
+			FString* DataPtr = (FString*)GetData(Register);
 			for (int32 Index = 0; Index < InCount; Index++)
 			{
 				DataPtr[Index] = InDataPtr[Index];
 			}
 		}
 
-		return Address;
+		return Register;
 	}
 
 	FORCEINLINE int32 AddStringArray(const FName& InNewName, const TArray<FString>& InArray)
@@ -443,27 +443,27 @@ public:
 
 	FORCEINLINE int32 AddStructArray(const FName& InNewName, UScriptStruct* InScriptStruct, int32 InCount, const void* InDataPtr = nullptr)
 	{
-		int32 Address = Allocate(InNewName, InScriptStruct->GetStructureSize(), InCount, nullptr, false);
-		if (Address == INDEX_NONE)
+		int32 Register = Allocate(InNewName, InScriptStruct->GetStructureSize(), InCount, nullptr, false);
+		if (Register == INDEX_NONE)
 		{
 			return INDEX_NONE;
 		}
 
-		Addresses[Address].Type = ERigVMAddressType::Struct;
-		Addresses[Address].ScriptStructIndex = FindOrAddScriptStruct(InScriptStruct);
+		Registers[Register].Type = ERigVMRegisterType::Struct;
+		Registers[Register].ScriptStructIndex = FindOrAddScriptStruct(InScriptStruct);
 
-		UpdateAddresses();
+		UpdateRegisters();
 
 		// construct the content
-		Construct(Address);
+		Construct(Register);
 
 		// copy values from the provided data
 		if (InDataPtr != nullptr)
 		{
-			InScriptStruct->CopyScriptStruct(GetData(Address), InDataPtr, InCount);
+			InScriptStruct->CopyScriptStruct(GetData(Register), InDataPtr, InCount);
 		}
 
-		return Address;
+		return Register;
 	}
 
 	FORCEINLINE int32 AddStructArray(UScriptStruct* InScriptStruct, int32 InCount, const void* InDataPtr = nullptr)
@@ -518,23 +518,23 @@ public:
 		return AddStruct<T>(NAME_None, InValue);
 	}
 
-	bool Remove(int32 InAddressIndex);
-	bool Remove(const FName& InAddressName);
-	FName Rename(int32 InAddressIndex, const FName& InNewName);
+	bool Remove(int32 InRegisterIndex);
+	bool Remove(const FName& InRegisterName);
+	FName Rename(int32 InRegisterIndex, const FName& InNewName);
 	FName Rename(const FName& InOldName, const FName& InNewName);
-	bool Resize(int32 InAddressIndex, int32 InNewElementCount);
-	bool Resize(const FName& InAddressName, int32 InNewElementCount);
+	bool Resize(int32 InRegisterIndex, int32 InNewElementCount);
+	bool Resize(const FName& InRegisterName, int32 InNewElementCount);
 
-	void UpdateAddresses();
+	void UpdateRegisters();
 
 private:
 
-	int32 Allocate(const FName& InNewName, int32 InElementSize, int32 InCount, const void* InDataPtr = nullptr, bool bUpdateAddresses = true);
-	int32 Allocate(int32 InElementSize, int32 InCount, const void* InDataPtr = nullptr, bool bUpdateAddresses = true);
-	bool Construct(int32 InAddressIndex, int32 InElementIndex = INDEX_NONE);
-	bool Destroy(int32 InAddressIndex, int32 InElementIndex = INDEX_NONE);
+	int32 Allocate(const FName& InNewName, int32 InElementSize, int32 InCount, const void* InDataPtr = nullptr, bool bUpdateRegisters = true);
+	int32 Allocate(int32 InElementSize, int32 InCount, const void* InDataPtr = nullptr, bool bUpdateRegisters = true);
+	bool Construct(int32 InRegisterIndex, int32 InElementIndex = INDEX_NONE);
+	bool Destroy(int32 InRegisterIndex, int32 InElementIndex = INDEX_NONE);
 
-	void FillWithZeroes(int32 InAddressIndex);
+	void FillWithZeroes(int32 InRegisterIndex);
 	int32 FindOrAddScriptStruct(UScriptStruct* InScriptStruct);
 
 	UPROPERTY()
@@ -544,7 +544,7 @@ private:
 	bool bIsLiteralStorage;
 
 	UPROPERTY()
-	TArray<FRigVMAddress> Addresses;
+	TArray<FRigVMRegister> Registers;
 
 	UPROPERTY()
 	TArray<uint8> Data;
