RigVM: Moving core functionality below RigVM/RigVMCore/

diff --git a/Engine/Plugins/Experimental/ControlRig/Source/ControlRig/Public/Units/RigUnit.h b/Engine/Plugins/Experimental/ControlRig/Source/ControlRig/Public/Units/RigUnit.h
--- a/Engine/Plugins/Experimental/ControlRig/Source/ControlRig/Public/Units/RigUnit.h
+++ b/Engine/Plugins/Experimental/ControlRig/Source/ControlRig/Public/Units/RigUnit.h
@@ -3,7 +3,7 @@
 #pragma once
 
 #include "ControlRigDefines.h"
-#include "RigVMRegistry.h"
+#include "RigVMCore/RigVMRegistry.h"
 #include "RigUnit.generated.h"
 
 struct FRigUnitContext;
diff --git a/Engine/Source/Runtime/RigVM/Private/RigVM.cpp b/Engine/Source/Runtime/RigVM/Private/RigVM.cpp
deleted file mode 100644
diff --git a/Engine/Source/Runtime/RigVM/Private/RigVMByteCode.cpp b/Engine/Source/Runtime/RigVM/Private/RigVMByteCode.cpp
deleted file mode 100644
diff --git a/Engine/Source/Runtime/RigVM/Private/RigVMCore/RigVM.cpp b/Engine/Source/Runtime/RigVM/Private/RigVMCore/RigVM.cpp
new file mode 100644
--- /dev/null
+++ b/Engine/Source/Runtime/RigVM/Private/RigVMCore/RigVM.cpp
@@ -0,0 +1,163 @@
+// Copyright 1998-2019 Epic Games, Inc. All Rights Reserved.
+
+#include "RigVMCore/RigVM.h"
+
+URigVM::URigVM()
+{
+	Literals.SetLiteralStorage(true);
+	WorkState.SetLiteralStorage(false);
+}
+
+URigVM::~URigVM()
+{
+	Reset();
+}
+
+void URigVM::Reset()
+{
+	Literals.Reset();
+	WorkState.Reset();
+	FunctionNames.Reset();
+	Functions.Reset();
+	ByteCode.Reset();
+	Instructions.Reset();
+}
+
+int32 URigVM::AddRigVMFunction(UScriptStruct* InRigVMStruct, const FName& InFunctionName)
+{
+	check(InRigVMStruct);
+	FString FunctionKey = FString::Printf(TEXT("F%s::%s"), *InRigVMStruct->GetName(), *InFunctionName.ToString());
+	int32 FunctionIndex = FunctionNames.Find(FunctionKey);
+	if (FunctionIndex != INDEX_NONE)
+	{
+		return FunctionIndex;
+	}
+
+	FRigVMFunctionPtr Function = FRigVMRegistry::Get().Find(*FunctionKey);
+	if (Function == nullptr)
+	{
+		return INDEX_NONE;
+	}
+
+	FunctionNames.Add(FunctionKey);
+	return Functions.Add(Function);
+}
+
+void URigVM::ResolveFunctionsIfRequired()
+{
+	if (Functions.Num() != FunctionNames.Num())
+	{
+		Functions.Reset();
+		Functions.SetNumZeroed(FunctionNames.Num());
+
+		for (int32 FunctionIndex = 0; FunctionIndex < FunctionNames.Num(); FunctionIndex)
+		{
+			Functions[FunctionIndex] = FRigVMRegistry::Get().Find(*FunctionNames[FunctionIndex]);
+		}
+	}
+}
+
+void URigVM::RefreshInstructionsIfRequired()
+{
+	if (Instructions.Num() == 0)
+	{
+		Instructions = ByteCode.GetTable();
+	}
+}
+
+bool URigVM::Execute(FRigVMStoragePtrArray Storage, TArrayView<void*> AdditionalArgs)
+{
+	ResolveFunctionsIfRequired();
+	RefreshInstructionsIfRequired();
+
+	if (Instructions.Num() == 0)
+	{
+		return true;
+	}
+
+	FRigVMStorage* LocalStorage[] = { &WorkState, &Literals }; 
+	if (Storage.Num() == 0)
+	{
+		Storage = FRigVMStoragePtrArray(LocalStorage, 2);
+	}
+
+	int32 InstructionIndex = 0;
+	while (Instructions[InstructionIndex].OpCode != ERigVMOpCode::Exit)
+	{
+		switch (Instructions[InstructionIndex].OpCode)
+		{
+			case ERigVMOpCode::Copy:
+			{
+				const FRigVMCopyOp& Op = ByteCode.GetOpAt<FRigVMCopyOp>(Instructions[InstructionIndex]);
+				Storage[0]->Copy(Op.Source.GetRegisterIndex(), Op.Target.GetRegisterIndex(), Storage[Op.Source.StorageType()], Op.SourceOffset, Op.TargetOffset, Op.NumBytes);
+				InstructionIndex++;
+				break;
+			}
+			case ERigVMOpCode::Increment:
+			{
+				const FRigVMIncrementOp& Op = ByteCode.GetOpAt<FRigVMIncrementOp>(Instructions[InstructionIndex]);
+				Storage[0]->GetRef<int32>(Op.Arg.GetRegisterIndex())++;
+				InstructionIndex++;
+				break;
+			}
+			case ERigVMOpCode::Decrement:
+			{
+				const FRigVMDecrementOp& Op = ByteCode.GetOpAt<FRigVMDecrementOp>(Instructions[InstructionIndex]);
+				Storage[0]->GetRef<int32>(Op.Arg.GetRegisterIndex())--;
+				InstructionIndex++;
+				break;
+			}
+			case ERigVMOpCode::Equals:
+			{
+				const FRigVMEqualsOp& Op = ByteCode.GetOpAt<FRigVMEqualsOp>(Instructions[InstructionIndex]);
+				ensureMsgf(false, TEXT("TO BE IMPLEMENTED"));
+				break;
+			}
+			case ERigVMOpCode::NotEquals:
+			{
+				const FRigVMNotEqualsOp& Op = ByteCode.GetOpAt<FRigVMNotEqualsOp>(Instructions[InstructionIndex]);
+				ensureMsgf(false, TEXT("TO BE IMPLEMENTED"));
+				break;
+			}
+			case ERigVMOpCode::Jump:
+			{
+				const FRigVMJumpOp& Op = ByteCode.GetOpAt<FRigVMJumpOp>(Instructions[InstructionIndex]);
+				ensureMsgf(false, TEXT("TO BE IMPLEMENTED"));
+				break;
+			}
+			case ERigVMOpCode::JumpIfTrue:
+			{
+				const FRigVMJumpIfTrueOp& Op = ByteCode.GetOpAt<FRigVMJumpIfTrueOp>(Instructions[InstructionIndex]);
+				ensureMsgf(false, TEXT("TO BE IMPLEMENTED"));
+				break;
+			}
+			case ERigVMOpCode::JumpIfFalse:
+			{
+				const FRigVMJumpIfFalseOp& Op = ByteCode.GetOpAt<FRigVMJumpIfFalseOp>(Instructions[InstructionIndex]);
+				ensureMsgf(false, TEXT("TO BE IMPLEMENTED"));
+				break;
+			}
+			case ERigVMOpCode::Execute:
+			{
+				const FRigVMExecuteOp& Op = ByteCode.GetOpAt<FRigVMExecuteOp>(Instructions[InstructionIndex]);
+				FRigVMArgumentArray Args = ByteCode.GetArgumentsForExecuteOp(Instructions[InstructionIndex]);
+				(*Functions[Op.FunctionIndex])(Args, Storage, AdditionalArgs);
+				InstructionIndex++;
+				break;
+			}
+			case ERigVMOpCode::Exit:
+			case ERigVMOpCode::Invalid:
+			{
+				ensure(false);
+				return false;
+			}
+		}
+
+		if (!Instructions.IsValidIndex(InstructionIndex))
+		{
+			break;
+		}
+	}
+
+	return true;
+}
diff --git a/Engine/Source/Runtime/RigVM/Private/RigVMCore/RigVMByteCode.cpp b/Engine/Source/Runtime/RigVM/Private/RigVMCore/RigVMByteCode.cpp
new file mode 100644
--- /dev/null
+++ b/Engine/Source/Runtime/RigVM/Private/RigVMCore/RigVMByteCode.cpp
@@ -0,0 +1,161 @@
+// Copyright 1998-2019 Epic Games, Inc. All Rights Reserved.
+
+#include "RigVMCore/RigVMByteCode.h"
+
+FRigVMByteCodeTable::FRigVMByteCodeTable()
+{
+}
+
+FRigVMByteCodeTable::FRigVMByteCodeTable(const FRigVMByteCode& InByteCode)
+{
+	uint64 ByteIndex = 0;
+	while (ByteIndex < InByteCode.Num())
+	{
+		ERigVMOpCode OpCode = InByteCode.GetOpCodeAt(ByteIndex);
+		Entries.Add(FRigVMByteCodeTableEntry(OpCode, ByteIndex));
+		ByteIndex += InByteCode.GetOpNumBytesAt(ByteIndex);
+	}
+}
+
+void FRigVMByteCodeTable::Reset()
+{
+	Entries.Reset();
+}
+
+FRigVMByteCode::FRigVMByteCode()
+{
+}
+
+void FRigVMByteCode::Reset()
+{
+	ByteCode.Reset();
+}
+
+uint64 FRigVMByteCode::Num() const
+{
+	return (uint64)ByteCode.Num();
+}
+
+uint64 FRigVMByteCode::GetOpNumBytesAt(uint64 InByteCodeIndex, bool bIncludeArguments) const
+{
+	ERigVMOpCode OpCode = GetOpCodeAt(InByteCodeIndex);
+	switch (OpCode)
+	{
+		case ERigVMOpCode::Copy:
+		{
+			return (uint64)sizeof(FRigVMCopyOp);
+		}
+		case ERigVMOpCode::Increment:
+		{
+			return (uint64)sizeof(FRigVMIncrementOp);
+		}
+		case ERigVMOpCode::Decrement:
+		{
+			return (uint64)sizeof(FRigVMDecrementOp);
+		}
+		case ERigVMOpCode::Equals:
+		{
+			return (uint64)sizeof(FRigVMEqualsOp);
+		}
+		case ERigVMOpCode::NotEquals:
+		{
+			return (uint64)sizeof(FRigVMNotEqualsOp);
+		}
+		case ERigVMOpCode::Jump:
+		{
+			return (uint64)sizeof(FRigVMJumpOp);
+		}
+		case ERigVMOpCode::JumpIfTrue:
+		{
+			return (uint64)sizeof(FRigVMJumpIfTrueOp);
+		}
+		case ERigVMOpCode::JumpIfFalse:
+		{
+			return (uint64)sizeof(FRigVMJumpIfFalseOp);
+		}
+		case ERigVMOpCode::Execute:
+		{
+			uint64 NumBytes = (uint64)sizeof(FRigVMExecuteOp);
+			if(bIncludeArguments)
+			{
+				const FRigVMExecuteOp& ExecuteOp = GetOpAt<FRigVMExecuteOp>(InByteCodeIndex);
+				NumBytes += (uint64)ExecuteOp.ArgumentCount * (uint64)sizeof(FRigVMArgument);
+			}
+			return NumBytes;
+		}
+		case ERigVMOpCode::Exit:
+		{
+			return (uint64)sizeof(FRigVMExitOp);
+		}
+		case ERigVMOpCode::Invalid:
+		{
+			ensure(false);
+			return 0;
+		}
+	}
+	return 0;
+}
+
+uint64 FRigVMByteCode::AddCopyOp(const FRigVMArgument& InSource, const FRigVMArgument& InTarget, int32 InSourceOffset, int32 InTargetOffset, int32 InNumBytes)
+{
+	FRigVMCopyOp Op(InSource, InTarget, InSourceOffset, InTargetOffset, InNumBytes);
+	return AddOp(Op);
+}
+
+uint64 FRigVMByteCode::AddIncrementOp(const FRigVMArgument& InArg)
+{
+	FRigVMIncrementOp Op(InArg);
+	return AddOp(Op);
+}
+
+uint64 FRigVMByteCode::AddDecrementOp(const FRigVMArgument& InArg)
+{
+	FRigVMDecrementOp Op(InArg);
+	return AddOp(Op);
+}
+
+uint64 FRigVMByteCode::AddEqualsOp(const FRigVMArgument& InA, const FRigVMArgument& InB, const FRigVMArgument& InResult)
+{
+	FRigVMEqualsOp Op(InA, InB, InResult);
+	return AddOp(Op);
+}
+
+uint64 FRigVMByteCode::AddNotEqualsOp(const FRigVMArgument& InA, const FRigVMArgument& InB, const FRigVMArgument& InResult)
+{
+	FRigVMNotEqualsOp Op(InA, InB, InResult);
+	return AddOp(Op);
+}
+
+uint64 FRigVMByteCode::AddJumpOp(uint64 InByteCodeIndex)
+{
+	FRigVMJumpOp Op(InByteCodeIndex);
+	return AddOp(Op);
+}
+
+uint64 FRigVMByteCode::AddJumpIfTrueOp(uint64 InByteCodeIndex, const FRigVMArgument& InCondition)
+{
+	FRigVMJumpIfTrueOp Op(InByteCodeIndex, InCondition);
+	return AddOp(Op);
+}
+
+uint64 FRigVMByteCode::AddJumpIfFalseOp(uint64 InByteCodeIndex, const FRigVMArgument& InCondition)
+{
+	FRigVMJumpIfFalseOp Op(InByteCodeIndex, InCondition);
+	return AddOp(Op);
+}
+
+uint64 FRigVMByteCode::AddExecuteOp(uint16 InFunctionIndex, const TArrayView<FRigVMArgument>& InArguments)
+{
+	FRigVMExecuteOp Op(InFunctionIndex, (uint16)InArguments.Num());
+	uint64 OpByteIndex = AddOp(Op);
+	uint64 ArgumentsByteIndex = (uint64)ByteCode.AddUninitialized(sizeof(FRigVMArgument) * InArguments.Num());
+	FMemory::Memcpy(ByteCode.GetData() + ArgumentsByteIndex, InArguments.GetData(), sizeof(FRigVMArgument) * InArguments.Num());
+	return OpByteIndex;
+}
+
+uint64 FRigVMByteCode::AddExitOp()
+{
+	FRigVMExitOp Op;
+	return AddOp(Op);
+}
+
diff --git a/Engine/Source/Runtime/RigVM/Private/RigVMCore/RigVMRegistry.cpp b/Engine/Source/Runtime/RigVM/Private/RigVMCore/RigVMRegistry.cpp
new file mode 100644
--- /dev/null
+++ b/Engine/Source/Runtime/RigVM/Private/RigVMCore/RigVMRegistry.cpp
@@ -0,0 +1,37 @@
+// Copyright 1998-2019 Epic Games, Inc. All Rights Reserved.
+
+#include "RigVMCore/RigVMRegistry.h"
+#include "UObject/UObjectIterator.h"
+
+FRigVMRegistry FRigVMRegistry::s_RigVMRegistry;
+
+FRigVMRegistry& FRigVMRegistry::Get()
+{
+	return s_RigVMRegistry;
+}
+
+void FRigVMRegistry::Refresh()
+{
+}
+
+void FRigVMRegistry::Register(const TCHAR* InName, FRigVMFunctionPtr InFunctionPtr)
+{
+	if (Find(InName) != nullptr)
+	{
+		return;
+	}
+	Functions.Add(FRigVMFunction(InName, InFunctionPtr));
+}
+
+FRigVMFunctionPtr FRigVMRegistry::Find(const TCHAR* InName) const
+{
+	for (const FRigVMFunction& Function : Functions)
+	{
+		if (FCString::Strcmp(Function.Name, InName) == 0)
+		{
+			return Function.FunctionPtr;
+		}
+	}
+
+	return nullptr;
+}
diff --git a/Engine/Source/Runtime/RigVM/Private/RigVMCore/RigVMStorage.cpp b/Engine/Source/Runtime/RigVM/Private/RigVMCore/RigVMStorage.cpp
new file mode 100644
--- /dev/null
+++ b/Engine/Source/Runtime/RigVM/Private/RigVMCore/RigVMStorage.cpp
@@ -0,0 +1,524 @@
+// Copyright 1998-2019 Epic Games, Inc. All Rights Reserved.
+
+#include "RigVMCore/RigVMStorage.h"
+
+FRigVMStorage::FRigVMStorage(bool bInUseNames)
+	: bUseNameMap(bInUseNames)
+	, bIsLiteralStorage(false)
+{
+}
+
+FRigVMStorage::FRigVMStorage(const FRigVMStorage& Other)
+{
+	*this = Other;
+}
+
+FRigVMStorage::~FRigVMStorage()
+{
+	Reset();
+}
+
+FRigVMStorage& FRigVMStorage::operator= (const FRigVMStorage &InOther)
+{
+	Reset();
+
+	bUseNameMap = InOther.bUseNameMap;
+	Data.Append(InOther.Data);
+	Registers.Append(InOther.Registers);
+	ScriptStructs.Append(InOther.ScriptStructs);
+
+	UpdateRegisters();
+
+	for (int32 Index = 0; Index < Registers.Num(); Index++)
+	{
+		Construct(Index);
+		Copy(Index, Index, &InOther);
+	}
+
+	return *this;
+}
+
+void FRigVMStorage::Reset()
+{
+	for (int32 Index = 0; Index < Registers.Num(); Index++)
+	{
+		Destroy(Index);
+	}
+
+	Data.Reset();
+	Registers.Reset();
+	ScriptStructs.Reset();
+	NameMap.Reset();
+}
+
+bool FRigVMStorage::Copy(
+	int32 InSourceRegisterIndex,
+	int32 InTargetRegisterIndex,
+	const FRigVMStorage* InSourceStorage,
+	int32 InSourceByteOffset,
+	int32 InTargetByteOffset,
+	int32 InNumBytes)
+{
+	if (InSourceStorage == nullptr)
+	{
+		InSourceStorage = this;
+	}
+
+	ensure(InSourceStorage->Registers.IsValidIndex(InSourceRegisterIndex));
+	ensure(Registers.IsValidIndex(InTargetRegisterIndex));
+
+	if (InSourceRegisterIndex == InTargetRegisterIndex && InSourceByteOffset == InTargetByteOffset && this == InSourceStorage)
+	{
+		return false;
+	}
+
+	const FRigVMRegister& Source = InSourceStorage->Registers[InSourceRegisterIndex];
+	const FRigVMRegister& Target = Registers[InTargetRegisterIndex];
+
+	int32 SourceStartByte = Source.FirstByte();
+	int32 SourceNumBytes = Source.NumBytes(false);
+	if(InSourceByteOffset != INDEX_NONE)
+	{
+		ensure(InNumBytes != INDEX_NONE);
+		ensure(InSourceByteOffset >= 0 && InSourceByteOffset < Source.NumBytes(false));
+		ensure(InNumBytes > 0 && InSourceByteOffset + InNumBytes <= Source.NumBytes(false));
+
+		SourceStartByte += InSourceByteOffset;
+		SourceNumBytes = InNumBytes;
+	}
+
+	int32 TargetStartByte = Target.FirstByte();
+	int32 TargetNumBytes = Target.NumBytes(false);
+	if(InTargetByteOffset != INDEX_NONE)
+	{
+		ensure(InNumBytes != INDEX_NONE);
+		ensure(InTargetByteOffset >= 0 && InTargetByteOffset < Target.NumBytes(false));
+		ensure(InNumBytes > 0 && InTargetByteOffset + InNumBytes <= Target.NumBytes(false));
+
+		TargetStartByte += InTargetByteOffset;
+		TargetNumBytes = InNumBytes;
+	}
+
+	if (SourceNumBytes != TargetNumBytes)
+	{
+		return false;
+	}
+
+	switch (Target.Type)
+	{
+		case ERigVMRegisterType::Plain:
+		{
+			FMemory::Memcpy(&Data[TargetStartByte], &InSourceStorage->Data[SourceStartByte], TargetNumBytes);
+			break;
+		}
+		case ERigVMRegisterType::Struct:
+		{
+			UScriptStruct* ScriptStruct = GetScriptStruct(InTargetRegisterIndex);
+			int32 NumStructs = TargetNumBytes / ScriptStruct->GetStructureSize();
+			ScriptStruct->CopyScriptStruct(&Data[TargetStartByte], &InSourceStorage->Data[SourceStartByte], NumStructs);
+			break;
+		}
+		case ERigVMRegisterType::Name:
+		{
+			int32 NumNames = TargetNumBytes / sizeof(FName);
+			TArrayView<FName> TargetNames((FName*)&Data[TargetStartByte], NumNames);
+			TArrayView<FName> SourceNames((FName*)&InSourceStorage->Data[SourceStartByte], NumNames);
+			for (int32 Index = 0; Index < NumNames; Index++)
+			{
+				TargetNames[Index] = SourceNames[Index];
+			}
+			break;
+		}
+		case ERigVMRegisterType::String:
+		{
+			int32 NumStrings = TargetNumBytes / sizeof(FString);
+			TArrayView<FString> TargetStrings((FString*)&Data[TargetStartByte], NumStrings);
+			TArrayView<FString> SourceStrings((FString*)&InSourceStorage->Data[SourceStartByte], NumStrings);
+			for (int32 Index = 0; Index < NumStrings; Index++)
+			{
+				TargetStrings[Index] = SourceStrings[Index];
+			}
+			break;
+		}
+		case ERigVMRegisterType::Invalid:
+		{
+			return false;
+		}
+	}
+
+	return true;
+}
+
+bool FRigVMStorage::Copy(
+	const FName& InSourceName,
+	const FName& InTargetName,
+	const FRigVMStorage* InSourceStorage,
+	int32 InSourceByteOffset,
+	int32 InTargetByteOffset,
+	int32 InNumBytes)
+{
+	ensure(bUseNameMap);
+
+	int32 SourceRegisterIndex = GetIndex(InSourceName);
+	int32 TargetRegisterIndex = GetIndex(InTargetName);
+
+	if(SourceRegisterIndex == INDEX_NONE || TargetRegisterIndex == INDEX_NONE)
+	{
+		return false;
+	}
+
+	return Copy(SourceRegisterIndex, TargetRegisterIndex, InSourceStorage, InSourceByteOffset, InTargetByteOffset, InNumBytes);
+}
+
+int32 FRigVMStorage::Allocate(const FName& InNewName, int32 InElementSize, int32 InCount, const void* InDataPtr, bool bUpdateRegisters)
+{
+	FName Name = InNewName;
+	if (bUseNameMap && InNewName == NAME_None)
+	{
+		const TCHAR* RegisterPrefix = TEXT("Register");
+		int32 RegisterSuffix = 0;
+		do
+		{
+			Name = FName(*FString::Printf(TEXT("%s_%d"), RegisterPrefix, RegisterSuffix++));
+		} while (!IsNameAvailable(Name));
+	}
+
+	ensure(InElementSize > 0 && InCount > 0);
+
+	if (bUseNameMap)
+	{
+		if (!IsNameAvailable(Name))
+		{
+			return INDEX_NONE;
+		}
+	}
+
+	FRigVMRegister NewRegister;
+	NewRegister.ByteIndex = Data.Num();
+	if (bUseNameMap)
+	{
+		NewRegister.Name = Name;
+	}
+	NewRegister.ElementSize = InElementSize;
+	NewRegister.ElementCount = InCount;
+	NewRegister.Type = ERigVMRegisterType::Plain;
+
+	Data.AddZeroed(NewRegister.NumBytes());
+
+	if (InDataPtr != nullptr)
+	{
+		FMemory::Memcpy(&Data[NewRegister.FirstByte()], InDataPtr, NewRegister.NumBytes());
+	}
+
+	int32 RegisterIndex = Registers.Num();
+	Registers.Add(NewRegister);
+
+	if (bUpdateRegisters)
+	{
+		UpdateRegisters();
+	}
+	return RegisterIndex;
+}
+
+int32 FRigVMStorage::Allocate(int32 InElementSize, int32 InCount, const void* InDataPtr, bool bUpdateRegisters)
+{
+	return Allocate(NAME_None, InElementSize, InCount, InDataPtr, bUpdateRegisters);
+}
+
+bool FRigVMStorage::Construct(int32 InRegisterIndex, int32 InElementIndex)
+{
+	ensure(Registers.IsValidIndex(InRegisterIndex));
+
+	const FRigVMRegister& Register = Registers[InRegisterIndex];
+	switch (Register.Type)
+	{
+		case ERigVMRegisterType::Struct:
+		{
+			void* DataPtr = (void*)&Data[InElementIndex == INDEX_NONE ? Register.FirstByte() : Register.FirstByte() + InElementIndex * Register.ElementSize];
+			int32 Count = InElementIndex == INDEX_NONE ? Register.ElementCount : 1;
+
+			UScriptStruct* ScriptStruct = GetScriptStruct(InRegisterIndex);
+			ScriptStruct->InitializeStruct(DataPtr, Count);
+			break;
+		}
+		case ERigVMRegisterType::String:
+		{
+			FString* DataPtr = (FString*)&Data[InElementIndex == INDEX_NONE ? Register.FirstByte() : Register.FirstByte() + InElementIndex * Register.ElementSize];
+			int32 Count = InElementIndex == INDEX_NONE ? Register.ElementCount : 1;
+
+			FMemory::Memzero(DataPtr, Count * Register.ElementSize);
+			for (int32 Index = 0; Index < Count; Index++)
+			{
+				DataPtr[Index] = FString();
+			}
+			break;
+		}
+		case ERigVMRegisterType::Name:
+		{
+			FName* DataPtr = (FName*)&Data[InElementIndex == INDEX_NONE ? Register.FirstByte() : Register.FirstByte() + InElementIndex * Register.ElementSize];
+			int32 Count = InElementIndex == INDEX_NONE ? Register.ElementCount : 1;
+
+			FMemory::Memzero(DataPtr, Count * Register.ElementSize);
+			for (int32 Index = 0; Index < Count; Index++)
+			{
+				DataPtr[Index] = FName();
+			}
+			break;
+		}
+		default:
+		{
+			return false;
+		}
+	}
+
+	return true;
+}
+
+bool FRigVMStorage::Destroy(int32 InRegisterIndex, int32 InElementIndex)
+{
+	ensure(Registers.IsValidIndex(InRegisterIndex));
+
+	const FRigVMRegister& Register = Registers[InRegisterIndex];
+	switch (Register.Type)
+	{
+		case ERigVMRegisterType::Struct:
+		{
+			void* DataPtr = (void*)&Data[InElementIndex == INDEX_NONE ? Register.FirstByte() : Register.FirstByte() + InElementIndex * Register.ElementSize];
+			int32 Count = InElementIndex == INDEX_NONE ? Register.ElementCount : 1;
+
+			UScriptStruct* ScriptStruct = GetScriptStruct(InRegisterIndex);
+			ScriptStruct->DestroyStruct(DataPtr, Count);
+			break;
+		}
+		case ERigVMRegisterType::String:
+		{
+			FString* DataPtr = (FString*)&Data[InElementIndex == INDEX_NONE ? Register.FirstByte() : Register.FirstByte() + InElementIndex * Register.ElementSize];
+			int32 Count = InElementIndex == INDEX_NONE ? Register.ElementCount : 1;
+
+			for (int32 Index = 0; Index < Count; Index++)
+			{
+				DataPtr[Index].Empty();
+			}
+			break;
+		}
+		case ERigVMRegisterType::Name:
+		{
+			FName* DataPtr = (FName*)&Data[InElementIndex == INDEX_NONE ? Register.FirstByte() : Register.FirstByte() + InElementIndex * Register.ElementSize];
+			int32 Count = InElementIndex == INDEX_NONE ? Register.ElementCount : 1;
+
+			for (int32 Index = 0; Index < Count; Index++)
+			{
+				DataPtr[Index] = FName();
+			}
+			break;
+		}
+		default:
+		{
+			return false;
+		}
+	}
+
+	return true;
+}
+
+bool FRigVMStorage::Remove(int32 InRegisterIndex)
+{
+	if (InRegisterIndex < 0 || InRegisterIndex >= Registers.Num())
+	{
+		return false;
+	}
+
+	Destroy(InRegisterIndex);
+
+	FRigVMRegister RegisterToRemove = Registers[InRegisterIndex];
+	Data.RemoveAt(RegisterToRemove.ByteIndex, RegisterToRemove.NumBytes());
+	Registers.RemoveAt(InRegisterIndex);
+
+	for (int32 Index = InRegisterIndex; Index < Registers.Num(); Index++)
+	{
+		Registers[Index].ByteIndex -= RegisterToRemove.NumBytes();
+	}
+
+	UpdateRegisters();
+	return true;
+}
+
+bool FRigVMStorage::Remove(const FName& InRegisterName)
+{
+	ensure(bUseNameMap);
+	return Remove(GetIndex(InRegisterName));
+}
+
+FName FRigVMStorage::Rename(int32 InRegisterIndex, const FName& InNewName)
+{
+	if (Registers[InRegisterIndex].Name == InNewName)
+	{
+		return Registers[InRegisterIndex].Name;
+	}
+
+	if (!IsNameAvailable(InNewName))
+	{
+		return Registers[InRegisterIndex].Name;
+	}
+
+	Registers[InRegisterIndex].Name = InNewName;
+	UpdateRegisters();
+
+	return InNewName;
+}
+
+FName FRigVMStorage::Rename(const FName& InOldName, const FName& InNewName)
+{
+	ensure(bUseNameMap);
+
+	int32 RegisterIndex = GetIndex(InOldName);
+	if (RegisterIndex == INDEX_NONE)
+	{
+		return NAME_None;
+	}
+
+	return Rename(RegisterIndex, InNewName);
+}
+
+bool FRigVMStorage::Resize(int32 InRegisterIndex, int32 InNewElementCount)
+{
+	if (InNewElementCount <= 0)
+	{
+		return Remove(InRegisterIndex);
+	}
+
+	if (Registers[InRegisterIndex].ElementCount == InNewElementCount)
+	{
+		return false;
+	}
+
+	FRigVMRegister& Register = Registers[InRegisterIndex];
+
+	if (Register.ElementCount > InNewElementCount) // shrink
+	{
+		int32 ElementsToRemove = Register.ElementCount - InNewElementCount;
+		int32 NumBytesToRemove = Register.ElementSize * ElementsToRemove;
+		int32 FirstByteToRemove = Register.FirstByte() + Register.ElementSize * InNewElementCount;
+
+		for (int32 ElementIndex = InNewElementCount; ElementIndex < Register.ElementCount; ElementIndex++)
+		{
+			Destroy(InRegisterIndex, ElementIndex);
+		}
+
+		Data.RemoveAt(FirstByteToRemove, NumBytesToRemove);
+		Register.ElementCount = InNewElementCount;
+
+		for (int32 RegisterIndex = InRegisterIndex + 1; RegisterIndex < Registers.Num(); RegisterIndex++)
+		{
+			Registers[RegisterIndex].ByteIndex -= NumBytesToRemove;
+		}
+	}
+	else // grow
+	{
+		int32 OldElementCount = Register.ElementCount;
+		int32 ElementsToAdd = InNewElementCount - Register.ElementCount;
+		int32 NumBytesToAdd = Register.ElementSize * ElementsToAdd;
+		int32 FirstByteToAdd = Register.FirstByte() + Register.ElementSize * Register.ElementCount;
+
+		Data.InsertZeroed(FirstByteToAdd, NumBytesToAdd);
+		Register.ElementCount = InNewElementCount;
+
+		for (int32 ElementIndex = OldElementCount; ElementIndex < InNewElementCount; ElementIndex++)
+		{
+			Construct(InRegisterIndex, ElementIndex);
+		}
+
+		for (int32 RegisterIndex = InRegisterIndex + 1; RegisterIndex < Registers.Num(); RegisterIndex++)
+		{
+			Registers[RegisterIndex].ByteIndex += NumBytesToAdd;
+		}
+	}
+
+	UpdateRegisters();
+	return true;
+}
+
+bool FRigVMStorage::Resize(const FName& InRegisterName, int32 InNewElementCount)
+{
+	ensure(bUseNameMap);
+
+	int32 RegisterIndex = GetIndex(InRegisterName);
+	if (RegisterIndex == INDEX_NONE)
+	{
+		return false;
+	}
+
+	return Resize(RegisterIndex, InNewElementCount);
+}
+
+void FRigVMStorage::UpdateRegisters()
+{
+	int32 AlignmentShift = 0;
+	for (int32 RegisterIndex = 0; RegisterIndex < Registers.Num(); RegisterIndex++)
+	{
+		FRigVMRegister& Register = Registers[RegisterIndex];
+		Register.ByteIndex += AlignmentShift;
+
+		UScriptStruct* ScriptStruct = GetScriptStruct(RegisterIndex);
+		if (ScriptStruct != nullptr)
+		{
+			UScriptStruct::ICppStructOps* TheCppStructOps = ScriptStruct->GetCppStructOps();
+			if (TheCppStructOps != NULL)
+			{
+				if (!TheCppStructOps->HasZeroConstructor())
+				{
+					void* Pointer = GetData(RegisterIndex);
+
+					if (Register.AlignmentBytes > 0)
+					{
+						if (!IsAligned(Pointer, TheCppStructOps->GetAlignment()))
+						{
+							Data.RemoveAt(Register.ByteIndex, Register.AlignmentBytes);
+							AlignmentShift -= Register.AlignmentBytes;
+							Register.AlignmentBytes = 0;
+							Pointer = GetData(RegisterIndex);
+						}
+					}
+
+					while (!IsAligned(Pointer, TheCppStructOps->GetAlignment()))
+					{
+						Data.InsertZeroed(Register.ByteIndex, 1);
+						Register.AlignmentBytes++;
+						AlignmentShift++;
+						Pointer = GetData(RegisterIndex);
+					}
+				}
+			}
+		}
+	}
+
+	for (int32 RegisterIndex = 0; RegisterIndex < Registers.Num(); RegisterIndex++)
+	{
+		Registers[RegisterIndex].Pointer = GetData(RegisterIndex);
+	}
+
+	if (bUseNameMap)
+	{
+		NameMap.Reset();
+		for (int32 Index = 0; Index < Registers.Num(); Index++)
+		{
+			NameMap.Add(Registers[Index].Name, Index);
+		}
+	}
+}
+
+void FRigVMStorage::FillWithZeroes(int32 InRegisterIndex)
+{
+	ensure(Registers.IsValidIndex(InRegisterIndex));
+	FMemory::Memzero(GetData(InRegisterIndex), Registers[InRegisterIndex].NumBytes());
+}
+
+int32 FRigVMStorage::FindOrAddScriptStruct(UScriptStruct* InScriptStruct)
+{
+	int32 StructIndex = INDEX_NONE;
+	if (ScriptStructs.Find(InScriptStruct, StructIndex))
+	{
+		return StructIndex;
+	}
+	return ScriptStructs.Add(InScriptStruct);
+}
diff --git a/Engine/Source/Runtime/RigVM/Private/RigVMRegistry.cpp b/Engine/Source/Runtime/RigVM/Private/RigVMRegistry.cpp
deleted file mode 100644
diff --git a/Engine/Source/Runtime/RigVM/Private/RigVMStorage.cpp b/Engine/Source/Runtime/RigVM/Private/RigVMStorage.cpp
deleted file mode 100644
diff --git a/Engine/Source/Runtime/RigVM/Public/RigVM.h b/Engine/Source/Runtime/RigVM/Public/RigVM.h
deleted file mode 100644
diff --git a/Engine/Source/Runtime/RigVM/Public/RigVMByteCode.h b/Engine/Source/Runtime/RigVM/Public/RigVMByteCode.h
deleted file mode 100644
diff --git a/Engine/Source/Runtime/RigVM/Public/RigVMCore/RigVM.h b/Engine/Source/Runtime/RigVM/Public/RigVMCore/RigVM.h
new file mode 100644
--- /dev/null
+++ b/Engine/Source/Runtime/RigVM/Public/RigVMCore/RigVM.h
@@ -0,0 +1,49 @@
+// Copyright 1998-2019 Epic Games, Inc. All Rights Reserved.
+
+#pragma once
+
+#include "CoreMinimal.h"
+#include "RigVMStorage.h"
+#include "RigVMRegistry.h"
+#include "RigVMByteCode.h"
+#include "RigVM.generated.h"
+
+UCLASS()
+class RIGVM_API URigVM : public UObject
+{
+	GENERATED_BODY()
+
+public:
+
+	URigVM();
+	virtual ~URigVM();
+
+	void Reset();
+
+	bool Execute(FRigVMStoragePtrArray Storage = FRigVMStoragePtrArray(), TArrayView<void*> AdditionalArgs = TArrayView<void*>());
+
+	UFUNCTION()
+	int32 AddRigVMFunction(UScriptStruct* InRigVMStruct, const FName& InFunctionName);
+
+	UPROPERTY()
+	FRigVMStorage Literals;
+
+	UPROPERTY()
+	FRigVMStorage WorkState;
+
+	UPROPERTY()
+	FRigVMByteCode ByteCode;
+
+	UPROPERTY()
+	FRigVMByteCodeTable Instructions;
+
+private:
+
+	void ResolveFunctionsIfRequired();
+	void RefreshInstructionsIfRequired();
+
+	UPROPERTY()
+	TArray<FString> FunctionNames;
+
+	TArray<FRigVMFunctionPtr> Functions;
+};
diff --git a/Engine/Source/Runtime/RigVM/Public/RigVMCore/RigVMByteCode.h b/Engine/Source/Runtime/RigVM/Public/RigVMCore/RigVMByteCode.h
new file mode 100644
--- /dev/null
+++ b/Engine/Source/Runtime/RigVM/Public/RigVMCore/RigVMByteCode.h
@@ -0,0 +1,375 @@
+// Copyright 1998-2019 Epic Games, Inc. All Rights Reserved.
+
+#pragma once
+
+#include "CoreMinimal.h"
+#include "RigVMRegistry.h"
+#include "RigVMByteCode.generated.h"
+
+struct FRigVMByteCode;
+
+UENUM()
+enum class ERigVMOpCode : uint8
+{
+	Copy,
+	Increment,
+	Decrement,
+	Equals,
+	NotEquals,
+	Jump,
+	JumpIfTrue,
+	JumpIfFalse,
+	Execute,
+	Exit,
+	Invalid
+};
+
+struct RIGVM_API FRigVMBaseOp
+{
+	FRigVMBaseOp(ERigVMOpCode InOpCode = ERigVMOpCode::Invalid)
+	: OpCode(InOpCode)
+	{
+	}
+
+	ERigVMOpCode OpCode;
+};
+
+struct RIGVM_API FRigVMCopyOp : public FRigVMBaseOp
+{
+	FRigVMCopyOp()
+	: FRigVMBaseOp(ERigVMOpCode::Copy)
+	, Source()
+	, Target()
+	, SourceOffset(INDEX_NONE)
+	, TargetOffset(INDEX_NONE)
+	, NumBytes(INDEX_NONE)
+	{
+	}
+
+	FRigVMCopyOp(
+		FRigVMArgument InSource,
+		FRigVMArgument InTarget,
+		int32 InSourceOffset = INDEX_NONE,
+		int32 InTargetOffset = INDEX_NONE,
+		int32 InNumBytes = INDEX_NONE
+	)
+		: FRigVMBaseOp(ERigVMOpCode::Copy)
+		, Source(InSource)
+		, Target(InTarget)
+		, SourceOffset(InSourceOffset)
+		, TargetOffset(InTargetOffset)
+		, NumBytes(InNumBytes)
+	{
+	}
+
+	FRigVMArgument Source;
+	FRigVMArgument Target;
+	int32 SourceOffset;
+	int32 TargetOffset;
+	int32 NumBytes;
+};
+
+struct RIGVM_API FRigVMIncrementOp : public FRigVMBaseOp
+{
+	FRigVMIncrementOp()
+	: FRigVMBaseOp(ERigVMOpCode::Increment)
+	, Arg()
+	{
+	}
+
+	FRigVMIncrementOp(FRigVMArgument InArg)
+	: FRigVMBaseOp(ERigVMOpCode::Increment)
+	, Arg(InArg)
+	{
+	}
+
+	FRigVMArgument Arg;
+};
+
+struct RIGVM_API FRigVMDecrementOp : public FRigVMBaseOp
+{
+	FRigVMDecrementOp()
+	: FRigVMBaseOp(ERigVMOpCode::Decrement)
+	, Arg()
+	{
+	}
+
+	FRigVMDecrementOp(FRigVMArgument InArg)
+	: FRigVMBaseOp(ERigVMOpCode::Decrement)
+	, Arg(InArg)
+	{
+	}
+
+	FRigVMArgument Arg;
+};
+
+struct RIGVM_API FRigVMEqualsOp : public FRigVMBaseOp
+{
+	FRigVMEqualsOp()
+	: FRigVMBaseOp(ERigVMOpCode::Equals)
+	, A()
+	, B()
+	, Result()
+	{
+	}
+
+	FRigVMEqualsOp(
+		FRigVMArgument InA,
+		FRigVMArgument InB,
+		FRigVMArgument InResult
+	)
+	: FRigVMBaseOp(ERigVMOpCode::Equals)
+	, A(InA)
+	, B(InB)
+	, Result(InResult)
+	{
+	}
+
+	FRigVMArgument A;
+	FRigVMArgument B;
+	FRigVMArgument Result;
+};
+
+struct RIGVM_API FRigVMNotEqualsOp : public FRigVMBaseOp
+{
+	FRigVMNotEqualsOp()
+	: FRigVMBaseOp(ERigVMOpCode::NotEquals)
+	, A()
+	, B()
+	, Result()
+	{
+	}
+
+	FRigVMNotEqualsOp(
+		FRigVMArgument InA,
+		FRigVMArgument InB,
+		FRigVMArgument InResult
+	)
+	: FRigVMBaseOp(ERigVMOpCode::NotEquals)
+	, A(InA)
+	, B(InB)
+	, Result(InResult)
+	{
+	}
+
+	FRigVMArgument A;
+	FRigVMArgument B;
+	FRigVMArgument Result;
+};
+
+struct RIGVM_API FRigVMJumpOp : public FRigVMBaseOp
+{
+	FRigVMJumpOp()
+	: FRigVMBaseOp(ERigVMOpCode::Jump)
+	, InstructionIndex(INDEX_NONE)
+	{
+	}
+
+	FRigVMJumpOp(int32 InInstructionIndex)
+	: FRigVMBaseOp(ERigVMOpCode::Jump)
+	, InstructionIndex(InInstructionIndex)
+	{
+	}
+
+	int32 InstructionIndex;
+};
+
+struct RIGVM_API FRigVMJumpIfTrueOp : public FRigVMBaseOp
+{
+	FRigVMJumpIfTrueOp()
+	: FRigVMBaseOp(ERigVMOpCode::JumpIfTrue)
+	, InstructionIndex(INDEX_NONE)
+	, Condition()
+	{
+	}
+
+	FRigVMJumpIfTrueOp(int32 InInstructionIndex, FRigVMArgument InCondition)
+	: FRigVMBaseOp(ERigVMOpCode::JumpIfTrue)
+	, InstructionIndex(InInstructionIndex)
+	, Condition(InCondition)
+	{
+	}
+
+	int32 InstructionIndex;
+	FRigVMArgument Condition;
+};
+
+struct RIGVM_API FRigVMJumpIfFalseOp : public FRigVMBaseOp
+{
+	FRigVMJumpIfFalseOp()
+	: FRigVMBaseOp(ERigVMOpCode::JumpIfFalse)
+	, InstructionIndex(INDEX_NONE)
+	, Condition()
+	{
+	}
+
+	FRigVMJumpIfFalseOp(int32 InInstructionIndex, FRigVMArgument InCondition)
+	: FRigVMBaseOp(ERigVMOpCode::JumpIfFalse)
+	, InstructionIndex(InInstructionIndex)
+	, Condition(InCondition)
+	{
+	}
+
+	int32 InstructionIndex;
+	FRigVMArgument Condition;
+};
+
+struct RIGVM_API FRigVMExecuteOp : public FRigVMBaseOp
+{
+	FRigVMExecuteOp()
+	: FRigVMBaseOp(ERigVMOpCode::Execute)
+	, FunctionIndex(INDEX_NONE)
+	, ArgumentCount(0)
+	{
+	}
+
+	FRigVMExecuteOp(uint16 InFunctionIndex, uint16 InArgumentCount)
+	: FRigVMBaseOp(ERigVMOpCode::Execute)
+	, FunctionIndex(InFunctionIndex)
+	, ArgumentCount(InArgumentCount)
+	{
+	}
+
+	uint16 FunctionIndex;
+	uint16 ArgumentCount;
+};
+
+struct RIGVM_API FRigVMExitOp : public FRigVMBaseOp
+{
+	FRigVMExitOp()
+	: FRigVMBaseOp(ERigVMOpCode::Exit)
+	{
+	}
+};
+
+USTRUCT()
+struct RIGVM_API FRigVMByteCodeTableEntry
+{
+	GENERATED_USTRUCT_BODY()
+
+	FRigVMByteCodeTableEntry(ERigVMOpCode InOpCode = ERigVMOpCode::Invalid, uint64 InByteCodeIndex = UINT64_MAX)
+		: OpCode(InOpCode)
+		, ByteCodeIndex(InByteCodeIndex)
+	{
+	}
+
+	UPROPERTY()
+	ERigVMOpCode OpCode;
+
+	UPROPERTY()
+	uint64 ByteCodeIndex;
+};
+
+USTRUCT()
+struct RIGVM_API FRigVMByteCodeTable
+{
+	GENERATED_USTRUCT_BODY()
+
+public:
+
+	FRigVMByteCodeTable();
+
+	void Reset();
+	FORCEINLINE bool IsValidIndex(int32 InIndex) const { return Entries.IsValidIndex(InIndex); }
+	FORCEINLINE int32 Num() const { return Entries.Num(); }
+	FORCEINLINE const FRigVMByteCodeTableEntry& operator[](int32 InIndex) const { return Entries[InIndex]; }
+
+private:
+
+	// hide utility constructor
+	FRigVMByteCodeTable(const FRigVMByteCode& InByteCode);
+
+	UPROPERTY()
+	TArray<FRigVMByteCodeTableEntry> Entries;
+
+	friend struct FRigVMByteCode;
+};
+
+USTRUCT()
+struct RIGVM_API FRigVMByteCode
+{
+	GENERATED_USTRUCT_BODY()
+
+public:
+
+	FRigVMByteCode();
+
+	void Reset();
+	uint64 Num() const;
+
+	uint64 AddCopyOp(const FRigVMArgument& InSource, const FRigVMArgument& InTarget, int32 InSourceOffset = INDEX_NONE, int32 InTargetOffset = INDEX_NONE, int32 InNumBytes = INDEX_NONE);
+	uint64 AddIncrementOp(const FRigVMArgument& InArg);
+	uint64 AddDecrementOp(const FRigVMArgument& InArg);
+	uint64 AddEqualsOp(const FRigVMArgument& InA, const FRigVMArgument& InB, const FRigVMArgument& InResult);
+	uint64 AddNotEqualsOp(const FRigVMArgument& InA, const FRigVMArgument& InB, const FRigVMArgument& InResult);
+	uint64 AddJumpOp(uint64 InByteCodeIndex);
+	uint64 AddJumpIfTrueOp(uint64 InByteCodeIndex, const FRigVMArgument& InCondition);
+	uint64 AddJumpIfFalseOp(uint64 InByteCodeIndex, const FRigVMArgument& InCondition);
+	uint64 AddExecuteOp(uint16 InFunctionIndex, const TArrayView<FRigVMArgument>& InArguments);
+	uint64 AddExitOp();
+
+	FORCEINLINE FRigVMByteCodeTable GetTable() const
+	{
+		return FRigVMByteCodeTable(*this);
+	}
+
+	FORCEINLINE ERigVMOpCode GetOpCodeAt(uint64 InByteCodeIndex) const
+	{
+		ensure(InByteCodeIndex >= 0 && InByteCodeIndex < ByteCode.Num());
+		return (ERigVMOpCode)ByteCode[InByteCodeIndex];
+	}
+
+	uint64 GetOpNumBytesAt(uint64 InByteCodeIndex, bool bIncludeArguments = true) const;
+
+	template<class OpType>
+	FORCEINLINE const OpType& GetOpAt(uint64 InByteCodeIndex) const
+	{
+		ensure(InByteCodeIndex >= 0 && InByteCodeIndex <= ByteCode.Num() - sizeof(OpType));
+		return *(const OpType*)(ByteCode.GetData() + InByteCodeIndex);
+	}
+
+	template<class OpType>
+	FORCEINLINE const OpType& GetOpAt(const FRigVMByteCodeTableEntry& InEntry) const
+	{
+		ensure(OpType().OpCode == InEntry.OpCode);
+		return GetOpAt<OpType>(InEntry.ByteCodeIndex);
+	}
+
+	FORCEINLINE TArrayView<FRigVMArgument> GetArgumentsAt(uint64 InByteCodeIndex, uint16 InArgumentCount) const
+	{
+		ensure(InByteCodeIndex >= 0 && InByteCodeIndex <= ByteCode.Num() - sizeof(FRigVMArgument) * InArgumentCount);
+		return TArrayView<FRigVMArgument>((FRigVMArgument*)(ByteCode.GetData() + InByteCodeIndex), InArgumentCount);
+	}
+
+	FORCEINLINE TArrayView<FRigVMArgument> GetArgumentsForExecuteOp(uint64 InByteCodeIndex) const
+	{
+		const FRigVMExecuteOp& ExecuteOp = GetOpAt<FRigVMExecuteOp>(InByteCodeIndex);
+		return GetArgumentsAt(InByteCodeIndex + sizeof(FRigVMExecuteOp), ExecuteOp.ArgumentCount);
+	}
+
+	FORCEINLINE TArrayView<FRigVMArgument> GetArgumentsForExecuteOp(const FRigVMByteCodeTableEntry& InEntry) const
+	{
+		const FRigVMExecuteOp& ExecuteOp = GetOpAt<FRigVMExecuteOp>(InEntry);
+		return GetArgumentsAt(InEntry.ByteCodeIndex + sizeof(FRigVMExecuteOp), ExecuteOp.ArgumentCount);
+	}
+
+	FORCEINLINE const TArrayView<uint8> GetByteCode() const
+	{
+		const uint8* Data = ByteCode.GetData();
+		return TArrayView<uint8>((uint8*)Data, ByteCode.Num());
+	}
+
+private:
+
+	template<class OpType>
+	FORCEINLINE uint64 AddOp(const OpType& InOp)
+	{
+		uint64 ByteIndex = (uint64)ByteCode.AddUninitialized(sizeof(OpType));
+		FMemory::Memcpy(ByteCode.GetData() + ByteIndex, &InOp, sizeof(OpType));
+		return ByteIndex;
+	}
+
+	// storage for all functions
+	UPROPERTY()
+	TArray<uint8> ByteCode;
+};
diff --git a/Engine/Source/Runtime/RigVM/Public/RigVMCore/RigVMRegistry.h b/Engine/Source/Runtime/RigVM/Public/RigVMCore/RigVMRegistry.h
new file mode 100644
--- /dev/null
+++ b/Engine/Source/Runtime/RigVM/Public/RigVMCore/RigVMRegistry.h
@@ -0,0 +1,52 @@
+// Copyright 1998-2019 Epic Games, Inc. All Rights Reserved.
+
+#pragma once
+
+#include "CoreMinimal.h"
+#include "RigVMStorage.h"
+
+typedef TArrayView<void*> FRigVMUserDataArray;
+typedef void (*FRigVMFunctionPtr)(const FRigVMArgumentArray&, FRigVMStoragePtrArray&, const FRigVMUserDataArray&);
+
+struct RIGVM_API FRigVMFunction
+{
+	const TCHAR* Name;
+	FRigVMFunctionPtr FunctionPtr;
+
+	FRigVMFunction()
+		: Name(nullptr)
+		, FunctionPtr(nullptr)
+	{
+	}
+
+	FRigVMFunction(const TCHAR* InName, FRigVMFunctionPtr InFunctionPtr)
+		: Name(InName)
+		, FunctionPtr(InFunctionPtr)
+	{
+	}
+};
+
+struct RIGVM_API FRigVMRegistry
+{
+public:
+	static FRigVMRegistry& Get();
+
+	void Register(const TCHAR* InName, FRigVMFunctionPtr InFunctionPtr);
+	void Refresh();
+
+	FRigVMFunctionPtr Find(const TCHAR* InName) const;
+
+private:
+
+	// disable default constructor
+	FRigVMRegistry() {}
+	// disable copy constructor
+	FRigVMRegistry(const FRigVMRegistry&) {}
+	// disable assignment operator
+	FRigVMRegistry& operator= (const FRigVMRegistry &InOther) { return *this; }
+
+	// storage for all functions
+	TArray<FRigVMFunction> Functions;
+
+	static FRigVMRegistry s_RigVMRegistry;
+};
diff --git a/Engine/Source/Runtime/RigVM/Public/RigVMCore/RigVMStorage.h b/Engine/Source/Runtime/RigVM/Public/RigVMCore/RigVMStorage.h
new file mode 100644
--- /dev/null
+++ b/Engine/Source/Runtime/RigVM/Public/RigVMCore/RigVMStorage.h
@@ -0,0 +1,560 @@
+// Copyright 1998-2019 Epic Games, Inc. All Rights Reserved.
+
+#pragma once
+
+#include "CoreMinimal.h"
+#include "UObject/ObjectMacros.h"
+#include "RigVMStorage.generated.h"
+
+struct RIGVM_API FRigVMArgument
+{
+public:
+
+	FRigVMArgument()
+		: bIsLiteral(false)
+		, RegisterIndex(INDEX_NONE)
+		, ByteIndex(INDEX_NONE)
+	{
+	}
+
+	FRigVMArgument(bool InIsLiteral, int32 InRegisterIndex, int32 InByteIndex)
+		: bIsLiteral(InIsLiteral)
+		, RegisterIndex(InRegisterIndex)
+		, ByteIndex(InByteIndex)
+	{
+	}
+
+	FORCEINLINE bool IsLiteral() const { return bIsLiteral; }
+	FORCEINLINE int32 StorageType() const { return bIsLiteral ? 1 : 0; }
+	FORCEINLINE int32 GetRegisterIndex() const { return RegisterIndex; }
+	FORCEINLINE int32 GetByteIndex() const { return ByteIndex; }
+
+private:
+
+	bool bIsLiteral;
+	int32 RegisterIndex;
+	int32 ByteIndex;
+};
+
+typedef TArrayView<FRigVMArgument> FRigVMArgumentArray;
+
+UENUM()
+enum class ERigVMRegisterType : uint8
+{
+	Plain,
+	String,
+	Name,
+	Struct,
+	Invalid
+};
+
+USTRUCT()
+struct RIGVM_API FRigVMRegister
+{
+	GENERATED_USTRUCT_BODY()
+
+		FRigVMRegister()
+		: Type(ERigVMRegisterType::Invalid)
+		, Pointer(nullptr)
+		, ByteIndex(INDEX_NONE)
+		, ElementSize(0)
+		, ElementCount(0)
+		, AlignmentBytes(0)
+		, Name(NAME_None)
+		, ScriptStructIndex(INDEX_NONE)
+	{
+	}
+
+	UPROPERTY()
+	ERigVMRegisterType Type;
+
+	void* Pointer;
+
+	UPROPERTY()
+	int32 ByteIndex;
+
+	UPROPERTY()
+	int32 ElementSize;
+
+	UPROPERTY()
+	int32 ElementCount;
+
+	UPROPERTY()
+	int32 AlignmentBytes;
+
+	UPROPERTY()
+	FName Name;
+
+	UPROPERTY()
+	int32 ScriptStructIndex;
+
+	FORCEINLINE int32 FirstByte() const { return ByteIndex + AlignmentBytes; }
+	FORCEINLINE bool IsArray() const { return ElementCount > 1; }
+	FORCEINLINE bool IsPlain() const { return ScriptStructIndex == INDEX_NONE;  }
+	FORCEINLINE int32 NumBytes(bool bIncludeAlignment = true) const { return ElementCount * ElementSize + (bIncludeAlignment ? AlignmentBytes : 0); }
+
+	template<class T>
+	FORCEINLINE const T* Get() const
+	{
+		ensure(ElementCount > 0);
+		return (const T*)Pointer;
+	}
+
+	template<class T>
+	FORCEINLINE const T& GetRef() const
+	{
+		return *Get<T>();
+	}
+
+	template<class T>
+	FORCEINLINE T* Get()
+	{
+		ensure(ElementCount > 0);
+		return (T*)Pointer;
+	}
+
+	template<class T>
+	FORCEINLINE T& GetRef()
+	{
+		return *Get<T>();
+	}
+
+	template<class T>
+	FORCEINLINE TArrayView<T> GetArray()
+	{
+		ensure(ElementCount > 0);
+		return TArrayView<T>((T*)Pointer, ElementCount);
+	}
+};
+
+typedef TArrayView<FRigVMRegister> FRigVMRegisterArray;
+
+USTRUCT()
+struct RIGVM_API FRigVMStorage
+{
+	GENERATED_USTRUCT_BODY()
+	
+public:
+
+	FRigVMStorage(bool bInUseNames = true);
+	FRigVMStorage(const FRigVMStorage& Other);
+	~FRigVMStorage();
+
+	FRigVMStorage& operator= (const FRigVMStorage &InOther);
+
+	FORCEINLINE bool IsLiteralStorage() const { return bIsLiteralStorage;  }
+	FORCEINLINE void SetLiteralStorage(bool InIsLiteralStorage = true) { bIsLiteralStorage = InIsLiteralStorage; }
+	FORCEINLINE bool SupportsNames() const { return bUseNameMap;  }
+	FORCEINLINE int32 Num() const { return Registers.Num(); }
+	FORCEINLINE const FRigVMRegister& operator[](int32 InIndex) const { return Registers[InIndex]; }
+	FORCEINLINE FRigVMRegister& operator[](int32 InIndex) { return Registers[InIndex]; }
+	FORCEINLINE const FRigVMRegister& operator[](const FRigVMArgument& InArg) const { return Registers[InArg.GetRegisterIndex()]; }
+	FORCEINLINE FRigVMRegister& operator[](const FRigVMArgument& InArg) { return Registers[InArg.GetRegisterIndex()]; }
+	FORCEINLINE const FRigVMRegister& operator[](const FName& InName) const { return Registers[GetIndex(InName)]; }
+	FORCEINLINE FRigVMRegister& operator[](const FName& InName) { return Registers[GetIndex(InName)]; }
+
+	FORCEINLINE TArray<FRigVMRegister>::RangedForIteratorType      begin() { return Registers.begin(); }
+	FORCEINLINE TArray<FRigVMRegister>::RangedForConstIteratorType begin() const { return Registers.begin(); }
+	FORCEINLINE TArray<FRigVMRegister>::RangedForIteratorType      end() { return Registers.end(); }
+	FORCEINLINE TArray<FRigVMRegister>::RangedForConstIteratorType end() const { return Registers.end(); }
+
+	FORCEINLINE FRigVMArgument GetArgument(int32 InRegisterIndex) const
+	{
+		ensure(Registers.IsValidIndex(InRegisterIndex));
+		return FRigVMArgument(IsLiteralStorage(), InRegisterIndex, Registers[InRegisterIndex].ByteIndex);
+	}
+
+	FORCEINLINE const void* GetData(int32 InRegisterIndex) const
+	{
+		ensure(Registers.IsValidIndex(InRegisterIndex));
+		const FRigVMRegister& Register = Registers[InRegisterIndex];
+		ensure(Register.ElementCount > 0);
+		return (const void*)&Data[Register.FirstByte()];
+	}
+
+	FORCEINLINE void* GetData(int32 InRegisterIndex)
+	{
+		ensure(Registers.IsValidIndex(InRegisterIndex));
+		const FRigVMRegister& Register = Registers[InRegisterIndex];
+		ensure(Register.ElementCount > 0);
+		return (void*)&Data[Register.FirstByte()];
+	}
+
+	template<class T>
+	FORCEINLINE const T* Get(int32 InRegisterIndex) const
+	{
+		ensure(Registers.IsValidIndex(InRegisterIndex));
+		const FRigVMRegister& Register = Registers[InRegisterIndex];
+		ensure(Register.ElementCount > 0);
+		return (const T*)&Data[Register.FirstByte()];
+	}
+
+	template<class T>
+	FORCEINLINE const T* Get(const FRigVMArgument& InArgument) const
+	{
+		return Get<T>(InArgument.Index());
+	}
+
+	template<class T>
+	FORCEINLINE const T& GetRef(int32 InRegisterIndex) const
+	{
+		return *Get<T>(InRegisterIndex);
+	}
+
+	template<class T>
+	FORCEINLINE const T& GetRef(const FRigVMArgument& InArgument) const
+	{
+		return GetRef<T>(InArgument.Index());
+	}
+
+	template<class T>
+	FORCEINLINE T* Get(int32 InRegisterIndex)
+	{
+		ensure(Registers.IsValidIndex(InRegisterIndex));
+		const FRigVMRegister& Register = Registers[InRegisterIndex];
+		ensure(Register.ElementCount > 0);
+		return (T*)&Data[Register.FirstByte()];
+	}
+
+	template<class T>
+	FORCEINLINE T* Get(const FRigVMArgument& InArgument)
+	{
+		return Get<T>(InArgument.Index());
+	}
+
+	template<class T>
+	FORCEINLINE T& GetRef(int32 InRegisterIndex)
+	{
+		return *Get<T>(InRegisterIndex);
+	}
+
+	template<class T>
+	FORCEINLINE T& GetRef(const FRigVMArgument& InArgument)
+	{
+		return GetRef<T>(InArgument.GetRegisterIndex());
+	}
+
+	template<class T>
+	FORCEINLINE TArrayView<T> GetArray(int32 InRegisterIndex)
+	{
+		ensure(Registers.IsValidIndex(InRegisterIndex));
+		const FRigVMRegister& Register = Registers[InRegisterIndex];
+		ensure(Register.ElementCount > 0);
+		return TArrayView<T>((T*)&Data[Register.FirstByte()], Register.ElementCount);
+	}
+	
+	template<class T>
+	FORCEINLINE TArrayView<T> GetArray(const FRigVMArgument& InArgument)
+	{
+		return GetArray<T>(InArgument.GetRegisterIndex());
+	}
+
+	FORCEINLINE UScriptStruct* GetScriptStruct(int32 InRegisterIndex) const
+	{
+		ensure(Registers.IsValidIndex(InRegisterIndex));
+		const FRigVMRegister& Register = Registers[InRegisterIndex];
+		if (Register.ScriptStructIndex != INDEX_NONE)
+		{
+			ensure(ScriptStructs.IsValidIndex(Register.ScriptStructIndex));
+			return ScriptStructs[Register.ScriptStructIndex];
+		}
+		return nullptr;
+	}
+
+	bool Copy(
+		int32 InSourceRegisterIndex,
+		int32 InTargetRegisterIndex,
+		const FRigVMStorage* InSourceStorage = nullptr,
+		int32 InSourceByteOffset = INDEX_NONE,
+		int32 InTargetByteOffset = INDEX_NONE,
+		int32 InNumBytes = INDEX_NONE);
+
+	bool Copy(
+		const FName& InSourceName,
+		const FName& InTargetName,
+		const FRigVMStorage* InSourceStorage = nullptr,
+		int32 InSourceByteOffset = INDEX_NONE,
+		int32 InTargetByteOffset = INDEX_NONE,
+		int32 InNumBytes = INDEX_NONE);
+
+	FORCEINLINE int32 GetIndex(const FName& InName) const
+	{
+		if (!bUseNameMap)
+		{
+			return INDEX_NONE;
+		}
+
+		if (NameMap.Num() != Registers.Num())
+		{
+			for (int32 Index = 0; Index < Registers.Num(); Index++)
+			{
+				if (Registers[Index].Name == InName)
+				{
+					return Index;
+				}
+			}
+		}
+		else
+		{
+			const int32* Index = NameMap.Find(InName);
+			if (Index != nullptr)
+			{
+				return *Index;
+			}
+		}
+
+		return INDEX_NONE;
+	}
+
+	FORCEINLINE bool IsNameAvailable(const FName& InPotentialNewName) const\
+	{
+		if (!bUseNameMap)
+		{
+			return false;
+		}
+		return GetIndex(InPotentialNewName) == INDEX_NONE;
+	}
+
+	void Reset();
+
+	FORCEINLINE int32 AddPlainArray(const FName& InNewName, int32 InElementSize, int32 InCount, const void* InDataPtr = nullptr)
+	{
+		return Allocate(InNewName, InElementSize, InCount, InDataPtr);
+	}
+
+	template<class T>
+	FORCEINLINE int32 AddPlainArray(const FName& InNewName, int32 InCount, const T* InDataPtr = nullptr)
+	{
+		return Allocate(InNewName, sizeof(T), InCount, (const void*)InDataPtr);
+	}
+
+	template<class T>
+	FORCEINLINE int32 AddPlainArray(const FName& InNewName, const TArray<T>& InArray)
+	{
+		return AddPlainArray<T>(InNewName, InArray.Num(), InArray.GetData());
+	}
+
+	template<class T>
+	FORCEINLINE int32 AddPlainArray(const TArray<T>& InArray)
+	{
+		return AddPlainArray<T>(NAME_None, InArray);
+	}
+
+	FORCEINLINE int32 AddPlain(const FName& InNewName, int32 InElementSize, const void* InValuePtr)
+	{
+		return AddPlainArray(InNewName, InElementSize, 1, InValuePtr);
+	}
+
+	FORCEINLINE int32 AddPlain(int32 InElementSize, const void* InValuePtr)
+	{
+		return AddPlain(NAME_None, InElementSize, InValuePtr);
+	}
+
+	template<class T>
+	FORCEINLINE int32 AddPlain(const FName& InNewName, const T& InValue)
+	{
+		return AddPlainArray<T>(InNewName, 1, &InValue);
+	}
+
+	template<class T>
+	FORCEINLINE int32 AddPlain(const T& InValue)
+	{
+		return AddPlain<T>(NAME_None, InValue);
+	}
+
+	FORCEINLINE int32 AddNameArray(const FName& InNewName, int32 InCount, const FName* InDataPtr = nullptr)
+	{
+		int32 Register = Allocate(InNewName, sizeof(FName), InCount, nullptr);
+		Registers[Register].Type = ERigVMRegisterType::Name;
+
+		Construct(Register);
+
+		if(InDataPtr)
+		{
+			FName* DataPtr = (FName*)GetData(Register);
+			for (int32 Index = 0; Index < InCount; Index++)
+			{
+				DataPtr[Index] = InDataPtr[Index];
+			}
+		}
+
+		return Register;
+	}
+
+	FORCEINLINE int32 AddNameArray(const FName& InNewName, const TArray<FName>& InArray)
+	{
+		return AddNameArray(InNewName, InArray.Num(), InArray.GetData());
+	}
+
+	FORCEINLINE int32 AddNameArray(const TArray<FName>& InArray)
+	{
+		return AddNameArray(NAME_None, InArray);
+	}
+
+	FORCEINLINE int32 AddName(const FName& InNewName, const FName& InValue)
+	{
+		return AddNameArray(InNewName, 1, &InValue);
+	}
+
+	FORCEINLINE int32 AddName(const FName& InValue)
+	{
+		return AddName(NAME_None, InValue);
+	}
+
+	FORCEINLINE int32 AddStringArray(const FName& InNewName, int32 InCount, const FString* InDataPtr = nullptr)
+	{
+		int32 Register = Allocate(InNewName, sizeof(FString), InCount, nullptr);
+		Registers[Register].Type = ERigVMRegisterType::String;
+
+		Construct(Register);
+
+		if(InDataPtr)
+		{
+			FString* DataPtr = (FString*)GetData(Register);
+			for (int32 Index = 0; Index < InCount; Index++)
+			{
+				DataPtr[Index] = InDataPtr[Index];
+			}
+		}
+
+		return Register;
+	}
+
+	FORCEINLINE int32 AddStringArray(const FName& InNewName, const TArray<FString>& InArray)
+	{
+		return AddStringArray(InNewName, InArray.Num(), InArray.GetData());
+	}
+
+	FORCEINLINE int32 AddStringArray(const TArray<FString>& InArray)
+	{
+		return AddStringArray(NAME_None, InArray);
+	}
+
+	FORCEINLINE int32 AddString(const FName& InNewName, const FString& InValue)
+	{
+		return AddStringArray(InNewName, 1, &InValue);
+	}
+
+	FORCEINLINE int32 AddString(const FString& InValue)
+	{
+		return AddString(NAME_None, InValue);
+	}
+
+
+	FORCEINLINE int32 AddStructArray(const FName& InNewName, UScriptStruct* InScriptStruct, int32 InCount, const void* InDataPtr = nullptr)
+	{
+		int32 Register = Allocate(InNewName, InScriptStruct->GetStructureSize(), InCount, nullptr, false);
+		if (Register == INDEX_NONE)
+		{
+			return INDEX_NONE;
+		}
+
+		Registers[Register].Type = ERigVMRegisterType::Struct;
+		Registers[Register].ScriptStructIndex = FindOrAddScriptStruct(InScriptStruct);
+
+		UpdateRegisters();
+
+		// construct the content
+		Construct(Register);
+
+		// copy values from the provided data
+		if (InDataPtr != nullptr)
+		{
+			InScriptStruct->CopyScriptStruct(GetData(Register), InDataPtr, InCount);
+		}
+
+		return Register;
+	}
+
+	FORCEINLINE int32 AddStructArray(UScriptStruct* InScriptStruct, int32 InCount, const void* InDataPtr = nullptr)
+	{
+		return AddStructArray(NAME_None, InScriptStruct, InCount, InDataPtr);
+	}
+
+	template<class T>
+	FORCEINLINE int32 AddStructArray(const FName& InNewName, int32 InCount, const T* InDataPtr = nullptr)
+	{
+		// if you are hitting this - you might need to use AddPlainArray instead!
+		UScriptStruct* Struct = T::StaticStruct();
+		if (Struct == nullptr)
+		{
+			return INDEX_NONE;
+		}
+	
+		return AddStructArray(InNewName, Struct, InCount, InDataPtr);
+	}
+
+	template<class T>
+	FORCEINLINE int32 AddStructArray(const FName& InNewName, const TArray<T>& InArray)
+	{
+		return AddStructArray<T>(InNewName, InArray.Num(), InArray.GetData());
+	}
+
+	template<class T>
+	FORCEINLINE int32 AddStructArray(const TArray<T>& InArray)
+	{
+		return AddStructArray<T>(NAME_None, InArray);
+	}
+
+	FORCEINLINE int32 AddStruct(const FName& InNewName, UScriptStruct* InScriptStruct, const void* InValuePtr)
+	{
+		return AddStructArray(InNewName, InScriptStruct, 1, InValuePtr);
+	}
+
+	FORCEINLINE int32 AddStruct(UScriptStruct* InScriptStruct, const void* InValuePtr)
+	{
+		return AddStruct(NAME_None, InScriptStruct, InValuePtr);
+	}
+
+	template<class T>
+	FORCEINLINE int32 AddStruct(const FName& InNewName, const T& InValue)
+	{
+		return AddStructArray<T>(InNewName, 1, &InValue);
+	}
+
+	template<class T>
+	FORCEINLINE int32 AddStruct(const T& InValue)
+	{
+		return AddStruct<T>(NAME_None, InValue);
+	}
+
+	bool Remove(int32 InRegisterIndex);
+	bool Remove(const FName& InRegisterName);
+	FName Rename(int32 InRegisterIndex, const FName& InNewName);
+	FName Rename(const FName& InOldName, const FName& InNewName);
+	bool Resize(int32 InRegisterIndex, int32 InNewElementCount);
+	bool Resize(const FName& InRegisterName, int32 InNewElementCount);
+
+	void UpdateRegisters();
+
+private:
+
+	int32 Allocate(const FName& InNewName, int32 InElementSize, int32 InCount, const void* InDataPtr = nullptr, bool bUpdateRegisters = true);
+	int32 Allocate(int32 InElementSize, int32 InCount, const void* InDataPtr = nullptr, bool bUpdateRegisters = true);
+	bool Construct(int32 InRegisterIndex, int32 InElementIndex = INDEX_NONE);
+	bool Destroy(int32 InRegisterIndex, int32 InElementIndex = INDEX_NONE);
+
+	void FillWithZeroes(int32 InRegisterIndex);
+	int32 FindOrAddScriptStruct(UScriptStruct* InScriptStruct);
+
+	UPROPERTY()
+	bool bUseNameMap;
+
+	UPROPERTY()
+	bool bIsLiteralStorage;
+
+	UPROPERTY()
+	TArray<FRigVMRegister> Registers;
+
+	UPROPERTY()
+	TArray<uint8> Data;
+
+	UPROPERTY()
+	TArray<UScriptStruct*> ScriptStructs;
+
+	UPROPERTY(transient)
+	TMap<FName, int32> NameMap;
+};
+
+typedef FRigVMStorage* FRigVMStoragePtr;
+typedef TArrayView<FRigVMStoragePtr> FRigVMStoragePtrArray;
\ No newline at end of file
diff --git a/Engine/Source/Runtime/RigVM/Public/RigVMRegistry.h b/Engine/Source/Runtime/RigVM/Public/RigVMRegistry.h
deleted file mode 100644
diff --git a/Engine/Source/Runtime/RigVM/Public/RigVMStorage.h b/Engine/Source/Runtime/RigVM/Public/RigVMStorage.h
deleted file mode 100644
