RigVM: Introduced controller layer, first steps towards MVC + unit tests

diff --git a/Engine/Source/Runtime/RigVM/Private/RigVMModel/RigVMController.cpp b/Engine/Source/Runtime/RigVM/Private/RigVMModel/RigVMController.cpp
new file mode 100644
--- /dev/null
+++ b/Engine/Source/Runtime/RigVM/Private/RigVMModel/RigVMController.cpp
@@ -0,0 +1,199 @@
+// Copyright 1998-2019 Epic Games, Inc. All Rights Reserved.
+
+#include "RigVMModel/RigVMController.h"
+#include "RigVMCore/RigVMRegistry.h"
+#include "RigVMModule.h"
+
+URigVMController::~URigVMController()
+{
+	SetGraph(nullptr);
+}
+
+URigVMGraph* URigVMController::GetGraph() const
+{
+	return Graph;
+}
+
+void URigVMController::SetGraph(URigVMGraph* InGraph)
+{
+	if (Graph)
+	{
+		Graph->OnModified().RemoveAll(this);
+	}
+
+	Graph = InGraph;
+
+	if (Graph)
+	{
+		Graph->OnModified().AddUObject(this, &URigVMController::HandleModifiedEvent);
+	}
+
+	ModifiedEvent.Broadcast(ERigVMGraphNotifType::GraphChanged, Graph, nullptr);
+}
+
+FRigVMGraphModifiedEvent& URigVMController::OnModified()
+{
+	return ModifiedEvent;
+}
+
+void URigVMController::Notify(ERigVMGraphNotifType InNotifType, URigVMGraph* InGraph, UObject* InSubject)
+{
+	if (Graph)
+	{
+		Graph->Notify(InNotifType, InGraph, InSubject);
+	}
+}
+
+void URigVMController::HandleModifiedEvent(ERigVMGraphNotifType InNotifType, URigVMGraph* InGraph, UObject* InSubject)
+{
+	ModifiedEvent.Broadcast(InNotifType, InGraph, InSubject);
+}
+
+
+URigVMStructNode* URigVMController::AddStructNode(UScriptStruct* InScriptStruct, const FName& InMethodName, const FVector2D& InPosition, bool bUndo)
+{
+	if (Graph == nullptr)
+	{
+		return nullptr;
+	}
+
+	if (InScriptStruct == nullptr)
+	{
+		return nullptr;
+	}
+	if (InMethodName == NAME_None)
+	{
+		return nullptr;
+	}
+
+	FString FunctionName = FString::Printf(TEXT("F%s::%s"), *InScriptStruct->GetName(), *InMethodName.ToString());
+	FRigVMFunctionPtr Function = FRigVMRegistry::Get().Find(*FunctionName);
+	if (Function == nullptr)
+	{
+		// todo: log error
+		return nullptr;
+	}
+
+	FString NamePrefix = InScriptStruct->GetName();
+	int32 NameSuffix = 0;
+	FString Name = NamePrefix;
+
+	while (!Graph->IsNameAvailable(Name))
+	{
+		Name = FString::Printf(TEXT("%s_%d"), *NamePrefix, ++NameSuffix);
+	}
+
+	URigVMStructNode* Node = NewObject<URigVMStructNode>(Graph, *Name);
+	Node->ScriptStruct = InScriptStruct;
+	Node->MethodName = InMethodName;
+	Node->Position = InPosition;
+
+	Graph->Nodes.Add(Node);
+	Graph->MarkPackageDirty();
+
+	// todo: undo redo
+
+	Notify(ERigVMGraphNotifType::NodeAdded, Graph, Node);
+
+	return Node;
+}
+
+bool URigVMController::RemoveNode(URigVMNode* InNode, bool bUndo)
+{
+	if (!IsValidNodeForGraph(InNode))
+	{
+		return false;
+	}
+
+	// todo: break all links
+	// todo: undo redo
+
+	SelectNode(InNode, false, bUndo);
+
+	Graph->Nodes.Remove(InNode);
+	Graph->MarkPackageDirty();
+
+	Notify(ERigVMGraphNotifType::NodeRemoved, Graph, InNode);
+
+	return true;
+}
+
+bool URigVMController::RemoveNode(const FName& InNodeName, bool bUndo)
+{
+	if (Graph == nullptr)
+	{
+		return false;
+	}
+	return RemoveNode(Graph->FindNode(InNodeName), bUndo);
+}
+
+bool URigVMController::SelectNode(URigVMNode* InNode, bool bSelect, bool bUndo)
+{
+	if (!IsValidNodeForGraph(InNode))
+	{
+		return false;
+	}
+
+	if (InNode->IsSelected() == bSelect)
+	{
+		return false;
+	}
+
+	if (bSelect)
+	{
+		Graph->SelectedNodes.Add(InNode->GetFName());
+		Notify(ERigVMGraphNotifType::NodeSelected, Graph, InNode);
+	}
+	else
+	{
+		Graph->SelectedNodes.Remove(InNode->GetFName());
+		Notify(ERigVMGraphNotifType::NodeDeselected, Graph, InNode);
+	}
+
+	return true;
+}
+
+bool URigVMController::SelectNode(const FName& InNodeName, bool bSelect, bool bUndo)
+{
+	if (Graph == nullptr)
+	{
+		return false;
+	}
+	return SelectNode(Graph->FindNode(InNodeName), bSelect, bUndo);
+}
+
+bool URigVMController::ClearNodeSelection(bool bUndo)
+{
+	if (Graph == nullptr)
+	{
+		return false;
+	}
+
+	TArray<FName> Selection = Graph->SelectedNodes;
+	for (const FName& SelectedNode : Selection)
+	{
+		SelectNode(SelectedNode, false, bUndo);
+	}
+
+	return Selection.Num() > 0;
+}
+
+bool URigVMController::IsValidNodeForGraph(URigVMNode* InNode)
+{
+	if (Graph == nullptr)
+	{
+		return false;
+	}
+
+	if (InNode == nullptr)
+	{
+		return false;
+	}
+
+	if (InNode->GetGraph() != Graph)
+	{
+		return false;
+	}
+
+	return true;
+}
diff --git a/Engine/Source/Runtime/RigVM/Private/RigVMModel/RigVMGraph.cpp b/Engine/Source/Runtime/RigVM/Private/RigVMModel/RigVMGraph.cpp
new file mode 100644
--- /dev/null
+++ b/Engine/Source/Runtime/RigVM/Private/RigVMModel/RigVMGraph.cpp
@@ -0,0 +1,48 @@
+// Copyright 1998-2019 Epic Games, Inc. All Rights Reserved.
+
+#include "RigVMModel/RigVMGraph.h"
+#include "RigVMModel/RigVMStructNode.h"
+
+const TArray<URigVMNode*>& URigVMGraph::GetNodes() const
+{
+	return Nodes;
+}
+
+URigVMNode* URigVMGraph::FindNode(const FName& InNodeName)
+{
+	for (URigVMNode* Node : Nodes)
+	{
+		if (Node->GetFName() == InNodeName)
+		{
+			return Node;
+		}
+	}
+	return nullptr;
+}
+
+bool URigVMGraph::IsNodeSelected(const FName& InNodeName) const
+{
+	return SelectedNodes.Contains(InNodeName);
+}
+
+FRigVMGraphModifiedEvent& URigVMGraph::OnModified()
+{
+	return ModifiedEvent;
+}
+
+void URigVMGraph::Notify(ERigVMGraphNotifType InNotifType, URigVMGraph* InGraph, UObject* InSubject)
+{
+	ModifiedEvent.Broadcast(InNotifType, InGraph, InSubject);
+}
+
+bool URigVMGraph::IsNameAvailable(const FString& InName)
+{
+	for (URigVMNode* Node : Nodes)
+	{
+		if (Node->GetName() == InName)
+		{
+			return false;
+		}
+	}
+	return true;
+}
diff --git a/Engine/Source/Runtime/RigVM/Private/RigVMModel/RigVMNode.cpp b/Engine/Source/Runtime/RigVM/Private/RigVMModel/RigVMNode.cpp
new file mode 100644
--- /dev/null
+++ b/Engine/Source/Runtime/RigVM/Private/RigVMModel/RigVMNode.cpp
@@ -0,0 +1,29 @@
+// Copyright 1998-2019 Epic Games, Inc. All Rights Reserved.
+
+#include "RigVMModel/RigVMNode.h"
+#include "RigVMModel/RigVMGraph.h"
+
+const TArray<URigVMPin*>& URigVMNode::GetPins() const
+{
+	return Pins;
+}
+
+URigVMGraph* URigVMNode::GetGraph() const
+{
+	return Cast<URigVMGraph>(GetOuter());
+}
+
+FVector2D URigVMNode::GetPosition() const
+{
+	return Position;
+}
+
+bool URigVMNode::IsSelected() const
+{
+	URigVMGraph* Graph = GetGraph();
+	if (Graph)
+	{
+		return Graph->IsNodeSelected(GetFName());
+	}
+	return false;
+}
diff --git a/Engine/Source/Runtime/RigVM/Private/RigVMModel/RigVMPin.cpp b/Engine/Source/Runtime/RigVM/Private/RigVMModel/RigVMPin.cpp
new file mode 100644
--- /dev/null
+++ b/Engine/Source/Runtime/RigVM/Private/RigVMModel/RigVMPin.cpp
@@ -0,0 +1,86 @@
+// Copyright 1998-2019 Epic Games, Inc. All Rights Reserved.
+
+#include "RigVMModel/RigVMPin.h"
+#include "RigVMModel/RigVMNode.h"
+
+URigVMPin::URigVMPin()
+	: Direction(ERigVMPinDirection::Invalid)
+	, bIsConstant(false)
+	, ArrayIndex(INDEX_NONE)
+	, CPPType(FString())
+{
+}
+
+FString URigVMPin::GetPinPath() const
+{
+	// todo
+	return FString();
+}
+
+ERigVMPinDirection URigVMPin::GetDirection() const
+{
+	return Direction;
+}
+
+bool URigVMPin::IsConstant() const
+{
+	return bIsConstant;
+}
+
+int32 URigVMPin::IsArray() const
+{
+	return ArrayIndex != INDEX_NONE;
+}
+
+int32 URigVMPin::GetArrayIndex() const
+{
+	return ArrayIndex;
+}
+
+FString URigVMPin::GetCPPType() const
+{
+	return CPPType;
+}
+
+URigVMPin* URigVMPin::GetParentPin() const
+{
+	return Cast<URigVMPin>(GetOuter());
+}
+
+const TArray<URigVMPin*>& URigVMPin::GetSubPins() const
+{
+	return SubPins;
+}
+
+const TArray<URigVMPin*>& URigVMPin::GetConnectedPins() const
+{
+	return ConnectedPins;
+}
+
+URigVMNode* URigVMPin::GetNode() const
+{
+	URigVMNode* Node = Cast<URigVMNode>(GetOuter());
+	if(Node)
+	{
+		return Node;
+	}
+
+	URigVMPin* ParentPin = GetParentPin();
+	if(ParentPin)
+	{
+		return ParentPin->GetNode();
+	}
+
+	return nullptr;
+}
+
+URigVMGraph* URigVMPin::GetGraph() const
+{
+	URigVMNode* Node = GetNode();
+	if(Node)
+	{
+		return Node->GetGraph();
+	}
+
+	return nullptr;
+}
diff --git a/Engine/Source/Runtime/RigVM/Private/RigVMModel/RigVMStructNode.cpp b/Engine/Source/Runtime/RigVM/Private/RigVMModel/RigVMStructNode.cpp
new file mode 100644
--- /dev/null
+++ b/Engine/Source/Runtime/RigVM/Private/RigVMModel/RigVMStructNode.cpp
@@ -0,0 +1,13 @@
+// Copyright 1998-2019 Epic Games, Inc. All Rights Reserved.
+
+#include "RigVMModel/RigVMStructNode.h"
+
+UScriptStruct* URigVMStructNode::GetScriptStruct() const
+{
+	return ScriptStruct;
+}
+
+FName URigVMStructNode::GetMethodName() const
+{
+	return MethodName;
+}
diff --git a/Engine/Source/Runtime/RigVM/Public/RigVMModel/RigVMController.h b/Engine/Source/Runtime/RigVM/Public/RigVMModel/RigVMController.h
new file mode 100644
--- /dev/null
+++ b/Engine/Source/Runtime/RigVM/Public/RigVMModel/RigVMController.h
@@ -0,0 +1,43 @@
+// Copyright 1998-2019 Epic Games, Inc. All Rights Reserved.
+
+#pragma once
+
+#include "RigVMGraph.h"
+#include "RigVMModel/RigVMStructNode.h"
+#include "RigVMController.generated.h"
+
+UCLASS()
+class RIGVM_API URigVMController : public UObject
+{
+	GENERATED_BODY()
+
+public:
+
+	~URigVMController();
+
+	URigVMGraph* GetGraph() const;
+	void SetGraph(URigVMGraph* InGraph);
+
+	FRigVMGraphModifiedEvent& OnModified();
+
+	URigVMStructNode* AddStructNode(UScriptStruct* InScriptStruct, const FName& InMethodName, const FVector2D& InPosition = FVector2D::ZeroVector, bool bUndo = true);
+
+	bool RemoveNode(URigVMNode* InNode, bool bUndo = true);
+	bool RemoveNode(const FName& InNodeName, bool bUndo = true);
+
+	bool SelectNode(URigVMNode* InNode, bool bSelect = true, bool bUndo = true);
+	bool SelectNode(const FName& InNodeName, bool bSelect = true, bool bUndo = true);
+	bool ClearNodeSelection(bool bUndo = true);
+
+private:
+
+	FRigVMGraphModifiedEvent ModifiedEvent;
+	void Notify(ERigVMGraphNotifType InNotifType, URigVMGraph* InGraph, UObject* InSubject);
+	void HandleModifiedEvent(ERigVMGraphNotifType InNotifType, URigVMGraph* InGraph, UObject* InSubject);
+
+	bool IsValidNodeForGraph(URigVMNode* InNode);
+
+	UPROPERTY()
+	URigVMGraph* Graph;
+};
+
diff --git a/Engine/Source/Runtime/RigVM/Public/RigVMModel/RigVMGraph.h b/Engine/Source/Runtime/RigVM/Public/RigVMModel/RigVMGraph.h
new file mode 100644
--- /dev/null
+++ b/Engine/Source/Runtime/RigVM/Public/RigVMModel/RigVMGraph.h
@@ -0,0 +1,38 @@
+// Copyright 1998-2019 Epic Games, Inc. All Rights Reserved.
+
+#pragma once
+
+#include "RigVMNode.h"
+#include "RigVMNotifications.h"
+#include "RigVMGraph.generated.h"
+
+UCLASS()
+class RIGVM_API URigVMGraph : public UObject
+{
+	GENERATED_BODY()
+
+public:
+
+	const TArray<URigVMNode*>& GetNodes() const;
+	URigVMNode* FindNode(const FName& InNodeName);
+
+	bool IsNodeSelected(const FName& InNodeName) const;
+
+	FRigVMGraphModifiedEvent& OnModified();
+
+private:
+
+	FRigVMGraphModifiedEvent ModifiedEvent;
+	void Notify(ERigVMGraphNotifType InNotifType, URigVMGraph* InGraph, UObject* InSubject);
+
+	UPROPERTY()
+	TArray<URigVMNode*> Nodes;
+
+	UPROPERTY()
+	TArray<FName> SelectedNodes;
+
+	bool IsNameAvailable(const FString& InName);
+
+	friend class URigVMController;
+};
+
diff --git a/Engine/Source/Runtime/RigVM/Public/RigVMModel/RigVMNode.h b/Engine/Source/Runtime/RigVM/Public/RigVMModel/RigVMNode.h
new file mode 100644
--- /dev/null
+++ b/Engine/Source/Runtime/RigVM/Public/RigVMModel/RigVMNode.h
@@ -0,0 +1,35 @@
+// Copyright 1998-2019 Epic Games, Inc. All Rights Reserved.
+
+#pragma once
+
+#include "RigVMPin.h"
+#include "RigVMNode.generated.h"
+
+class URigVMGraph;
+
+UCLASS()
+class RIGVM_API URigVMNode : public UObject
+{
+	GENERATED_BODY()
+
+public:
+
+	const TArray<URigVMPin*>& GetPins() const;
+
+	URigVMGraph* GetGraph() const;
+
+	FVector2D GetPosition() const;
+
+	bool IsSelected() const;
+
+private:
+
+	UPROPERTY()
+	TArray<URigVMPin*> Pins;
+
+	UPROPERTY()
+	FVector2D Position;
+
+	friend class URigVMController;
+};
+
diff --git a/Engine/Source/Runtime/RigVM/Public/RigVMModel/RigVMNotifications.h b/Engine/Source/Runtime/RigVM/Public/RigVMModel/RigVMNotifications.h
new file mode 100644
--- /dev/null
+++ b/Engine/Source/Runtime/RigVM/Public/RigVMModel/RigVMNotifications.h
@@ -0,0 +1,20 @@
+// Copyright 1998-2019 Epic Games, Inc. All Rights Reserved.
+
+#pragma once
+
+#include "CoreMinimal.h"
+#include "UObject/ObjectMacros.h"
+#include "RigVMNotifications.generated.h"
+
+UENUM()
+enum class ERigVMGraphNotifType : uint8
+{
+	GraphChanged,
+	NodeAdded,
+	NodeRemoved,
+	NodeSelected,
+	NodeDeselected,
+	Invalid
+};
+
+DECLARE_MULTICAST_DELEGATE_ThreeParams(FRigVMGraphModifiedEvent, ERigVMGraphNotifType /* type */, URigVMGraph* /* graph */, UObject* /* subject */);
\ No newline at end of file
diff --git a/Engine/Source/Runtime/RigVM/Public/RigVMModel/RigVMPin.h b/Engine/Source/Runtime/RigVM/Public/RigVMModel/RigVMPin.h
new file mode 100644
--- /dev/null
+++ b/Engine/Source/Runtime/RigVM/Public/RigVMModel/RigVMPin.h
@@ -0,0 +1,66 @@
+// Copyright 1998-2019 Epic Games, Inc. All Rights Reserved.
+
+#pragma once
+
+#include "CoreMinimal.h"
+#include "UObject/ObjectMacros.h"
+#include "RigVMPin.generated.h"
+
+class URigVMGraph;
+class URigVMNode;
+
+UENUM()
+enum class ERigVMPinDirection : uint8
+{
+	Input,
+	Output,
+	IO,
+	Invalid
+};
+
+UCLASS()
+class RIGVM_API URigVMPin : public UObject
+{
+	GENERATED_BODY()
+
+public:
+
+	URigVMPin();
+
+	FString GetPinPath() const;
+	ERigVMPinDirection GetDirection() const;
+	bool IsConstant() const;
+	int32 IsArray() const;
+	int32 GetArrayIndex() const;
+	FString GetCPPType() const;
+
+	URigVMPin* GetParentPin() const;
+	const TArray<URigVMPin*>& GetSubPins() const;
+	const TArray<URigVMPin*>& GetConnectedPins() const;
+
+	URigVMNode* GetNode() const;
+	URigVMGraph* GetGraph() const;
+
+private:
+
+	UPROPERTY()
+	ERigVMPinDirection Direction;
+
+	UPROPERTY()
+	bool bIsConstant;
+
+	UPROPERTY()
+	int32 ArrayIndex;
+
+	UPROPERTY()
+	FString CPPType;
+
+	UPROPERTY()
+	TArray<URigVMPin*> SubPins;
+
+	UPROPERTY()
+	TArray<URigVMPin*> ConnectedPins;
+
+	friend class URigVMController;
+};
+
diff --git a/Engine/Source/Runtime/RigVM/Public/RigVMModel/RigVMStructNode.h b/Engine/Source/Runtime/RigVM/Public/RigVMModel/RigVMStructNode.h
new file mode 100644
--- /dev/null
+++ b/Engine/Source/Runtime/RigVM/Public/RigVMModel/RigVMStructNode.h
@@ -0,0 +1,31 @@
+// Copyright 1998-2019 Epic Games, Inc. All Rights Reserved.
+
+#pragma once
+
+#include "RigVMNode.h"
+#include "RigVMStructNode.generated.h"
+
+UCLASS()
+class RIGVM_API URigVMStructNode : public URigVMNode
+{
+	GENERATED_BODY()
+
+public:
+
+	UFUNCTION()
+	UScriptStruct* GetScriptStruct() const;
+
+	UFUNCTION()
+	FName GetMethodName() const;
+
+private:
+
+	UPROPERTY()
+	UScriptStruct* ScriptStruct;
+
+	UPROPERTY()
+	FName MethodName;
+
+	friend class URigVMController;
+};
+
