RigVM: Added unit test for node selection

diff --git a/Engine/Source/Runtime/RigVM/Private/RigVMModel/RigVMGraph.cpp b/Engine/Source/Runtime/RigVM/Private/RigVMModel/RigVMGraph.cpp
--- a/Engine/Source/Runtime/RigVM/Private/RigVMModel/RigVMGraph.cpp
+++ b/Engine/Source/Runtime/RigVM/Private/RigVMModel/RigVMGraph.cpp
@@ -25,6 +25,11 @@ bool URigVMGraph::IsNodeSelected(const FName& InNodeName) const
 	return SelectedNodes.Contains(InNodeName);
 }
 
+const TArray<FName>& URigVMGraph::GetSelectNodes() const
+{
+	return SelectedNodes;
+}
+
 FRigVMGraphModifiedEvent& URigVMGraph::OnModified()
 {
 	return ModifiedEvent;
diff --git a/Engine/Source/Runtime/RigVM/Public/RigVMModel/RigVMGraph.h b/Engine/Source/Runtime/RigVM/Public/RigVMModel/RigVMGraph.h
--- a/Engine/Source/Runtime/RigVM/Public/RigVMModel/RigVMGraph.h
+++ b/Engine/Source/Runtime/RigVM/Public/RigVMModel/RigVMGraph.h
@@ -17,6 +17,7 @@ public:
 	URigVMNode* FindNode(const FName& InNodeName);
 
 	bool IsNodeSelected(const FName& InNodeName) const;
+	const TArray<FName>& GetSelectNodes() const;
 
 	FRigVMGraphModifiedEvent& OnModified();
 
