RigVM: Move away from isliterals bool to storage type enum

diff --git a/Engine/Source/Programs/UnrealHeaderTool/Private/CodeGenerator.cpp b/Engine/Source/Programs/UnrealHeaderTool/Private/CodeGenerator.cpp
--- a/Engine/Source/Programs/UnrealHeaderTool/Private/CodeGenerator.cpp
+++ b/Engine/Source/Programs/UnrealHeaderTool/Private/CodeGenerator.cpp
@@ -3400,7 +3400,7 @@ void FNativeClassHeaderGenerator::ExportGeneratedStructBodyMacros(FOutputDevice&
 
 			MemberDeclarations.Add(FString::Printf(TEXT("%s %s"), *ParameterCPPType, *Name));
 			FString AddressName = FString::Printf(TEXT("%s_Arg"), *Name);
-			MemberMulticastProlog.Add(FString::Printf(TEXT("%s %s = RigVMStorage[%s.StorageType()]->%s<%s>(%s.GetRegisterIndex());"), *ExtractedMulticastCPPType, *Name, *AddressName, *MulticastGetter, *MultiCastCPPType, *AddressName));
+			MemberMulticastProlog.Add(FString::Printf(TEXT("%s %s = RigVMStorage[%s.GetStorageIndex()]->%s<%s>(%s.GetRegisterIndex());"), *ExtractedMulticastCPPType, *Name, *AddressName, *MulticastGetter, *MultiCastCPPType, *AddressName));
 		}
 
 		OutGeneratedHeaderText.Log(TEXT("\n"));
diff --git a/Engine/Source/Runtime/RigVM/Private/RigVMCore/RigVM.cpp b/Engine/Source/Runtime/RigVM/Private/RigVMCore/RigVM.cpp
--- a/Engine/Source/Runtime/RigVM/Private/RigVMCore/RigVM.cpp
+++ b/Engine/Source/Runtime/RigVM/Private/RigVMCore/RigVM.cpp
@@ -4,8 +4,8 @@
 
 URigVM::URigVM()
 {
-	Literals.SetLiteralStorage(true);
-	WorkState.SetLiteralStorage(false);
+	WorkStorage.SetStorageType(ERigVMStorageType::Work);
+	LiteralStorage.SetStorageType(ERigVMStorageType::Literal);
 }
 
 URigVM::~URigVM()
@@ -15,8 +15,8 @@ URigVM::~URigVM()
 
 void URigVM::Reset()
 {
-	Literals.Reset();
-	WorkState.Reset();
+	WorkStorage.Reset();
+	LiteralStorage.Reset();
 	FunctionNames.Reset();
 	Functions.Reset();
 	ByteCode.Reset();
@@ -75,7 +75,7 @@ bool URigVM::Execute(FRigVMStoragePtrArray Storage, TArrayView<void*> Additional
 		return true;
 	}
 
-	FRigVMStorage* LocalStorage[] = { &WorkState, &Literals }; 
+	FRigVMStorage* LocalStorage[] = { &WorkStorage, &LiteralStorage }; 
 	if (Storage.Num() == 0)
 	{
 		Storage = FRigVMStoragePtrArray(LocalStorage, 2);
@@ -89,7 +89,7 @@ bool URigVM::Execute(FRigVMStoragePtrArray Storage, TArrayView<void*> Additional
 			case ERigVMOpCode::Copy:
 			{
 				const FRigVMCopyOp& Op = ByteCode.GetOpAt<FRigVMCopyOp>(Instructions[InstructionIndex]);
-				Storage[0]->Copy(Op.Source.GetRegisterIndex(), Op.Target.GetRegisterIndex(), Storage[Op.Source.StorageType()], Op.SourceOffset, Op.TargetOffset, Op.NumBytes);
+				Storage[0]->Copy(Op.Source.GetRegisterIndex(), Op.Target.GetRegisterIndex(), Storage[Op.Source.GetStorageIndex()], Op.SourceOffset, Op.TargetOffset, Op.NumBytes);
 				InstructionIndex++;
 				break;
 			}
diff --git a/Engine/Source/Runtime/RigVM/Private/RigVMCore/RigVMStorage.cpp b/Engine/Source/Runtime/RigVM/Private/RigVMCore/RigVMStorage.cpp
--- a/Engine/Source/Runtime/RigVM/Private/RigVMCore/RigVMStorage.cpp
+++ b/Engine/Source/Runtime/RigVM/Private/RigVMCore/RigVMStorage.cpp
@@ -4,7 +4,7 @@
 
 FRigVMStorage::FRigVMStorage(bool bInUseNames)
 	: bUseNameMap(bInUseNames)
-	, bIsLiteralStorage(false)
+	, StorageType(ERigVMStorageType::Work)
 {
 }
 
@@ -75,25 +75,25 @@ bool FRigVMStorage::Copy(
 	const FRigVMRegister& Source = InSourceStorage->Registers[InSourceRegisterIndex];
 	const FRigVMRegister& Target = Registers[InTargetRegisterIndex];
 
-	int32 SourceStartByte = Source.FirstByte();
-	int32 SourceNumBytes = Source.NumBytes(false);
+	int32 SourceStartByte = Source.GetWorkByteIndex();
+	int32 SourceNumBytes = Source.GetNumBytesAllSlices();
 	if(InSourceByteOffset != INDEX_NONE)
 	{
 		ensure(InNumBytes != INDEX_NONE);
-		ensure(InSourceByteOffset >= 0 && InSourceByteOffset < Source.NumBytes(false));
-		ensure(InNumBytes > 0 && InSourceByteOffset + InNumBytes <= Source.NumBytes(false));
+		ensure(InSourceByteOffset >= 0 && InSourceByteOffset < Source.GetNumBytesAllSlices());
+		ensure(InNumBytes > 0 && InSourceByteOffset + InNumBytes <= Source.GetNumBytesAllSlices());
 
 		SourceStartByte += InSourceByteOffset;
 		SourceNumBytes = InNumBytes;
 	}
 
-	int32 TargetStartByte = Target.FirstByte();
-	int32 TargetNumBytes = Target.NumBytes(false);
+	int32 TargetStartByte = Target.GetWorkByteIndex();
+	int32 TargetNumBytes = Target.GetNumBytesAllSlices();
 	if(InTargetByteOffset != INDEX_NONE)
 	{
 		ensure(InNumBytes != INDEX_NONE);
-		ensure(InTargetByteOffset >= 0 && InTargetByteOffset < Target.NumBytes(false));
-		ensure(InNumBytes > 0 && InTargetByteOffset + InNumBytes <= Target.NumBytes(false));
+		ensure(InTargetByteOffset >= 0 && InTargetByteOffset < Target.GetNumBytesAllSlices());
+		ensure(InNumBytes > 0 && InTargetByteOffset + InNumBytes <= Target.GetNumBytesAllSlices());
 
 		TargetStartByte += InTargetByteOffset;
 		TargetNumBytes = InNumBytes;
@@ -203,11 +203,11 @@ int32 FRigVMStorage::Allocate(const FName& InNewName, int32 InElementSize, int32
 	NewRegister.ElementCount = InCount;
 	NewRegister.Type = ERigVMRegisterType::Plain;
 
-	Data.AddZeroed(NewRegister.NumBytes());
+	Data.AddZeroed(NewRegister.GetAllocatedBytes());
 
 	if (InDataPtr != nullptr)
 	{
-		FMemory::Memcpy(&Data[NewRegister.FirstByte()], InDataPtr, NewRegister.NumBytes());
+		FMemory::Memcpy(&Data[NewRegister.GetWorkByteIndex()], InDataPtr, NewRegister.GetAllocatedBytes());
 	}
 
 	int32 RegisterIndex = Registers.Num();
@@ -234,7 +234,7 @@ bool FRigVMStorage::Construct(int32 InRegisterIndex, int32 InElementIndex)
 	{
 		case ERigVMRegisterType::Struct:
 		{
-			void* DataPtr = (void*)&Data[InElementIndex == INDEX_NONE ? Register.FirstByte() : Register.FirstByte() + InElementIndex * Register.ElementSize];
+			void* DataPtr = (void*)&Data[InElementIndex == INDEX_NONE ? Register.GetWorkByteIndex() : Register.GetWorkByteIndex() + InElementIndex * Register.ElementSize];
 			int32 Count = InElementIndex == INDEX_NONE ? Register.ElementCount : 1;
 
 			UScriptStruct* ScriptStruct = GetScriptStruct(InRegisterIndex);
@@ -243,7 +243,7 @@ bool FRigVMStorage::Construct(int32 InRegisterIndex, int32 InElementIndex)
 		}
 		case ERigVMRegisterType::String:
 		{
-			FString* DataPtr = (FString*)&Data[InElementIndex == INDEX_NONE ? Register.FirstByte() : Register.FirstByte() + InElementIndex * Register.ElementSize];
+			FString* DataPtr = (FString*)&Data[InElementIndex == INDEX_NONE ? Register.GetWorkByteIndex() : Register.GetWorkByteIndex() + InElementIndex * Register.ElementSize];
 			int32 Count = InElementIndex == INDEX_NONE ? Register.ElementCount : 1;
 
 			FMemory::Memzero(DataPtr, Count * Register.ElementSize);
@@ -255,7 +255,7 @@ bool FRigVMStorage::Construct(int32 InRegisterIndex, int32 InElementIndex)
 		}
 		case ERigVMRegisterType::Name:
 		{
-			FName* DataPtr = (FName*)&Data[InElementIndex == INDEX_NONE ? Register.FirstByte() : Register.FirstByte() + InElementIndex * Register.ElementSize];
+			FName* DataPtr = (FName*)&Data[InElementIndex == INDEX_NONE ? Register.GetWorkByteIndex() : Register.GetWorkByteIndex() + InElementIndex * Register.ElementSize];
 			int32 Count = InElementIndex == INDEX_NONE ? Register.ElementCount : 1;
 
 			FMemory::Memzero(DataPtr, Count * Register.ElementSize);
@@ -283,7 +283,7 @@ bool FRigVMStorage::Destroy(int32 InRegisterIndex, int32 InElementIndex)
 	{
 		case ERigVMRegisterType::Struct:
 		{
-			void* DataPtr = (void*)&Data[InElementIndex == INDEX_NONE ? Register.FirstByte() : Register.FirstByte() + InElementIndex * Register.ElementSize];
+			void* DataPtr = (void*)&Data[InElementIndex == INDEX_NONE ? Register.GetWorkByteIndex() : Register.GetWorkByteIndex() + InElementIndex * Register.ElementSize];
 			int32 Count = InElementIndex == INDEX_NONE ? Register.ElementCount : 1;
 
 			UScriptStruct* ScriptStruct = GetScriptStruct(InRegisterIndex);
@@ -292,7 +292,7 @@ bool FRigVMStorage::Destroy(int32 InRegisterIndex, int32 InElementIndex)
 		}
 		case ERigVMRegisterType::String:
 		{
-			FString* DataPtr = (FString*)&Data[InElementIndex == INDEX_NONE ? Register.FirstByte() : Register.FirstByte() + InElementIndex * Register.ElementSize];
+			FString* DataPtr = (FString*)&Data[InElementIndex == INDEX_NONE ? Register.GetWorkByteIndex() : Register.GetWorkByteIndex() + InElementIndex * Register.ElementSize];
 			int32 Count = InElementIndex == INDEX_NONE ? Register.ElementCount : 1;
 
 			for (int32 Index = 0; Index < Count; Index++)
@@ -303,7 +303,7 @@ bool FRigVMStorage::Destroy(int32 InRegisterIndex, int32 InElementIndex)
 		}
 		case ERigVMRegisterType::Name:
 		{
-			FName* DataPtr = (FName*)&Data[InElementIndex == INDEX_NONE ? Register.FirstByte() : Register.FirstByte() + InElementIndex * Register.ElementSize];
+			FName* DataPtr = (FName*)&Data[InElementIndex == INDEX_NONE ? Register.GetWorkByteIndex() : Register.GetWorkByteIndex() + InElementIndex * Register.ElementSize];
 			int32 Count = InElementIndex == INDEX_NONE ? Register.ElementCount : 1;
 
 			for (int32 Index = 0; Index < Count; Index++)
@@ -331,12 +331,12 @@ bool FRigVMStorage::Remove(int32 InRegisterIndex)
 	Destroy(InRegisterIndex);
 
 	FRigVMRegister RegisterToRemove = Registers[InRegisterIndex];
-	Data.RemoveAt(RegisterToRemove.ByteIndex, RegisterToRemove.NumBytes());
+	Data.RemoveAt(RegisterToRemove.ByteIndex, RegisterToRemove.GetAllocatedBytes());
 	Registers.RemoveAt(InRegisterIndex);
 
 	for (int32 Index = InRegisterIndex; Index < Registers.Num(); Index++)
 	{
-		Registers[Index].ByteIndex -= RegisterToRemove.NumBytes();
+		Registers[Index].ByteIndex -= RegisterToRemove.GetAllocatedBytes();
 	}
 
 	UpdateRegisters();
@@ -380,7 +380,7 @@ FName FRigVMStorage::Rename(const FName& InOldName, const FName& InNewName)
 	return Rename(RegisterIndex, InNewName);
 }
 
-bool FRigVMStorage::Resize(int32 InRegisterIndex, int32 InNewElementCount)
+bool FRigVMStorage::Resize(int32 InRegisterIndex, int32 InNewElementCount, int32 InNewSliceCount)
 {
 	if (InNewElementCount <= 0)
 	{
@@ -398,7 +398,7 @@ bool FRigVMStorage::Resize(int32 InRegisterIndex, int32 InNewElementCount)
 	{
 		int32 ElementsToRemove = Register.ElementCount - InNewElementCount;
 		int32 NumBytesToRemove = Register.ElementSize * ElementsToRemove;
-		int32 FirstByteToRemove = Register.FirstByte() + Register.ElementSize * InNewElementCount;
+		int32 FirstByteToRemove = Register.GetWorkByteIndex() + Register.ElementSize * InNewElementCount;
 
 		for (int32 ElementIndex = InNewElementCount; ElementIndex < Register.ElementCount; ElementIndex++)
 		{
@@ -418,7 +418,7 @@ bool FRigVMStorage::Resize(int32 InRegisterIndex, int32 InNewElementCount)
 		int32 OldElementCount = Register.ElementCount;
 		int32 ElementsToAdd = InNewElementCount - Register.ElementCount;
 		int32 NumBytesToAdd = Register.ElementSize * ElementsToAdd;
-		int32 FirstByteToAdd = Register.FirstByte() + Register.ElementSize * Register.ElementCount;
+		int32 FirstByteToAdd = Register.GetWorkByteIndex() + Register.ElementSize * Register.ElementCount;
 
 		Data.InsertZeroed(FirstByteToAdd, NumBytesToAdd);
 		Register.ElementCount = InNewElementCount;
@@ -438,7 +438,7 @@ bool FRigVMStorage::Resize(int32 InRegisterIndex, int32 InNewElementCount)
 	return true;
 }
 
-bool FRigVMStorage::Resize(const FName& InRegisterName, int32 InNewElementCount)
+bool FRigVMStorage::Resize(const FName& InRegisterName, int32 InNewElementCount, int32 InNewSliceCount)
 {
 	ensure(bUseNameMap);
 
@@ -448,7 +448,7 @@ bool FRigVMStorage::Resize(const FName& InRegisterName, int32 InNewElementCount)
 		return false;
 	}
 
-	return Resize(RegisterIndex, InNewElementCount);
+	return Resize(RegisterIndex, InNewElementCount, InNewSliceCount);
 }
 
 void FRigVMStorage::UpdateRegisters()
@@ -458,6 +458,7 @@ void FRigVMStorage::UpdateRegisters()
 	{
 		FRigVMRegister& Register = Registers[RegisterIndex];
 		Register.ByteIndex += AlignmentShift;
+		Register.MoveToFirstSlice();
 
 		UScriptStruct* ScriptStruct = GetScriptStruct(RegisterIndex);
 		if (ScriptStruct != nullptr)
@@ -473,8 +474,9 @@ void FRigVMStorage::UpdateRegisters()
 					{
 						if (!IsAligned(Pointer, TheCppStructOps->GetAlignment()))
 						{
-							Data.RemoveAt(Register.ByteIndex, Register.AlignmentBytes);
+							Data.RemoveAt(Register.GetStorageByteIndex(), Register.AlignmentBytes);
 							AlignmentShift -= Register.AlignmentBytes;
+							Register.ByteIndex -= Register.AlignmentBytes;
 							Register.AlignmentBytes = 0;
 							Pointer = GetData(RegisterIndex);
 						}
@@ -482,8 +484,9 @@ void FRigVMStorage::UpdateRegisters()
 
 					while (!IsAligned(Pointer, TheCppStructOps->GetAlignment()))
 					{
-						Data.InsertZeroed(Register.ByteIndex, 1);
+						Data.InsertZeroed(Register.GetStorageByteIndex(), 1);
 						Register.AlignmentBytes++;
+						Register.ByteIndex++;
 						AlignmentShift++;
 						Pointer = GetData(RegisterIndex);
 					}
@@ -492,11 +495,6 @@ void FRigVMStorage::UpdateRegisters()
 		}
 	}
 
-	for (int32 RegisterIndex = 0; RegisterIndex < Registers.Num(); RegisterIndex++)
-	{
-		Registers[RegisterIndex].Pointer = GetData(RegisterIndex);
-	}
-
 	if (bUseNameMap)
 	{
 		NameMap.Reset();
@@ -510,7 +508,7 @@ void FRigVMStorage::UpdateRegisters()
 void FRigVMStorage::FillWithZeroes(int32 InRegisterIndex)
 {
 	ensure(Registers.IsValidIndex(InRegisterIndex));
-	FMemory::Memzero(GetData(InRegisterIndex), Registers[InRegisterIndex].NumBytes());
+	FMemory::Memzero(GetData(InRegisterIndex), Registers[InRegisterIndex].GetNumBytesAllSlices());
 }
 
 int32 FRigVMStorage::FindOrAddScriptStruct(UScriptStruct* InScriptStruct)
diff --git a/Engine/Source/Runtime/RigVM/Public/RigVMCore/RigVM.h b/Engine/Source/Runtime/RigVM/Public/RigVMCore/RigVM.h
--- a/Engine/Source/Runtime/RigVM/Public/RigVMCore/RigVM.h
+++ b/Engine/Source/Runtime/RigVM/Public/RigVMCore/RigVM.h
@@ -26,10 +26,10 @@ public:
 	int32 AddRigVMFunction(UScriptStruct* InRigVMStruct, const FName& InFunctionName);
 
 	UPROPERTY()
-	FRigVMStorage Literals;
+	FRigVMStorage WorkStorage;
 
 	UPROPERTY()
-	FRigVMStorage WorkState;
+	FRigVMStorage LiteralStorage;
 
 	UPROPERTY()
 	FRigVMByteCode ByteCode;
diff --git a/Engine/Source/Runtime/RigVM/Public/RigVMCore/RigVMStorage.h b/Engine/Source/Runtime/RigVM/Public/RigVMCore/RigVMStorage.h
--- a/Engine/Source/Runtime/RigVM/Public/RigVMCore/RigVMStorage.h
+++ b/Engine/Source/Runtime/RigVM/Public/RigVMCore/RigVMStorage.h
@@ -6,34 +6,42 @@
 #include "UObject/ObjectMacros.h"
 #include "RigVMStorage.generated.h"
 
+UENUM()
+enum class ERigVMStorageType: uint8
+{
+	Work,
+	Literal,
+	Invalid
+};
+
 struct RIGVM_API FRigVMArgument
 {
 public:
 
 	FRigVMArgument()
-		: bIsLiteral(false)
+		: StorageType(ERigVMStorageType::Work)
 		, RegisterIndex(INDEX_NONE)
 		, ByteIndex(INDEX_NONE)
 	{
 	}
 
-	FRigVMArgument(bool InIsLiteral, int32 InRegisterIndex, int32 InByteIndex)
-		: bIsLiteral(InIsLiteral)
+	FRigVMArgument(ERigVMStorageType InStorageType, int32 InRegisterIndex, int32 InByteIndex)
+		: StorageType(InStorageType)
 		, RegisterIndex(InRegisterIndex)
 		, ByteIndex(InByteIndex)
 	{
 	}
 
-	FORCEINLINE bool IsLiteral() const { return bIsLiteral; }
-	FORCEINLINE int32 StorageType() const { return bIsLiteral ? 1 : 0; }
-	FORCEINLINE int32 GetRegisterIndex() const { return RegisterIndex; }
-	FORCEINLINE int32 GetByteIndex() const { return ByteIndex; }
+	FORCEINLINE ERigVMStorageType GetStorageType() const { return StorageType; }
+	FORCEINLINE int32 GetStorageIndex() const { return (int32)StorageType; }
+	FORCEINLINE uint16 GetRegisterIndex() const { return RegisterIndex; }
+	FORCEINLINE uint64 GetByteIndex() const { return ByteIndex; }
 
 private:
 
-	bool bIsLiteral;
-	int32 RegisterIndex;
-	int32 ByteIndex;
+	ERigVMStorageType StorageType;
+	uint16 RegisterIndex;
+	uint64 ByteIndex;
 };
 
 typedef TArrayView<FRigVMArgument> FRigVMArgumentArray;
@@ -55,10 +63,11 @@ struct RIGVM_API FRigVMRegister
 
 		FRigVMRegister()
 		: Type(ERigVMRegisterType::Invalid)
-		, Pointer(nullptr)
 		, ByteIndex(INDEX_NONE)
 		, ElementSize(0)
 		, ElementCount(0)
+		, SliceIndex(0)
+		, SliceCount(1)
 		, AlignmentBytes(0)
 		, Name(NAME_None)
 		, ScriptStructIndex(INDEX_NONE)
@@ -68,19 +77,23 @@ struct RIGVM_API FRigVMRegister
 	UPROPERTY()
 	ERigVMRegisterType Type;
 
-	void* Pointer;
+	UPROPERTY()
+	uint32 ByteIndex;
+
+	UPROPERTY()
+	uint16 ElementSize;
 
 	UPROPERTY()
-	int32 ByteIndex;
+	uint16 ElementCount;
 
 	UPROPERTY()
-	int32 ElementSize;
+	uint16 SliceIndex;
 
 	UPROPERTY()
-	int32 ElementCount;
+	uint16 SliceCount;
 
 	UPROPERTY()
-	int32 AlignmentBytes;
+	uint8 AlignmentBytes;
 
 	UPROPERTY()
 	FName Name;
@@ -88,43 +101,16 @@ struct RIGVM_API FRigVMRegister
 	UPROPERTY()
 	int32 ScriptStructIndex;
 
-	FORCEINLINE int32 FirstByte() const { return ByteIndex + AlignmentBytes; }
+	FORCEINLINE uint64 GetWorkByteIndex() const { return ByteIndex; }
+	FORCEINLINE uint64 GetStorageByteIndex() const { return ByteIndex - (uint64)AlignmentBytes - (uint64)(SliceIndex * GetNumBytesPerSlice()); }
+	FORCEINLINE uint8 GetAlignmentBytes() const { return AlignmentBytes; }
 	FORCEINLINE bool IsArray() const { return ElementCount > 1; }
 	FORCEINLINE bool IsPlain() const { return ScriptStructIndex == INDEX_NONE;  }
-	FORCEINLINE int32 NumBytes(bool bIncludeAlignment = true) const { return ElementCount * ElementSize + (bIncludeAlignment ? AlignmentBytes : 0); }
-
-	template<class T>
-	FORCEINLINE const T* Get() const
-	{
-		ensure(ElementCount > 0);
-		return (const T*)Pointer;
-	}
-
-	template<class T>
-	FORCEINLINE const T& GetRef() const
-	{
-		return *Get<T>();
-	}
-
-	template<class T>
-	FORCEINLINE T* Get()
-	{
-		ensure(ElementCount > 0);
-		return (T*)Pointer;
-	}
-
-	template<class T>
-	FORCEINLINE T& GetRef()
-	{
-		return *Get<T>();
-	}
-
-	template<class T>
-	FORCEINLINE TArrayView<T> GetArray()
-	{
-		ensure(ElementCount > 0);
-		return TArrayView<T>((T*)Pointer, ElementCount);
-	}
+	FORCEINLINE uint16 GetAllocatedBytes() const { return ElementCount * ElementSize * SliceCount + (uint16)AlignmentBytes; }
+	FORCEINLINE uint16 GetNumBytesPerSlice() const { return ElementCount * ElementSize; }
+	FORCEINLINE uint16 GetNumBytesAllSlices() const { return ElementCount * ElementSize * SliceCount; }
+	FORCEINLINE void MoveToFirstSlice() { ByteIndex -= SliceIndex * GetNumBytesPerSlice(); SliceIndex = 0; }
+	FORCEINLINE void MoveToNextSlice() { ByteIndex += GetNumBytesPerSlice(); SliceIndex++; }
 };
 
 typedef TArrayView<FRigVMRegister> FRigVMRegisterArray;
@@ -142,8 +128,8 @@ public:
 
 	FRigVMStorage& operator= (const FRigVMStorage &InOther);
 
-	FORCEINLINE bool IsLiteralStorage() const { return bIsLiteralStorage;  }
-	FORCEINLINE void SetLiteralStorage(bool InIsLiteralStorage = true) { bIsLiteralStorage = InIsLiteralStorage; }
+	FORCEINLINE ERigVMStorageType GetStorageType() const { return StorageType;  }
+	FORCEINLINE void SetStorageType(ERigVMStorageType InStorageType) { StorageType = InStorageType; }
 	FORCEINLINE bool SupportsNames() const { return bUseNameMap;  }
 	FORCEINLINE int32 Num() const { return Registers.Num(); }
 	FORCEINLINE const FRigVMRegister& operator[](int32 InIndex) const { return Registers[InIndex]; }
@@ -161,7 +147,7 @@ public:
 	FORCEINLINE FRigVMArgument GetArgument(int32 InRegisterIndex) const
 	{
 		ensure(Registers.IsValidIndex(InRegisterIndex));
-		return FRigVMArgument(IsLiteralStorage(), InRegisterIndex, Registers[InRegisterIndex].ByteIndex);
+		return FRigVMArgument(StorageType, InRegisterIndex, Registers[InRegisterIndex].ByteIndex);
 	}
 
 	FORCEINLINE const void* GetData(int32 InRegisterIndex) const
@@ -169,7 +155,7 @@ public:
 		ensure(Registers.IsValidIndex(InRegisterIndex));
 		const FRigVMRegister& Register = Registers[InRegisterIndex];
 		ensure(Register.ElementCount > 0);
-		return (const void*)&Data[Register.FirstByte()];
+		return (const void*)&Data[Register.GetWorkByteIndex()];
 	}
 
 	FORCEINLINE void* GetData(int32 InRegisterIndex)
@@ -177,7 +163,7 @@ public:
 		ensure(Registers.IsValidIndex(InRegisterIndex));
 		const FRigVMRegister& Register = Registers[InRegisterIndex];
 		ensure(Register.ElementCount > 0);
-		return (void*)&Data[Register.FirstByte()];
+		return (void*)&Data[Register.GetWorkByteIndex()];
 	}
 
 	template<class T>
@@ -186,7 +172,7 @@ public:
 		ensure(Registers.IsValidIndex(InRegisterIndex));
 		const FRigVMRegister& Register = Registers[InRegisterIndex];
 		ensure(Register.ElementCount > 0);
-		return (const T*)&Data[Register.FirstByte()];
+		return (const T*)&Data[Register.GetWorkByteIndex()];
 	}
 
 	template<class T>
@@ -213,7 +199,7 @@ public:
 		ensure(Registers.IsValidIndex(InRegisterIndex));
 		const FRigVMRegister& Register = Registers[InRegisterIndex];
 		ensure(Register.ElementCount > 0);
-		return (T*)&Data[Register.FirstByte()];
+		return (T*)&Data[Register.GetWorkByteIndex()];
 	}
 
 	template<class T>
@@ -240,7 +226,7 @@ public:
 		ensure(Registers.IsValidIndex(InRegisterIndex));
 		const FRigVMRegister& Register = Registers[InRegisterIndex];
 		ensure(Register.ElementCount > 0);
-		return TArrayView<T>((T*)&Data[Register.FirstByte()], Register.ElementCount);
+		return TArrayView<T>((T*)&Data[Register.GetWorkByteIndex()], Register.ElementCount);
 	}
 	
 	template<class T>
@@ -522,8 +508,8 @@ public:
 	bool Remove(const FName& InRegisterName);
 	FName Rename(int32 InRegisterIndex, const FName& InNewName);
 	FName Rename(const FName& InOldName, const FName& InNewName);
-	bool Resize(int32 InRegisterIndex, int32 InNewElementCount);
-	bool Resize(const FName& InRegisterName, int32 InNewElementCount);
+	bool Resize(int32 InRegisterIndex, int32 InNewElementCount, int32 InNewSliceCount = 1);
+	bool Resize(const FName& InRegisterName, int32 InNewElementCount, int32 InNewSliceCount = 1);
 
 	void UpdateRegisters();
 
@@ -541,7 +527,7 @@ private:
 	bool bUseNameMap;
 
 	UPROPERTY()
-	bool bIsLiteralStorage;
+	ERigVMStorageType StorageType;
 
 	UPROPERTY()
 	TArray<FRigVMRegister> Registers;
