RigVM: Implement support for slices

diff --git a/Engine/Source/Programs/UnrealHeaderTool/Private/CodeGenerator.cpp b/Engine/Source/Programs/UnrealHeaderTool/Private/CodeGenerator.cpp
--- a/Engine/Source/Programs/UnrealHeaderTool/Private/CodeGenerator.cpp
+++ b/Engine/Source/Programs/UnrealHeaderTool/Private/CodeGenerator.cpp
@@ -3400,7 +3400,7 @@ void FNativeClassHeaderGenerator::ExportGeneratedStructBodyMacros(FOutputDevice&
 
 			MemberDeclarations.Add(FString::Printf(TEXT("%s %s"), *ParameterCPPType, *Name));
 			FString AddressName = FString::Printf(TEXT("%s_Arg"), *Name);
-			MemberMulticastProlog.Add(FString::Printf(TEXT("%s %s = RigVMStorage[%s.GetStorageIndex()]->%s<%s>(%s.GetRegisterIndex());"), *ExtractedMulticastCPPType, *Name, *AddressName, *MulticastGetter, *MultiCastCPPType, *AddressName));
+			MemberMulticastProlog.Add(FString::Printf(TEXT("%s %s = RigVMStorage[%s.GetStorageIndex()]->%s<%s>(%s.GetRegisterIndex(), true);"), *ExtractedMulticastCPPType, *Name, *AddressName, *MulticastGetter, *MultiCastCPPType, *AddressName));
 		}
 
 		OutGeneratedHeaderText.Log(TEXT("\n"));
diff --git a/Engine/Source/Runtime/RigVM/Private/RigVMCore/RigVMStorage.cpp b/Engine/Source/Runtime/RigVM/Private/RigVMCore/RigVMStorage.cpp
--- a/Engine/Source/Runtime/RigVM/Private/RigVMCore/RigVMStorage.cpp
+++ b/Engine/Source/Runtime/RigVM/Private/RigVMCore/RigVMStorage.cpp
@@ -170,7 +170,7 @@ bool FRigVMStorage::Copy(
 	return Copy(SourceRegisterIndex, TargetRegisterIndex, InSourceStorage, InSourceByteOffset, InTargetByteOffset, InNumBytes);
 }
 
-int32 FRigVMStorage::Allocate(const FName& InNewName, int32 InElementSize, int32 InCount, const void* InDataPtr, bool bUpdateRegisters)
+int32 FRigVMStorage::Allocate(const FName& InNewName, int32 InElementSize, int32 InElementCount, int32 InSliceCount, const void* InDataPtr, bool bUpdateRegisters)
 {
 	FName Name = InNewName;
 	if (bUseNameMap && InNewName == NAME_None)
@@ -183,7 +183,7 @@ int32 FRigVMStorage::Allocate(const FName& InNewName, int32 InElementSize, int32
 		} while (!IsNameAvailable(Name));
 	}
 
-	ensure(InElementSize > 0 && InCount > 0);
+	ensure(InElementSize > 0 && InElementCount > 0 && InSliceCount > 0);
 
 	if (bUseNameMap)
 	{
@@ -200,14 +200,21 @@ int32 FRigVMStorage::Allocate(const FName& InNewName, int32 InElementSize, int32
 		NewRegister.Name = Name;
 	}
 	NewRegister.ElementSize = InElementSize;
-	NewRegister.ElementCount = InCount;
+	NewRegister.ElementCount = InElementCount;
+	NewRegister.SliceCount = InSliceCount;
 	NewRegister.Type = ERigVMRegisterType::Plain;
 
 	Data.AddZeroed(NewRegister.GetAllocatedBytes());
 
 	if (InDataPtr != nullptr)
 	{
-		FMemory::Memcpy(&Data[NewRegister.GetWorkByteIndex()], InDataPtr, NewRegister.GetAllocatedBytes());
+		NewRegister.MoveToFirstSlice();
+		for (uint16 SliceIndex = 0; SliceIndex < NewRegister.SliceCount; SliceIndex++)
+		{
+			FMemory::Memcpy(&Data[NewRegister.GetWorkByteIndex()], InDataPtr, NewRegister.GetNumBytesPerSlice());
+			NewRegister.MoveToNextSlice();
+		}
+		NewRegister.MoveToFirstSlice();
 	}
 
 	int32 RegisterIndex = Registers.Num();
@@ -220,9 +227,9 @@ int32 FRigVMStorage::Allocate(const FName& InNewName, int32 InElementSize, int32
 	return RegisterIndex;
 }
 
-int32 FRigVMStorage::Allocate(int32 InElementSize, int32 InCount, const void* InDataPtr, bool bUpdateRegisters)
+int32 FRigVMStorage::Allocate(int32 InElementSize, int32 InElementCount, int32 InSliceCount, const void* InDataPtr, bool bUpdateRegisters)
 {
-	return Allocate(NAME_None, InElementSize, InCount, InDataPtr, bUpdateRegisters);
+	return Allocate(NAME_None, InElementSize, InElementCount, InSliceCount, InDataPtr, bUpdateRegisters);
 }
 
 bool FRigVMStorage::Construct(int32 InRegisterIndex, int32 InElementIndex)
@@ -235,7 +242,7 @@ bool FRigVMStorage::Construct(int32 InRegisterIndex, int32 InElementIndex)
 		case ERigVMRegisterType::Struct:
 		{
 			void* DataPtr = (void*)&Data[InElementIndex == INDEX_NONE ? Register.GetWorkByteIndex() : Register.GetWorkByteIndex() + InElementIndex * Register.ElementSize];
-			int32 Count = InElementIndex == INDEX_NONE ? Register.ElementCount : 1;
+			int32 Count = InElementIndex == INDEX_NONE ? Register.GetTotalElementCount() : 1;
 
 			UScriptStruct* ScriptStruct = GetScriptStruct(InRegisterIndex);
 			ScriptStruct->InitializeStruct(DataPtr, Count);
@@ -244,7 +251,7 @@ bool FRigVMStorage::Construct(int32 InRegisterIndex, int32 InElementIndex)
 		case ERigVMRegisterType::String:
 		{
 			FString* DataPtr = (FString*)&Data[InElementIndex == INDEX_NONE ? Register.GetWorkByteIndex() : Register.GetWorkByteIndex() + InElementIndex * Register.ElementSize];
-			int32 Count = InElementIndex == INDEX_NONE ? Register.ElementCount : 1;
+			int32 Count = InElementIndex == INDEX_NONE ? Register.GetTotalElementCount() : 1;
 
 			FMemory::Memzero(DataPtr, Count * Register.ElementSize);
 			for (int32 Index = 0; Index < Count; Index++)
@@ -256,7 +263,7 @@ bool FRigVMStorage::Construct(int32 InRegisterIndex, int32 InElementIndex)
 		case ERigVMRegisterType::Name:
 		{
 			FName* DataPtr = (FName*)&Data[InElementIndex == INDEX_NONE ? Register.GetWorkByteIndex() : Register.GetWorkByteIndex() + InElementIndex * Register.ElementSize];
-			int32 Count = InElementIndex == INDEX_NONE ? Register.ElementCount : 1;
+			int32 Count = InElementIndex == INDEX_NONE ? Register.GetTotalElementCount() : 1;
 
 			FMemory::Memzero(DataPtr, Count * Register.ElementSize);
 			for (int32 Index = 0; Index < Count; Index++)
@@ -284,7 +291,7 @@ bool FRigVMStorage::Destroy(int32 InRegisterIndex, int32 InElementIndex)
 		case ERigVMRegisterType::Struct:
 		{
 			void* DataPtr = (void*)&Data[InElementIndex == INDEX_NONE ? Register.GetWorkByteIndex() : Register.GetWorkByteIndex() + InElementIndex * Register.ElementSize];
-			int32 Count = InElementIndex == INDEX_NONE ? Register.ElementCount : 1;
+			int32 Count = InElementIndex == INDEX_NONE ? Register.GetTotalElementCount() : 1;
 
 			UScriptStruct* ScriptStruct = GetScriptStruct(InRegisterIndex);
 			ScriptStruct->DestroyStruct(DataPtr, Count);
@@ -293,7 +300,7 @@ bool FRigVMStorage::Destroy(int32 InRegisterIndex, int32 InElementIndex)
 		case ERigVMRegisterType::String:
 		{
 			FString* DataPtr = (FString*)&Data[InElementIndex == INDEX_NONE ? Register.GetWorkByteIndex() : Register.GetWorkByteIndex() + InElementIndex * Register.ElementSize];
-			int32 Count = InElementIndex == INDEX_NONE ? Register.ElementCount : 1;
+			int32 Count = InElementIndex == INDEX_NONE ? Register.GetTotalElementCount() : 1;
 
 			for (int32 Index = 0; Index < Count; Index++)
 			{
@@ -304,7 +311,7 @@ bool FRigVMStorage::Destroy(int32 InRegisterIndex, int32 InElementIndex)
 		case ERigVMRegisterType::Name:
 		{
 			FName* DataPtr = (FName*)&Data[InElementIndex == INDEX_NONE ? Register.GetWorkByteIndex() : Register.GetWorkByteIndex() + InElementIndex * Register.ElementSize];
-			int32 Count = InElementIndex == INDEX_NONE ? Register.ElementCount : 1;
+			int32 Count = InElementIndex == INDEX_NONE ? Register.GetTotalElementCount() : 1;
 
 			for (int32 Index = 0; Index < Count; Index++)
 			{
@@ -382,31 +389,34 @@ FName FRigVMStorage::Rename(const FName& InOldName, const FName& InNewName)
 
 bool FRigVMStorage::Resize(int32 InRegisterIndex, int32 InNewElementCount, int32 InNewSliceCount)
 {
-	if (InNewElementCount <= 0)
+	if (InNewElementCount <= 0 || InNewSliceCount < 0)
 	{
 		return Remove(InRegisterIndex);
 	}
 
-	if (Registers[InRegisterIndex].ElementCount == InNewElementCount)
+	uint32 NewTotalCount = (uint32)InNewElementCount * (uint32)InNewSliceCount;
+	if (Registers[InRegisterIndex].GetTotalElementCount() == NewTotalCount)
 	{
 		return false;
 	}
 
 	FRigVMRegister& Register = Registers[InRegisterIndex];
+	Register.MoveToFirstSlice();
 
-	if (Register.ElementCount > InNewElementCount) // shrink
+	if (Register.GetTotalElementCount() > NewTotalCount) // shrink
 	{
-		int32 ElementsToRemove = Register.ElementCount - InNewElementCount;
+		int32 ElementsToRemove = Register.GetTotalElementCount() - NewTotalCount;
 		int32 NumBytesToRemove = Register.ElementSize * ElementsToRemove;
-		int32 FirstByteToRemove = Register.GetWorkByteIndex() + Register.ElementSize * InNewElementCount;
+		int32 FirstByteToRemove = Register.GetWorkByteIndex() + Register.ElementSize * NewTotalCount;
 
-		for (int32 ElementIndex = InNewElementCount; ElementIndex < Register.ElementCount; ElementIndex++)
+		for (uint32 ElementIndex = NewTotalCount; ElementIndex < Register.GetTotalElementCount(); ElementIndex++)
 		{
-			Destroy(InRegisterIndex, ElementIndex);
+			Destroy(InRegisterIndex, (uint32)ElementIndex);
 		}
 
 		Data.RemoveAt(FirstByteToRemove, NumBytesToRemove);
-		Register.ElementCount = InNewElementCount;
+		Register.ElementCount = (uint32)InNewElementCount;
+		Register.SliceCount = (uint32)InNewSliceCount;
 
 		for (int32 RegisterIndex = InRegisterIndex + 1; RegisterIndex < Registers.Num(); RegisterIndex++)
 		{
@@ -415,17 +425,18 @@ bool FRigVMStorage::Resize(int32 InRegisterIndex, int32 InNewElementCount, int32
 	}
 	else // grow
 	{
-		int32 OldElementCount = Register.ElementCount;
-		int32 ElementsToAdd = InNewElementCount - Register.ElementCount;
+		int32 OldElementCount = Register.GetTotalElementCount();
+		int32 ElementsToAdd = NewTotalCount - Register.GetTotalElementCount();
 		int32 NumBytesToAdd = Register.ElementSize * ElementsToAdd;
-		int32 FirstByteToAdd = Register.GetWorkByteIndex() + Register.ElementSize * Register.ElementCount;
+		int32 FirstByteToAdd = Register.GetWorkByteIndex() + Register.ElementSize * Register.GetTotalElementCount();
 
 		Data.InsertZeroed(FirstByteToAdd, NumBytesToAdd);
-		Register.ElementCount = InNewElementCount;
+		Register.ElementCount = (uint32)InNewElementCount;
+		Register.SliceCount = (uint32)InNewSliceCount;
 
-		for (int32 ElementIndex = OldElementCount; ElementIndex < InNewElementCount; ElementIndex++)
+		for (uint32 ElementIndex = OldElementCount; ElementIndex < NewTotalCount; ElementIndex++)
 		{
-			Construct(InRegisterIndex, ElementIndex);
+			Construct(InRegisterIndex, (int32)ElementIndex);
 		}
 
 		for (int32 RegisterIndex = InRegisterIndex + 1; RegisterIndex < Registers.Num(); RegisterIndex++)
diff --git a/Engine/Source/Runtime/RigVM/Public/RigVMCore/RigVMStorage.h b/Engine/Source/Runtime/RigVM/Public/RigVMCore/RigVMStorage.h
--- a/Engine/Source/Runtime/RigVM/Public/RigVMCore/RigVMStorage.h
+++ b/Engine/Source/Runtime/RigVM/Public/RigVMCore/RigVMStorage.h
@@ -75,42 +75,59 @@ struct RIGVM_API FRigVMRegister
 	}
 
 	UPROPERTY()
-	ERigVMRegisterType Type;
+		ERigVMRegisterType Type;
 
 	UPROPERTY()
-	uint32 ByteIndex;
+		uint32 ByteIndex;
 
 	UPROPERTY()
-	uint16 ElementSize;
+		uint16 ElementSize;
 
 	UPROPERTY()
-	uint16 ElementCount;
+		uint16 ElementCount;
 
 	UPROPERTY()
-	uint16 SliceIndex;
+		uint16 SliceIndex;
 
 	UPROPERTY()
-	uint16 SliceCount;
+		uint16 SliceCount;
 
 	UPROPERTY()
-	uint8 AlignmentBytes;
+		uint8 AlignmentBytes;
 
 	UPROPERTY()
-	FName Name;
+		FName Name;
 
 	UPROPERTY()
-	int32 ScriptStructIndex;
+		int32 ScriptStructIndex;
 
 	FORCEINLINE uint64 GetWorkByteIndex() const { return ByteIndex; }
 	FORCEINLINE uint64 GetStorageByteIndex() const { return ByteIndex - (uint64)AlignmentBytes - (uint64)(SliceIndex * GetNumBytesPerSlice()); }
 	FORCEINLINE uint8 GetAlignmentBytes() const { return AlignmentBytes; }
 	FORCEINLINE bool IsArray() const { return ElementCount > 1; }
-	FORCEINLINE bool IsPlain() const { return ScriptStructIndex == INDEX_NONE;  }
+	FORCEINLINE bool IsPlain() const { return ScriptStructIndex == INDEX_NONE; }
 	FORCEINLINE uint16 GetAllocatedBytes() const { return ElementCount * ElementSize * SliceCount + (uint16)AlignmentBytes; }
 	FORCEINLINE uint16 GetNumBytesPerSlice() const { return ElementCount * ElementSize; }
 	FORCEINLINE uint16 GetNumBytesAllSlices() const { return ElementCount * ElementSize * SliceCount; }
+	FORCEINLINE uint32 GetTotalElementCount() const { return (uint32)ElementCount * (uint32)SliceCount; }
 	FORCEINLINE void MoveToFirstSlice() { ByteIndex -= SliceIndex * GetNumBytesPerSlice(); SliceIndex = 0; }
-	FORCEINLINE void MoveToNextSlice() { ByteIndex += GetNumBytesPerSlice(); SliceIndex++; }
+	FORCEINLINE void MoveToNextSlice()
+	{
+		if(SliceCount == 1)
+		{
+			return;
+		}
+		
+		if(SliceIndex == SliceCount - 1)
+		{
+			MoveToFirstSlice();
+		}
+		else
+		{
+			ByteIndex += GetNumBytesPerSlice();
+			SliceIndex++;
+		}
+	}
 };
 
 typedef TArrayView<FRigVMRegister> FRigVMRegisterArray;
@@ -158,11 +175,15 @@ public:
 		return (const void*)&Data[Register.GetWorkByteIndex()];
 	}
 
-	FORCEINLINE void* GetData(int32 InRegisterIndex)
+	FORCEINLINE void* GetData(int32 InRegisterIndex, bool bMoveToNextSlice = false)
 	{
 		ensure(Registers.IsValidIndex(InRegisterIndex));
-		const FRigVMRegister& Register = Registers[InRegisterIndex];
+		FRigVMRegister& Register = Registers[InRegisterIndex];
 		ensure(Register.ElementCount > 0);
+		if(bMoveToNextSlice)
+		{
+			Register.MoveToNextSlice();
+		}
 		return (void*)&Data[Register.GetWorkByteIndex()];
 	}
 
@@ -178,7 +199,7 @@ public:
 	template<class T>
 	FORCEINLINE const T* Get(const FRigVMArgument& InArgument) const
 	{
-		return Get<T>(InArgument.Index());
+		return Get<T>(InArgument.GetRegisterIndex());
 	}
 
 	template<class T>
@@ -190,49 +211,57 @@ public:
 	template<class T>
 	FORCEINLINE const T& GetRef(const FRigVMArgument& InArgument) const
 	{
-		return GetRef<T>(InArgument.Index());
+		return GetRef<T>(InArgument.GetRegisterIndex());
 	}
 
 	template<class T>
-	FORCEINLINE T* Get(int32 InRegisterIndex)
+	FORCEINLINE T* Get(int32 InRegisterIndex, bool bMoveToNextSlice = false)
 	{
 		ensure(Registers.IsValidIndex(InRegisterIndex));
-		const FRigVMRegister& Register = Registers[InRegisterIndex];
+		FRigVMRegister& Register = Registers[InRegisterIndex];
 		ensure(Register.ElementCount > 0);
+		if(bMoveToNextSlice)
+		{
+			Register.MoveToNextSlice();
+		}
 		return (T*)&Data[Register.GetWorkByteIndex()];
 	}
 
 	template<class T>
-	FORCEINLINE T* Get(const FRigVMArgument& InArgument)
+	FORCEINLINE T* Get(const FRigVMArgument& InArgument, bool bMoveToNextSlice = false)
 	{
-		return Get<T>(InArgument.Index());
+		return Get<T>(InArgument.GetRegisterIndex(), bMoveToNextSlice);
 	}
 
 	template<class T>
-	FORCEINLINE T& GetRef(int32 InRegisterIndex)
+	FORCEINLINE T& GetRef(int32 InRegisterIndex, bool bMoveToNextSlice = false)
 	{
-		return *Get<T>(InRegisterIndex);
+		return *Get<T>(InRegisterIndex, bMoveToNextSlice);
 	}
 
 	template<class T>
-	FORCEINLINE T& GetRef(const FRigVMArgument& InArgument)
+	FORCEINLINE T& GetRef(const FRigVMArgument& InArgument, bool bMoveToNextSlice = false)
 	{
-		return GetRef<T>(InArgument.GetRegisterIndex());
+		return GetRef<T>(InArgument.GetRegisterIndex(), bMoveToNextSlice);
 	}
 
 	template<class T>
-	FORCEINLINE TArrayView<T> GetArray(int32 InRegisterIndex)
+	FORCEINLINE TArrayView<T> GetArray(int32 InRegisterIndex, bool bMoveToNextSlice = false)
 	{
 		ensure(Registers.IsValidIndex(InRegisterIndex));
-		const FRigVMRegister& Register = Registers[InRegisterIndex];
+		FRigVMRegister& Register = Registers[InRegisterIndex];
 		ensure(Register.ElementCount > 0);
+		if(bMoveToNextSlice)
+		{
+			Register.MoveToNextSlice();
+		}
 		return TArrayView<T>((T*)&Data[Register.GetWorkByteIndex()], Register.ElementCount);
 	}
 	
 	template<class T>
-	FORCEINLINE TArrayView<T> GetArray(const FRigVMArgument& InArgument)
+	FORCEINLINE TArrayView<T> GetArray(const FRigVMArgument& InArgument, bool bMoveToNextSlice = false)
 	{
-		return GetArray<T>(InArgument.GetRegisterIndex());
+		return GetArray<T>(InArgument.GetRegisterIndex(), bMoveToNextSlice);
 	}
 
 	FORCEINLINE UScriptStruct* GetScriptStruct(int32 InRegisterIndex) const
@@ -303,133 +332,145 @@ public:
 
 	void Reset();
 
-	FORCEINLINE int32 AddPlainArray(const FName& InNewName, int32 InElementSize, int32 InCount, const void* InDataPtr = nullptr)
+	FORCEINLINE int32 AddPlainArray(const FName& InNewName, int32 InElementSize, int32 InCount, const void* InDataPtr = nullptr, int32 InSliceCount = 1)
 	{
-		return Allocate(InNewName, InElementSize, InCount, InDataPtr);
+		return Allocate(InNewName, InElementSize, InCount, InSliceCount, InDataPtr);
 	}
 
 	template<class T>
-	FORCEINLINE int32 AddPlainArray(const FName& InNewName, int32 InCount, const T* InDataPtr = nullptr)
+	FORCEINLINE int32 AddPlainArray(const FName& InNewName, int32 InCount, const T* InDataPtr = nullptr, int32 InSliceCount = 1)
 	{
-		return Allocate(InNewName, sizeof(T), InCount, (const void*)InDataPtr);
+		return Allocate(InNewName, sizeof(T), InCount, InSliceCount, (const void*)InDataPtr);
 	}
 
 	template<class T>
-	FORCEINLINE int32 AddPlainArray(const FName& InNewName, const TArray<T>& InArray)
+	FORCEINLINE int32 AddPlainArray(const FName& InNewName, const TArray<T>& InArray, int32 InSliceCount = 1)
 	{
-		return AddPlainArray<T>(InNewName, InArray.Num(), InArray.GetData());
+		return AddPlainArray<T>(InNewName, InArray.Num(), InArray.GetData(), InSliceCount);
 	}
 
 	template<class T>
-	FORCEINLINE int32 AddPlainArray(const TArray<T>& InArray)
+	FORCEINLINE int32 AddPlainArray(const TArray<T>& InArray, int32 InSliceCount = 1)
 	{
-		return AddPlainArray<T>(NAME_None, InArray);
+		return AddPlainArray<T>(NAME_None, InArray, InSliceCount);
 	}
 
-	FORCEINLINE int32 AddPlain(const FName& InNewName, int32 InElementSize, const void* InValuePtr)
+	FORCEINLINE int32 AddPlain(const FName& InNewName, int32 InElementSize, const void* InValuePtr, int32 InSliceCount = 1)
 	{
-		return AddPlainArray(InNewName, InElementSize, 1, InValuePtr);
+		return AddPlainArray(InNewName, InElementSize, 1, InValuePtr, InSliceCount);
 	}
 
-	FORCEINLINE int32 AddPlain(int32 InElementSize, const void* InValuePtr)
+	FORCEINLINE int32 AddPlain(int32 InElementSize, const void* InValuePtr, int32 InSliceCount = 1)
 	{
-		return AddPlain(NAME_None, InElementSize, InValuePtr);
+		return AddPlain(NAME_None, InElementSize, InValuePtr, InSliceCount);
 	}
 
 	template<class T>
-	FORCEINLINE int32 AddPlain(const FName& InNewName, const T& InValue)
+	FORCEINLINE int32 AddPlain(const FName& InNewName, const T& InValue, int32 InSliceCount = 1)
 	{
-		return AddPlainArray<T>(InNewName, 1, &InValue);
+		return AddPlainArray<T>(InNewName, 1, &InValue, InSliceCount);
 	}
 
 	template<class T>
-	FORCEINLINE int32 AddPlain(const T& InValue)
+	FORCEINLINE int32 AddPlain(const T& InValue, int32 InSliceCount = 1)
 	{
-		return AddPlain<T>(NAME_None, InValue);
+		return AddPlain<T>(NAME_None, InValue, InSliceCount);
 	}
 
-	FORCEINLINE int32 AddNameArray(const FName& InNewName, int32 InCount, const FName* InDataPtr = nullptr)
+	FORCEINLINE int32 AddNameArray(const FName& InNewName, int32 InCount, const FName* InDataPtr = nullptr, int32 InSliceCount = 1)
 	{
-		int32 Register = Allocate(InNewName, sizeof(FName), InCount, nullptr);
+		int32 Register = Allocate(InNewName, sizeof(FName), InCount, InSliceCount, nullptr);
 		Registers[Register].Type = ERigVMRegisterType::Name;
 
 		Construct(Register);
 
 		if(InDataPtr)
 		{
-			FName* DataPtr = (FName*)GetData(Register);
-			for (int32 Index = 0; Index < InCount; Index++)
+			Registers[Register].MoveToFirstSlice();
+			for (uint16 SliceIndex = 0; SliceIndex < Registers[Register].SliceCount; SliceIndex++)
 			{
-				DataPtr[Index] = InDataPtr[Index];
+				FName* DataPtr = (FName*)GetData(Register);
+				for (int32 Index = 0; Index < InCount; Index++)
+				{
+					DataPtr[Index] = InDataPtr[Index];
+				}
+				Registers[Register].MoveToNextSlice();
 			}
+			Registers[Register].MoveToFirstSlice();
 		}
 
 		return Register;
 	}
 
-	FORCEINLINE int32 AddNameArray(const FName& InNewName, const TArray<FName>& InArray)
+	FORCEINLINE int32 AddNameArray(const FName& InNewName, const TArray<FName>& InArray, int32 InSliceCount = 1)
 	{
-		return AddNameArray(InNewName, InArray.Num(), InArray.GetData());
+		return AddNameArray(InNewName, InArray.Num(), InArray.GetData(), InSliceCount);
 	}
 
-	FORCEINLINE int32 AddNameArray(const TArray<FName>& InArray)
+	FORCEINLINE int32 AddNameArray(const TArray<FName>& InArray, int32 InSliceCount = 1)
 	{
-		return AddNameArray(NAME_None, InArray);
+		return AddNameArray(NAME_None, InArray, InSliceCount);
 	}
 
-	FORCEINLINE int32 AddName(const FName& InNewName, const FName& InValue)
+	FORCEINLINE int32 AddName(const FName& InNewName, const FName& InValue, int32 InSliceCount = 1)
 	{
-		return AddNameArray(InNewName, 1, &InValue);
+		return AddNameArray(InNewName, 1, &InValue, InSliceCount);
 	}
 
-	FORCEINLINE int32 AddName(const FName& InValue)
+	FORCEINLINE int32 AddName(const FName& InValue, int32 InSliceCount = 1)
 	{
-		return AddName(NAME_None, InValue);
+		return AddName(NAME_None, InValue, InSliceCount);
 	}
 
-	FORCEINLINE int32 AddStringArray(const FName& InNewName, int32 InCount, const FString* InDataPtr = nullptr)
+	FORCEINLINE int32 AddStringArray(const FName& InNewName, int32 InCount, const FString* InDataPtr = nullptr, int32 InSliceCount = 1)
 	{
-		int32 Register = Allocate(InNewName, sizeof(FString), InCount, nullptr);
+		int32 Register = Allocate(InNewName, sizeof(FString), InCount, InSliceCount, nullptr);
 		Registers[Register].Type = ERigVMRegisterType::String;
 
 		Construct(Register);
 
 		if(InDataPtr)
 		{
-			FString* DataPtr = (FString*)GetData(Register);
-			for (int32 Index = 0; Index < InCount; Index++)
+			Registers[Register].MoveToFirstSlice();
+			for (uint16 SliceIndex = 0; SliceIndex < Registers[Register].SliceCount; SliceIndex++)
 			{
-				DataPtr[Index] = InDataPtr[Index];
+				FString* DataPtr = (FString*)GetData(Register);
+				for (int32 Index = 0; Index < InCount; Index++)
+				{
+					DataPtr[Index] = InDataPtr[Index];
+				}
+				Registers[Register].MoveToNextSlice();
 			}
+			Registers[Register].MoveToFirstSlice();
 		}
 
 		return Register;
 	}
 
-	FORCEINLINE int32 AddStringArray(const FName& InNewName, const TArray<FString>& InArray)
+	FORCEINLINE int32 AddStringArray(const FName& InNewName, const TArray<FString>& InArray, int32 InSliceCount = 1)
 	{
-		return AddStringArray(InNewName, InArray.Num(), InArray.GetData());
+		return AddStringArray(InNewName, InArray.Num(), InArray.GetData(), InSliceCount);
 	}
 
-	FORCEINLINE int32 AddStringArray(const TArray<FString>& InArray)
+	FORCEINLINE int32 AddStringArray(const TArray<FString>& InArray, int32 InSliceCount = 1)
 	{
-		return AddStringArray(NAME_None, InArray);
+		return AddStringArray(NAME_None, InArray, InSliceCount);
 	}
 
-	FORCEINLINE int32 AddString(const FName& InNewName, const FString& InValue)
+	FORCEINLINE int32 AddString(const FName& InNewName, const FString& InValue, int32 InSliceCount = 1)
 	{
-		return AddStringArray(InNewName, 1, &InValue);
+		return AddStringArray(InNewName, 1, &InValue, InSliceCount);
 	}
 
-	FORCEINLINE int32 AddString(const FString& InValue)
+	FORCEINLINE int32 AddString(const FString& InValue, int32 InSliceCount = 1)
 	{
-		return AddString(NAME_None, InValue);
+		return AddString(NAME_None, InValue, InSliceCount);
 	}
 
 
-	FORCEINLINE int32 AddStructArray(const FName& InNewName, UScriptStruct* InScriptStruct, int32 InCount, const void* InDataPtr = nullptr)
+	FORCEINLINE int32 AddStructArray(const FName& InNewName, UScriptStruct* InScriptStruct, int32 InCount, const void* InDataPtr = nullptr, int32 InSliceCount = 1)
 	{
-		int32 Register = Allocate(InNewName, InScriptStruct->GetStructureSize(), InCount, nullptr, false);
+		int32 Register = Allocate(InNewName, InScriptStruct->GetStructureSize(), InCount, InSliceCount, nullptr, false);
 		if (Register == INDEX_NONE)
 		{
 			return INDEX_NONE;
@@ -446,19 +487,25 @@ public:
 		// copy values from the provided data
 		if (InDataPtr != nullptr)
 		{
-			InScriptStruct->CopyScriptStruct(GetData(Register), InDataPtr, InCount);
+			Registers[Register].MoveToFirstSlice();
+			for (uint16 SliceIndex = 0; SliceIndex < Registers[Register].SliceCount; SliceIndex++)
+			{
+				InScriptStruct->CopyScriptStruct(GetData(Register), InDataPtr, InCount);
+				Registers[Register].MoveToNextSlice();
+			}
+			Registers[Register].MoveToFirstSlice();
 		}
 
 		return Register;
 	}
 
-	FORCEINLINE int32 AddStructArray(UScriptStruct* InScriptStruct, int32 InCount, const void* InDataPtr = nullptr)
+	FORCEINLINE int32 AddStructArray(UScriptStruct* InScriptStruct, int32 InCount, const void* InDataPtr = nullptr, int32 InSliceCount = 1)
 	{
-		return AddStructArray(NAME_None, InScriptStruct, InCount, InDataPtr);
+		return AddStructArray(NAME_None, InScriptStruct, InCount, InDataPtr, InSliceCount);
 	}
 
 	template<class T>
-	FORCEINLINE int32 AddStructArray(const FName& InNewName, int32 InCount, const T* InDataPtr = nullptr)
+	FORCEINLINE int32 AddStructArray(const FName& InNewName, int32 InCount, const T* InDataPtr = nullptr, int32 InSliceCount = 1)
 	{
 		// if you are hitting this - you might need to use AddPlainArray instead!
 		UScriptStruct* Struct = T::StaticStruct();
@@ -467,41 +514,41 @@ public:
 			return INDEX_NONE;
 		}
 	
-		return AddStructArray(InNewName, Struct, InCount, InDataPtr);
+		return AddStructArray(InNewName, Struct, InCount, InDataPtr, InSliceCount);
 	}
 
 	template<class T>
-	FORCEINLINE int32 AddStructArray(const FName& InNewName, const TArray<T>& InArray)
+	FORCEINLINE int32 AddStructArray(const FName& InNewName, const TArray<T>& InArray, int32 InSliceCount = 1)
 	{
-		return AddStructArray<T>(InNewName, InArray.Num(), InArray.GetData());
+		return AddStructArray<T>(InNewName, InArray.Num(), InArray.GetData(), InSliceCount);
 	}
 
 	template<class T>
-	FORCEINLINE int32 AddStructArray(const TArray<T>& InArray)
+	FORCEINLINE int32 AddStructArray(const TArray<T>& InArray, int32 InSliceCount = 1)
 	{
-		return AddStructArray<T>(NAME_None, InArray);
+		return AddStructArray<T>(NAME_None, InArray, InSliceCount);
 	}
 
-	FORCEINLINE int32 AddStruct(const FName& InNewName, UScriptStruct* InScriptStruct, const void* InValuePtr)
+	FORCEINLINE int32 AddStruct(const FName& InNewName, UScriptStruct* InScriptStruct, const void* InValuePtr, int32 InSliceCount = 1)
 	{
-		return AddStructArray(InNewName, InScriptStruct, 1, InValuePtr);
+		return AddStructArray(InNewName, InScriptStruct, 1, InValuePtr, InSliceCount);
 	}
 
-	FORCEINLINE int32 AddStruct(UScriptStruct* InScriptStruct, const void* InValuePtr)
+	FORCEINLINE int32 AddStruct(UScriptStruct* InScriptStruct, const void* InValuePtr, int32 InSliceCount = 1)
 	{
-		return AddStruct(NAME_None, InScriptStruct, InValuePtr);
+		return AddStruct(NAME_None, InScriptStruct, InValuePtr, InSliceCount);
 	}
 
 	template<class T>
-	FORCEINLINE int32 AddStruct(const FName& InNewName, const T& InValue)
+	FORCEINLINE int32 AddStruct(const FName& InNewName, const T& InValue, int32 InSliceCount = 1)
 	{
-		return AddStructArray<T>(InNewName, 1, &InValue);
+		return AddStructArray<T>(InNewName, 1, &InValue, InSliceCount);
 	}
 
 	template<class T>
-	FORCEINLINE int32 AddStruct(const T& InValue)
+	FORCEINLINE int32 AddStruct(const T& InValue, int32 InSliceCount = 1)
 	{
-		return AddStruct<T>(NAME_None, InValue);
+		return AddStruct<T>(NAME_None, InValue, InSliceCount);
 	}
 
 	bool Remove(int32 InRegisterIndex);
@@ -515,8 +562,8 @@ public:
 
 private:
 
-	int32 Allocate(const FName& InNewName, int32 InElementSize, int32 InCount, const void* InDataPtr = nullptr, bool bUpdateRegisters = true);
-	int32 Allocate(int32 InElementSize, int32 InCount, const void* InDataPtr = nullptr, bool bUpdateRegisters = true);
+	int32 Allocate(const FName& InNewName, int32 InElementSize, int32 InElementCount, int32 InSliceCount, const void* InDataPtr = nullptr, bool bUpdateRegisters = true);
+	int32 Allocate(int32 InElementSize, int32 InElementCount, int32 InSliceCount, const void* InDataPtr = nullptr, bool bUpdateRegisters = true);
 	bool Construct(int32 InRegisterIndex, int32 InElementIndex = INDEX_NONE);
 	bool Destroy(int32 InRegisterIndex, int32 InElementIndex = INDEX_NONE);
 
