RigVM: Move FName + FString storage specialization to template specialization

diff --git a/Engine/Source/Runtime/RigVM/Public/RigVMCore/RigVMStorage.h b/Engine/Source/Runtime/RigVM/Public/RigVMCore/RigVMStorage.h
--- a/Engine/Source/Runtime/RigVM/Public/RigVMCore/RigVMStorage.h
+++ b/Engine/Source/Runtime/RigVM/Public/RigVMCore/RigVMStorage.h
@@ -343,48 +343,15 @@ public:
 		return Allocate(InNewName, sizeof(T), InCount, InSliceCount, (const void*)InDataPtr);
 	}
 
-	template<class T>
-	FORCEINLINE int32 AddPlainArray(const FName& InNewName, const TArray<T>& InArray, int32 InSliceCount = 1)
-	{
-		return AddPlainArray<T>(InNewName, InArray.Num(), InArray.GetData(), InSliceCount);
-	}
-
-	template<class T>
-	FORCEINLINE int32 AddPlainArray(const TArray<T>& InArray, int32 InSliceCount = 1)
-	{
-		return AddPlainArray<T>(NAME_None, InArray, InSliceCount);
-	}
-
-	FORCEINLINE int32 AddPlain(const FName& InNewName, int32 InElementSize, const void* InValuePtr, int32 InSliceCount = 1)
-	{
-		return AddPlainArray(InNewName, InElementSize, 1, InValuePtr, InSliceCount);
-	}
-
-	FORCEINLINE int32 AddPlain(int32 InElementSize, const void* InValuePtr, int32 InSliceCount = 1)
-	{
-		return AddPlain(NAME_None, InElementSize, InValuePtr, InSliceCount);
-	}
-
-	template<class T>
-	FORCEINLINE int32 AddPlain(const FName& InNewName, const T& InValue, int32 InSliceCount = 1)
-	{
-		return AddPlainArray<T>(InNewName, 1, &InValue, InSliceCount);
-	}
-
-	template<class T>
-	FORCEINLINE int32 AddPlain(const T& InValue, int32 InSliceCount = 1)
-	{
-		return AddPlain<T>(NAME_None, InValue, InSliceCount);
-	}
-
-	FORCEINLINE int32 AddNameArray(const FName& InNewName, int32 InCount, const FName* InDataPtr = nullptr, int32 InSliceCount = 1)
+	template<>
+	FORCEINLINE int32 AddPlainArray(const FName& InNewName, int32 InCount, const FName* InDataPtr, int32 InSliceCount)
 	{
 		int32 Register = Allocate(InNewName, sizeof(FName), InCount, InSliceCount, nullptr);
 		Registers[Register].Type = ERigVMRegisterType::Name;
 
 		Construct(Register);
 
-		if(InDataPtr)
+		if (InDataPtr)
 		{
 			Registers[Register].MoveToFirstSlice();
 			for (uint16 SliceIndex = 0; SliceIndex < Registers[Register].SliceCount; SliceIndex++)
@@ -402,34 +369,15 @@ public:
 		return Register;
 	}
 
-	FORCEINLINE int32 AddNameArray(const FName& InNewName, const TArray<FName>& InArray, int32 InSliceCount = 1)
-	{
-		return AddNameArray(InNewName, InArray.Num(), InArray.GetData(), InSliceCount);
-	}
-
-	FORCEINLINE int32 AddNameArray(const TArray<FName>& InArray, int32 InSliceCount = 1)
-	{
-		return AddNameArray(NAME_None, InArray, InSliceCount);
-	}
-
-	FORCEINLINE int32 AddName(const FName& InNewName, const FName& InValue, int32 InSliceCount = 1)
-	{
-		return AddNameArray(InNewName, 1, &InValue, InSliceCount);
-	}
-
-	FORCEINLINE int32 AddName(const FName& InValue, int32 InSliceCount = 1)
-	{
-		return AddName(NAME_None, InValue, InSliceCount);
-	}
-
-	FORCEINLINE int32 AddStringArray(const FName& InNewName, int32 InCount, const FString* InDataPtr = nullptr, int32 InSliceCount = 1)
+	template<>
+	FORCEINLINE int32 AddPlainArray(const FName& InNewName, int32 InCount, const FString* InDataPtr, int32 InSliceCount)
 	{
 		int32 Register = Allocate(InNewName, sizeof(FString), InCount, InSliceCount, nullptr);
 		Registers[Register].Type = ERigVMRegisterType::String;
 
 		Construct(Register);
 
-		if(InDataPtr)
+		if (InDataPtr)
 		{
 			Registers[Register].MoveToFirstSlice();
 			for (uint16 SliceIndex = 0; SliceIndex < Registers[Register].SliceCount; SliceIndex++)
@@ -447,26 +395,39 @@ public:
 		return Register;
 	}
 
-	FORCEINLINE int32 AddStringArray(const FName& InNewName, const TArray<FString>& InArray, int32 InSliceCount = 1)
+	template<class T>
+	FORCEINLINE int32 AddPlainArray(const FName& InNewName, const TArray<T>& InArray, int32 InSliceCount = 1)
 	{
-		return AddStringArray(InNewName, InArray.Num(), InArray.GetData(), InSliceCount);
+		return AddPlainArray<T>(InNewName, InArray.Num(), InArray.GetData(), InSliceCount);
 	}
 
-	FORCEINLINE int32 AddStringArray(const TArray<FString>& InArray, int32 InSliceCount = 1)
+	template<class T>
+	FORCEINLINE int32 AddPlainArray(const TArray<T>& InArray, int32 InSliceCount = 1)
 	{
-		return AddStringArray(NAME_None, InArray, InSliceCount);
+		return AddPlainArray<T>(NAME_None, InArray, InSliceCount);
+	}
+
+	FORCEINLINE int32 AddPlain(const FName& InNewName, int32 InElementSize, const void* InValuePtr, int32 InSliceCount = 1)
+	{
+		return AddPlainArray(InNewName, InElementSize, 1, InValuePtr, InSliceCount);
 	}
 
-	FORCEINLINE int32 AddString(const FName& InNewName, const FString& InValue, int32 InSliceCount = 1)
+	FORCEINLINE int32 AddPlain(int32 InElementSize, const void* InValuePtr, int32 InSliceCount = 1)
 	{
-		return AddStringArray(InNewName, 1, &InValue, InSliceCount);
+		return AddPlain(NAME_None, InElementSize, InValuePtr, InSliceCount);
 	}
 
-	FORCEINLINE int32 AddString(const FString& InValue, int32 InSliceCount = 1)
+	template<class T>
+	FORCEINLINE int32 AddPlain(const FName& InNewName, const T& InValue, int32 InSliceCount = 1)
 	{
-		return AddString(NAME_None, InValue, InSliceCount);
+		return AddPlainArray<T>(InNewName, 1, &InValue, InSliceCount);
 	}
 
+	template<class T>
+	FORCEINLINE int32 AddPlain(const T& InValue, int32 InSliceCount = 1)
+	{
+		return AddPlain<T>(NAME_None, InValue, InSliceCount);
+	}
 
 	FORCEINLINE int32 AddStructArray(const FName& InNewName, UScriptStruct* InScriptStruct, int32 InCount, const void* InDataPtr = nullptr, int32 InSliceCount = 1)
 	{
