RigVM: Support for type switching of plain registers

diff --git a/Engine/Source/Runtime/RigVM/Private/RigVMCore/RigVMStorage.cpp b/Engine/Source/Runtime/RigVM/Private/RigVMCore/RigVMStorage.cpp
--- a/Engine/Source/Runtime/RigVM/Private/RigVMCore/RigVMStorage.cpp
+++ b/Engine/Source/Runtime/RigVM/Private/RigVMCore/RigVMStorage.cpp
@@ -199,9 +199,9 @@ int32 FRigVMStorage::Allocate(const FName& InNewName, int32 InElementSize, int32
 	{
 		NewRegister.Name = Name;
 	}
-	NewRegister.ElementSize = InElementSize;
-	NewRegister.ElementCount = InElementCount;
-	NewRegister.SliceCount = InSliceCount;
+	NewRegister.ElementSize = (uint16)InElementSize;
+	NewRegister.ElementCount = (uint16)InElementCount;
+	NewRegister.SliceCount = (uint16)InSliceCount;
 	NewRegister.Type = ERigVMRegisterType::Plain;
 
 	Data.AddZeroed(NewRegister.GetAllocatedBytes());
@@ -285,7 +285,13 @@ bool FRigVMStorage::Destroy(int32 InRegisterIndex, int32 InElementIndex)
 {
 	ensure(Registers.IsValidIndex(InRegisterIndex));
 
-	const FRigVMRegister& Register = Registers[InRegisterIndex];
+	FRigVMRegister& Register = Registers[InRegisterIndex];
+
+	if (InElementIndex == INDEX_NONE)
+	{
+		Register.MoveToFirstSlice();
+	}
+
 	switch (Register.Type)
 	{
 		case ERigVMRegisterType::Struct:
@@ -304,7 +310,7 @@ bool FRigVMStorage::Destroy(int32 InRegisterIndex, int32 InElementIndex)
 
 			for (int32 Index = 0; Index < Count; Index++)
 			{
-				DataPtr[Index].Empty();
+				DataPtr[Index] = FString();
 			}
 			break;
 		}
@@ -389,6 +395,9 @@ FName FRigVMStorage::Rename(const FName& InOldName, const FName& InNewName)
 
 bool FRigVMStorage::Resize(int32 InRegisterIndex, int32 InNewElementCount, int32 InNewSliceCount)
 {
+	ensure(Registers.IsValidIndex(InRegisterIndex));
+	ensure(Registers[InRegisterIndex].TrailingBytes == 0);
+
 	if (InNewElementCount <= 0 || InNewSliceCount < 0)
 	{
 		return Remove(InRegisterIndex);
@@ -462,6 +471,49 @@ bool FRigVMStorage::Resize(const FName& InRegisterName, int32 InNewElementCount,
 	return Resize(RegisterIndex, InNewElementCount, InNewSliceCount);
 }
 
+bool FRigVMStorage::ChangeRegisterType(int32 InRegisterIndex, ERigVMRegisterType InNewType, int32 InElementSize, const void* InDataPtr, int32 InNewElementCount, int32 InNewSliceCount)
+{
+	ensure(Registers.IsValidIndex(InRegisterIndex));
+	
+	FRigVMRegister& Register = Registers[InRegisterIndex];
+	ensure(Register.ScriptStructIndex == INDEX_NONE);
+	ensure(Register.AlignmentBytes == 0);
+	ensure(InNewType == ERigVMRegisterType::Name || InNewType == ERigVMRegisterType::String || InNewType == ERigVMRegisterType::Plain);
+
+	Register.MoveToFirstSlice();
+
+	if (Register.Type == InNewType && Register.ElementSize == InElementSize && Register.ElementCount == InNewElementCount && Register.SliceCount == InNewSliceCount)
+	{
+		return false;
+	}
+
+	Destroy(InRegisterIndex);
+
+	uint16 OldAllocatedBytes = Register.GetAllocatedBytes();
+	uint16 NewAllocatedBytes = (uint16)InElementSize * (uint16)InNewElementCount * (uint16)InNewSliceCount;
+	ensure(NewAllocatedBytes <= OldAllocatedBytes);
+
+	Register.Type = InNewType;
+	Register.ElementSize = (uint16)InElementSize;
+	Register.ElementCount = (uint16)InNewElementCount;
+	Register.SliceCount = (uint16)InNewSliceCount;
+	Register.TrailingBytes = OldAllocatedBytes - NewAllocatedBytes;
+
+	Construct(InRegisterIndex);
+
+	if (InDataPtr != nullptr)
+	{
+		for (int32 SliceIndex = 0; SliceIndex < Register.SliceCount; SliceIndex++)
+		{
+			FMemory::Memcpy(GetData(InRegisterIndex), InDataPtr, Register.GetNumBytesPerSlice());
+			Register.MoveToNextSlice();
+		}
+		Register.MoveToFirstSlice();
+	}
+
+	return true;
+}
+
 void FRigVMStorage::UpdateRegisters()
 {
 	int32 AlignmentShift = 0;
diff --git a/Engine/Source/Runtime/RigVM/Public/RigVMCore/RigVMStorage.h b/Engine/Source/Runtime/RigVM/Public/RigVMCore/RigVMStorage.h
--- a/Engine/Source/Runtime/RigVM/Public/RigVMCore/RigVMStorage.h
+++ b/Engine/Source/Runtime/RigVM/Public/RigVMCore/RigVMStorage.h
@@ -69,51 +69,55 @@ struct RIGVM_API FRigVMRegister
 		, SliceIndex(0)
 		, SliceCount(1)
 		, AlignmentBytes(0)
+		, TrailingBytes(0)
 		, Name(NAME_None)
 		, ScriptStructIndex(INDEX_NONE)
 	{
 	}
 
 	UPROPERTY()
-		ERigVMRegisterType Type;
+	ERigVMRegisterType Type;
 
 	UPROPERTY()
-		uint32 ByteIndex;
+	uint32 ByteIndex;
 
 	UPROPERTY()
-		uint16 ElementSize;
+	uint16 ElementSize;
 
 	UPROPERTY()
-		uint16 ElementCount;
+	uint16 ElementCount;
 
 	UPROPERTY()
-		uint16 SliceIndex;
+	uint16 SliceIndex;
 
 	UPROPERTY()
-		uint16 SliceCount;
+	uint16 SliceCount;
 
 	UPROPERTY()
-		uint8 AlignmentBytes;
+	uint8 AlignmentBytes;
 
 	UPROPERTY()
-		FName Name;
+	uint16 TrailingBytes;
 
 	UPROPERTY()
-		int32 ScriptStructIndex;
+	FName Name;
+
+	UPROPERTY()
+	int32 ScriptStructIndex;
 
 	FORCEINLINE uint64 GetWorkByteIndex() const { return ByteIndex; }
 	FORCEINLINE uint64 GetStorageByteIndex() const { return ByteIndex - (uint64)AlignmentBytes - (uint64)(SliceIndex * GetNumBytesPerSlice()); }
 	FORCEINLINE uint8 GetAlignmentBytes() const { return AlignmentBytes; }
 	FORCEINLINE bool IsArray() const { return ElementCount > 1; }
 	FORCEINLINE bool IsPlain() const { return ScriptStructIndex == INDEX_NONE; }
-	FORCEINLINE uint16 GetAllocatedBytes() const { return ElementCount * ElementSize * SliceCount + (uint16)AlignmentBytes; }
+	FORCEINLINE uint16 GetAllocatedBytes() const { return ElementCount * ElementSize * SliceCount + (uint16)AlignmentBytes + TrailingBytes; }
 	FORCEINLINE uint16 GetNumBytesPerSlice() const { return ElementCount * ElementSize; }
 	FORCEINLINE uint16 GetNumBytesAllSlices() const { return ElementCount * ElementSize * SliceCount; }
 	FORCEINLINE uint32 GetTotalElementCount() const { return (uint32)ElementCount * (uint32)SliceCount; }
 	FORCEINLINE void MoveToFirstSlice() { ByteIndex -= SliceIndex * GetNumBytesPerSlice(); SliceIndex = 0; }
 	FORCEINLINE void MoveToNextSlice()
 	{
-		if(SliceCount == 1)
+		if(SliceCount <= 1)
 		{
 			return;
 		}
@@ -519,6 +523,66 @@ public:
 	bool Resize(int32 InRegisterIndex, int32 InNewElementCount, int32 InNewSliceCount = 1);
 	bool Resize(const FName& InRegisterName, int32 InNewElementCount, int32 InNewSliceCount = 1);
 
+	bool ChangeRegisterType(int32 InRegisterIndex, ERigVMRegisterType InNewType, int32 InElementSize, const void* InDataPtr = nullptr, int32 InNewElementCount = 1, int32 InNewSliceCount = 1);
+
+	template<class T>
+	FORCEINLINE bool ChangeRegisterType(int32 InRegisterIndex, const T* InDataPtr = nullptr, int32 InNewElementCount = 1, int32 InNewSliceCount = 1)
+	{
+		return ChangeRegisterType(InRegisterIndex, ERigVMRegisterType::Plain, sizeof(T), InDataPtr, InNewElementCount, InNewSliceCount);
+	}
+
+	template<>
+	FORCEINLINE bool ChangeRegisterType(int32 InRegisterIndex, const FName* InDataPtr, int32 InNewElementCount, int32 InNewSliceCount)
+	{
+		if (!ChangeRegisterType(InRegisterIndex, ERigVMRegisterType::Name, sizeof(FName), nullptr, InNewElementCount, InNewSliceCount))
+		{
+			return false;
+		}
+
+		if (InDataPtr)
+		{
+			Registers[InRegisterIndex].MoveToFirstSlice();
+			for (uint16 SliceIndex = 0; SliceIndex < Registers[InRegisterIndex].SliceCount; SliceIndex++)
+			{
+				FName* DataPtr = (FName*)GetData(InRegisterIndex);
+				for (int32 Index = 0; Index < InNewElementCount; Index++)
+				{
+					DataPtr[Index] = InDataPtr[Index];
+				}
+				Registers[InRegisterIndex].MoveToNextSlice();
+			}
+			Registers[InRegisterIndex].MoveToFirstSlice();
+		}
+
+		return true;
+	}
+
+	template<>
+	FORCEINLINE bool ChangeRegisterType(int32 InRegisterIndex, const FString* InDataPtr, int32 InNewElementCount, int32 InNewSliceCount)
+	{
+		if (!ChangeRegisterType(InRegisterIndex, ERigVMRegisterType::String, sizeof(FString), nullptr, InNewElementCount, InNewSliceCount))
+		{
+			return false;
+		}
+
+		if (InDataPtr)
+		{
+			Registers[InRegisterIndex].MoveToFirstSlice();
+			for (uint16 SliceIndex = 0; SliceIndex < Registers[InRegisterIndex].SliceCount; SliceIndex++)
+			{
+				FString* DataPtr = (FString*)GetData(InRegisterIndex);
+				for (int32 Index = 0; Index < InNewElementCount; Index++)
+				{
+					DataPtr[Index] = InDataPtr[Index];
+				}
+				Registers[InRegisterIndex].MoveToNextSlice();
+			}
+			Registers[InRegisterIndex].MoveToFirstSlice();
+		}
+
+		return true;
+	}
+
 	void UpdateRegisters();
 
 private:
