RigVM: Move to Execute_N_Args opcodes instead of additional storage for number of arguments

diff --git a/Engine/Source/Programs/UnrealHeaderTool/Private/CodeGenerator.cpp b/Engine/Source/Programs/UnrealHeaderTool/Private/CodeGenerator.cpp
--- a/Engine/Source/Programs/UnrealHeaderTool/Private/CodeGenerator.cpp
+++ b/Engine/Source/Programs/UnrealHeaderTool/Private/CodeGenerator.cpp
@@ -3403,6 +3403,12 @@ void FNativeClassHeaderGenerator::ExportGeneratedStructBodyMacros(FOutputDevice&
 			MemberMulticastProlog.Add(FString::Printf(TEXT("%s %s = RigVMStorage[%s.GetStorageIndex()]->%s<%s>(%s.GetRegisterIndex(), true);"), *ExtractedMulticastCPPType, *Name, *AddressName, *MulticastGetter, *MultiCastCPPType, *AddressName));
 		}
 
+		// we have a hard limit on the number of parameters for our function
+		if (MemberNames.Num() > 64)
+		{
+			UE_LOG_ERROR_UHT(TEXT("RIGVM_METHOD: Struct '%s' has more than 64 members (64 is the limit)."), StructNameCPP);
+		}
+
 		OutGeneratedHeaderText.Log(TEXT("\n"));
 		OutGeneratedHeaderText.Logf(TEXT("#define %s_IMPLEMENT_RIGVM(ReturnType, FunctionName, ...) \\\n"), *RigVMMacroPrefix);
 		OutGeneratedHeaderText.Logf(TEXT("\tReturnType %s::Static##FunctionName(\\\r\n\t\t%s, ##__VA_ARGS__)\n"), StructNameCPP, *FString::Join(MemberDeclarations, TEXT(",\\\r\n\t\t")));
diff --git a/Engine/Source/Runtime/RigVM/Private/RigVMCore/RigVM.cpp b/Engine/Source/Runtime/RigVM/Private/RigVMCore/RigVM.cpp
--- a/Engine/Source/Runtime/RigVM/Private/RigVMCore/RigVM.cpp
+++ b/Engine/Source/Runtime/RigVM/Private/RigVMCore/RigVM.cpp
@@ -86,6 +86,78 @@ bool URigVM::Execute(FRigVMStoragePtrArray Storage, TArrayView<void*> Additional
 	{
 		switch (Instructions[InstructionIndex].OpCode)
 		{
+			case ERigVMOpCode::Execute_0_Args:
+			case ERigVMOpCode::Execute_1_Args:
+			case ERigVMOpCode::Execute_2_Args:
+			case ERigVMOpCode::Execute_3_Args:
+			case ERigVMOpCode::Execute_4_Args:
+			case ERigVMOpCode::Execute_5_Args:
+			case ERigVMOpCode::Execute_6_Args:
+			case ERigVMOpCode::Execute_7_Args:
+			case ERigVMOpCode::Execute_8_Args:
+			case ERigVMOpCode::Execute_9_Args:
+			case ERigVMOpCode::Execute_10_Args:
+			case ERigVMOpCode::Execute_11_Args:
+			case ERigVMOpCode::Execute_12_Args:
+			case ERigVMOpCode::Execute_13_Args:
+			case ERigVMOpCode::Execute_14_Args:
+			case ERigVMOpCode::Execute_15_Args:
+			case ERigVMOpCode::Execute_16_Args:
+			case ERigVMOpCode::Execute_17_Args:
+			case ERigVMOpCode::Execute_18_Args:
+			case ERigVMOpCode::Execute_19_Args:
+			case ERigVMOpCode::Execute_20_Args:
+			case ERigVMOpCode::Execute_21_Args:
+			case ERigVMOpCode::Execute_22_Args:
+			case ERigVMOpCode::Execute_23_Args:
+			case ERigVMOpCode::Execute_24_Args:
+			case ERigVMOpCode::Execute_25_Args:
+			case ERigVMOpCode::Execute_26_Args:
+			case ERigVMOpCode::Execute_27_Args:
+			case ERigVMOpCode::Execute_28_Args:
+			case ERigVMOpCode::Execute_29_Args:
+			case ERigVMOpCode::Execute_30_Args:
+			case ERigVMOpCode::Execute_31_Args:
+			case ERigVMOpCode::Execute_32_Args:
+			case ERigVMOpCode::Execute_33_Args:
+			case ERigVMOpCode::Execute_34_Args:
+			case ERigVMOpCode::Execute_35_Args:
+			case ERigVMOpCode::Execute_36_Args:
+			case ERigVMOpCode::Execute_37_Args:
+			case ERigVMOpCode::Execute_38_Args:
+			case ERigVMOpCode::Execute_39_Args:
+			case ERigVMOpCode::Execute_40_Args:
+			case ERigVMOpCode::Execute_41_Args:
+			case ERigVMOpCode::Execute_42_Args:
+			case ERigVMOpCode::Execute_43_Args:
+			case ERigVMOpCode::Execute_44_Args:
+			case ERigVMOpCode::Execute_45_Args:
+			case ERigVMOpCode::Execute_46_Args:
+			case ERigVMOpCode::Execute_47_Args:
+			case ERigVMOpCode::Execute_48_Args:
+			case ERigVMOpCode::Execute_49_Args:
+			case ERigVMOpCode::Execute_50_Args:
+			case ERigVMOpCode::Execute_51_Args:
+			case ERigVMOpCode::Execute_52_Args:
+			case ERigVMOpCode::Execute_53_Args:
+			case ERigVMOpCode::Execute_54_Args:
+			case ERigVMOpCode::Execute_55_Args:
+			case ERigVMOpCode::Execute_56_Args:
+			case ERigVMOpCode::Execute_57_Args:
+			case ERigVMOpCode::Execute_58_Args:
+			case ERigVMOpCode::Execute_59_Args:
+			case ERigVMOpCode::Execute_60_Args:
+			case ERigVMOpCode::Execute_61_Args:
+			case ERigVMOpCode::Execute_62_Args:
+			case ERigVMOpCode::Execute_63_Args:
+			case ERigVMOpCode::Execute_64_Args:
+			{
+				const FRigVMExecuteOp& Op = ByteCode.GetOpAt<FRigVMExecuteOp>(Instructions[InstructionIndex]);
+				FRigVMArgumentArray Args = ByteCode.GetArgumentsForExecuteOp(Instructions[InstructionIndex]);
+				(*Functions[Op.FunctionIndex])(Args, Storage, AdditionalArgs);
+				InstructionIndex++;
+				break;
+			}
 			case ERigVMOpCode::Copy:
 			{
 				const FRigVMCopyOp& Op = ByteCode.GetOpAt<FRigVMCopyOp>(Instructions[InstructionIndex]);
@@ -137,14 +209,6 @@ bool URigVM::Execute(FRigVMStoragePtrArray Storage, TArrayView<void*> Additional
 				ensureMsgf(false, TEXT("TO BE IMPLEMENTED"));
 				break;
 			}
-			case ERigVMOpCode::Execute:
-			{
-				const FRigVMExecuteOp& Op = ByteCode.GetOpAt<FRigVMExecuteOp>(Instructions[InstructionIndex]);
-				FRigVMArgumentArray Args = ByteCode.GetArgumentsForExecuteOp(Instructions[InstructionIndex]);
-				(*Functions[Op.FunctionIndex])(Args, Storage, AdditionalArgs);
-				InstructionIndex++;
-				break;
-			}
 			case ERigVMOpCode::Exit:
 			case ERigVMOpCode::Invalid:
 			{
diff --git a/Engine/Source/Runtime/RigVM/Private/RigVMCore/RigVMByteCode.cpp b/Engine/Source/Runtime/RigVM/Private/RigVMCore/RigVMByteCode.cpp
--- a/Engine/Source/Runtime/RigVM/Private/RigVMCore/RigVMByteCode.cpp
+++ b/Engine/Source/Runtime/RigVM/Private/RigVMCore/RigVMByteCode.cpp
@@ -41,6 +41,80 @@ uint64 FRigVMByteCode::GetOpNumBytesAt(uint64 InByteCodeIndex, bool bIncludeArgu
 	ERigVMOpCode OpCode = GetOpCodeAt(InByteCodeIndex);
 	switch (OpCode)
 	{
+		case ERigVMOpCode::Execute_0_Args:
+		case ERigVMOpCode::Execute_1_Args:
+		case ERigVMOpCode::Execute_2_Args:
+		case ERigVMOpCode::Execute_3_Args:
+		case ERigVMOpCode::Execute_4_Args:
+		case ERigVMOpCode::Execute_5_Args:
+		case ERigVMOpCode::Execute_6_Args:
+		case ERigVMOpCode::Execute_7_Args:
+		case ERigVMOpCode::Execute_8_Args:
+		case ERigVMOpCode::Execute_9_Args:
+		case ERigVMOpCode::Execute_10_Args:
+		case ERigVMOpCode::Execute_11_Args:
+		case ERigVMOpCode::Execute_12_Args:
+		case ERigVMOpCode::Execute_13_Args:
+		case ERigVMOpCode::Execute_14_Args:
+		case ERigVMOpCode::Execute_15_Args:
+		case ERigVMOpCode::Execute_16_Args:
+		case ERigVMOpCode::Execute_17_Args:
+		case ERigVMOpCode::Execute_18_Args:
+		case ERigVMOpCode::Execute_19_Args:
+		case ERigVMOpCode::Execute_20_Args:
+		case ERigVMOpCode::Execute_21_Args:
+		case ERigVMOpCode::Execute_22_Args:
+		case ERigVMOpCode::Execute_23_Args:
+		case ERigVMOpCode::Execute_24_Args:
+		case ERigVMOpCode::Execute_25_Args:
+		case ERigVMOpCode::Execute_26_Args:
+		case ERigVMOpCode::Execute_27_Args:
+		case ERigVMOpCode::Execute_28_Args:
+		case ERigVMOpCode::Execute_29_Args:
+		case ERigVMOpCode::Execute_30_Args:
+		case ERigVMOpCode::Execute_31_Args:
+		case ERigVMOpCode::Execute_32_Args:
+		case ERigVMOpCode::Execute_33_Args:
+		case ERigVMOpCode::Execute_34_Args:
+		case ERigVMOpCode::Execute_35_Args:
+		case ERigVMOpCode::Execute_36_Args:
+		case ERigVMOpCode::Execute_37_Args:
+		case ERigVMOpCode::Execute_38_Args:
+		case ERigVMOpCode::Execute_39_Args:
+		case ERigVMOpCode::Execute_40_Args:
+		case ERigVMOpCode::Execute_41_Args:
+		case ERigVMOpCode::Execute_42_Args:
+		case ERigVMOpCode::Execute_43_Args:
+		case ERigVMOpCode::Execute_44_Args:
+		case ERigVMOpCode::Execute_45_Args:
+		case ERigVMOpCode::Execute_46_Args:
+		case ERigVMOpCode::Execute_47_Args:
+		case ERigVMOpCode::Execute_48_Args:
+		case ERigVMOpCode::Execute_49_Args:
+		case ERigVMOpCode::Execute_50_Args:
+		case ERigVMOpCode::Execute_51_Args:
+		case ERigVMOpCode::Execute_52_Args:
+		case ERigVMOpCode::Execute_53_Args:
+		case ERigVMOpCode::Execute_54_Args:
+		case ERigVMOpCode::Execute_55_Args:
+		case ERigVMOpCode::Execute_56_Args:
+		case ERigVMOpCode::Execute_57_Args:
+		case ERigVMOpCode::Execute_58_Args:
+		case ERigVMOpCode::Execute_59_Args:
+		case ERigVMOpCode::Execute_60_Args:
+		case ERigVMOpCode::Execute_61_Args:
+		case ERigVMOpCode::Execute_62_Args:
+		case ERigVMOpCode::Execute_63_Args:
+		case ERigVMOpCode::Execute_64_Args:
+		{
+			uint64 NumBytes = (uint64)sizeof(FRigVMExecuteOp);
+			if(bIncludeArguments)
+			{
+				const FRigVMExecuteOp& ExecuteOp = GetOpAt<FRigVMExecuteOp>(InByteCodeIndex);
+				NumBytes += (uint64)ExecuteOp.GetArgumentCount() * (uint64)sizeof(FRigVMArgument);
+			}
+			return NumBytes;
+		}
 		case ERigVMOpCode::Copy:
 		{
 			return (uint64)sizeof(FRigVMCopyOp);
@@ -73,16 +147,6 @@ uint64 FRigVMByteCode::GetOpNumBytesAt(uint64 InByteCodeIndex, bool bIncludeArgu
 		{
 			return (uint64)sizeof(FRigVMJumpIfFalseOp);
 		}
-		case ERigVMOpCode::Execute:
-		{
-			uint64 NumBytes = (uint64)sizeof(FRigVMExecuteOp);
-			if(bIncludeArguments)
-			{
-				const FRigVMExecuteOp& ExecuteOp = GetOpAt<FRigVMExecuteOp>(InByteCodeIndex);
-				NumBytes += (uint64)ExecuteOp.ArgumentCount * (uint64)sizeof(FRigVMArgument);
-			}
-			return NumBytes;
-		}
 		case ERigVMOpCode::Exit:
 		{
 			return (uint64)sizeof(FRigVMExitOp);
@@ -146,7 +210,7 @@ uint64 FRigVMByteCode::AddJumpIfFalseOp(uint64 InByteCodeIndex, const FRigVMArgu
 
 uint64 FRigVMByteCode::AddExecuteOp(uint16 InFunctionIndex, const TArrayView<FRigVMArgument>& InArguments)
 {
-	FRigVMExecuteOp Op(InFunctionIndex, (uint16)InArguments.Num());
+	FRigVMExecuteOp Op(InFunctionIndex, (uint8)InArguments.Num());
 	uint64 OpByteIndex = AddOp(Op);
 	uint64 ArgumentsByteIndex = (uint64)ByteCode.AddUninitialized(sizeof(FRigVMArgument) * InArguments.Num());
 	FMemory::Memcpy(ByteCode.GetData() + ArgumentsByteIndex, InArguments.GetData(), sizeof(FRigVMArgument) * InArguments.Num());
diff --git a/Engine/Source/Runtime/RigVM/Public/RigVMCore/RigVMByteCode.h b/Engine/Source/Runtime/RigVM/Public/RigVMCore/RigVMByteCode.h
--- a/Engine/Source/Runtime/RigVM/Public/RigVMCore/RigVMByteCode.h
+++ b/Engine/Source/Runtime/RigVM/Public/RigVMCore/RigVMByteCode.h
@@ -11,6 +11,71 @@ struct FRigVMByteCode;
 UENUM()
 enum class ERigVMOpCode : uint8
 {
+	Execute_0_Args,
+	Execute_1_Args,
+	Execute_2_Args,
+	Execute_3_Args,
+	Execute_4_Args,
+	Execute_5_Args,
+	Execute_6_Args,
+	Execute_7_Args,
+	Execute_8_Args,
+	Execute_9_Args,
+	Execute_10_Args,
+	Execute_11_Args,
+	Execute_12_Args,
+	Execute_13_Args,
+	Execute_14_Args,
+	Execute_15_Args,
+	Execute_16_Args,
+	Execute_17_Args,
+	Execute_18_Args,
+	Execute_19_Args,
+	Execute_20_Args,
+	Execute_21_Args,
+	Execute_22_Args,
+	Execute_23_Args,
+	Execute_24_Args,
+	Execute_25_Args,
+	Execute_26_Args,
+	Execute_27_Args,
+	Execute_28_Args,
+	Execute_29_Args,
+	Execute_30_Args,
+	Execute_31_Args,
+	Execute_32_Args,
+	Execute_33_Args,
+	Execute_34_Args,
+	Execute_35_Args,
+	Execute_36_Args,
+	Execute_37_Args,
+	Execute_38_Args,
+	Execute_39_Args,
+	Execute_40_Args,
+	Execute_41_Args,
+	Execute_42_Args,
+	Execute_43_Args,
+	Execute_44_Args,
+	Execute_45_Args,
+	Execute_46_Args,
+	Execute_47_Args,
+	Execute_48_Args,
+	Execute_49_Args,
+	Execute_50_Args,
+	Execute_51_Args,
+	Execute_52_Args,
+	Execute_53_Args,
+	Execute_54_Args,
+	Execute_55_Args,
+	Execute_56_Args,
+	Execute_57_Args,
+	Execute_58_Args,
+	Execute_59_Args,
+	Execute_60_Args,
+	Execute_61_Args,
+	Execute_62_Args,
+	Execute_63_Args,
+	Execute_64_Args,
 	Copy,
 	Increment,
 	Decrement,
@@ -19,7 +84,6 @@ enum class ERigVMOpCode : uint8
 	Jump,
 	JumpIfTrue,
 	JumpIfFalse,
-	Execute,
 	Exit,
 	Invalid
 };
@@ -34,6 +98,25 @@ struct RIGVM_API FRigVMBaseOp
 	ERigVMOpCode OpCode;
 };
 
+struct RIGVM_API FRigVMExecuteOp : public FRigVMBaseOp
+{
+	FRigVMExecuteOp()
+	: FRigVMBaseOp()
+	, FunctionIndex(INDEX_NONE)
+	{
+	}
+
+	FRigVMExecuteOp(uint16 InFunctionIndex, uint8 InArgumentCount)
+	: FRigVMBaseOp((ERigVMOpCode)(uint8(ERigVMOpCode::Execute_0_Args) + InArgumentCount))
+	, FunctionIndex(InFunctionIndex)
+	{
+	}
+
+	uint16 FunctionIndex;
+
+	FORCEINLINE uint8 GetArgumentCount() const { return uint8(OpCode) - uint8(ERigVMOpCode::Execute_0_Args); }
+};
+
 struct RIGVM_API FRigVMCopyOp : public FRigVMBaseOp
 {
 	FRigVMCopyOp()
@@ -214,26 +297,6 @@ struct RIGVM_API FRigVMJumpIfFalseOp : public FRigVMBaseOp
 	FRigVMArgument Condition;
 };
 
-struct RIGVM_API FRigVMExecuteOp : public FRigVMBaseOp
-{
-	FRigVMExecuteOp()
-	: FRigVMBaseOp(ERigVMOpCode::Execute)
-	, FunctionIndex(INDEX_NONE)
-	, ArgumentCount(0)
-	{
-	}
-
-	FRigVMExecuteOp(uint16 InFunctionIndex, uint16 InArgumentCount)
-	: FRigVMBaseOp(ERigVMOpCode::Execute)
-	, FunctionIndex(InFunctionIndex)
-	, ArgumentCount(InArgumentCount)
-	{
-	}
-
-	uint16 FunctionIndex;
-	uint16 ArgumentCount;
-};
-
 struct RIGVM_API FRigVMExitOp : public FRigVMBaseOp
 {
 	FRigVMExitOp()
@@ -331,7 +394,6 @@ public:
 	template<class OpType>
 	FORCEINLINE const OpType& GetOpAt(const FRigVMByteCodeTableEntry& InEntry) const
 	{
-		ensure(OpType().OpCode == InEntry.OpCode);
 		return GetOpAt<OpType>(InEntry.ByteCodeIndex);
 	}
 
@@ -344,13 +406,13 @@ public:
 	FORCEINLINE TArrayView<FRigVMArgument> GetArgumentsForExecuteOp(uint64 InByteCodeIndex) const
 	{
 		const FRigVMExecuteOp& ExecuteOp = GetOpAt<FRigVMExecuteOp>(InByteCodeIndex);
-		return GetArgumentsAt(InByteCodeIndex + sizeof(FRigVMExecuteOp), ExecuteOp.ArgumentCount);
+		return GetArgumentsAt(InByteCodeIndex + sizeof(FRigVMExecuteOp), ExecuteOp.GetArgumentCount());
 	}
 
 	FORCEINLINE TArrayView<FRigVMArgument> GetArgumentsForExecuteOp(const FRigVMByteCodeTableEntry& InEntry) const
 	{
 		const FRigVMExecuteOp& ExecuteOp = GetOpAt<FRigVMExecuteOp>(InEntry);
-		return GetArgumentsAt(InEntry.ByteCodeIndex + sizeof(FRigVMExecuteOp), ExecuteOp.ArgumentCount);
+		return GetArgumentsAt(InEntry.ByteCodeIndex + sizeof(FRigVMExecuteOp), ExecuteOp.GetArgumentCount());
 	}
 
 	FORCEINLINE const TArrayView<uint8> GetByteCode() const
