RigVM: Support for jumps + compare operators

diff --git a/Engine/Source/Runtime/RigVM/Private/RigVMCore/RigVM.cpp b/Engine/Source/Runtime/RigVM/Private/RigVMCore/RigVM.cpp
--- a/Engine/Source/Runtime/RigVM/Private/RigVMCore/RigVM.cpp
+++ b/Engine/Source/Runtime/RigVM/Private/RigVMCore/RigVM.cpp
@@ -81,7 +81,7 @@ bool URigVM::Execute(FRigVMStoragePtrArray Storage, TArrayView<void*> Additional
 		Storage = FRigVMStoragePtrArray(LocalStorage, 2);
 	}
 
-	int32 InstructionIndex = 0;
+	uint16 InstructionIndex = 0;
 	while (Instructions[InstructionIndex].OpCode != ERigVMOpCode::Exit)
 	{
 		switch (Instructions[InstructionIndex].OpCode)
@@ -158,55 +158,179 @@ bool URigVM::Execute(FRigVMStoragePtrArray Storage, TArrayView<void*> Additional
 				InstructionIndex++;
 				break;
 			}
+			case ERigVMOpCode::Zero:
+			{
+				const FRigVMZeroOp& Op = ByteCode.GetOpAt<FRigVMZeroOp>(Instructions[InstructionIndex]);
+				Storage[Op.Arg.GetStorageIndex()]->GetRef<int32>(Op.Arg.GetRegisterIndex()) = 0;
+				InstructionIndex++;
+				break;
+			}
+			case ERigVMOpCode::BoolFalse:
+			{
+				const FRigVMFalseOp& Op = ByteCode.GetOpAt<FRigVMFalseOp>(Instructions[InstructionIndex]);
+				Storage[Op.Arg.GetStorageIndex()]->GetRef<bool>(Op.Arg.GetRegisterIndex()) = false;
+				InstructionIndex++;
+				break;
+			}
+			case ERigVMOpCode::BoolTrue:
+			{
+				const FRigVMTrueOp& Op = ByteCode.GetOpAt<FRigVMTrueOp>(Instructions[InstructionIndex]);
+				Storage[Op.Arg.GetStorageIndex()]->GetRef<bool>(Op.Arg.GetRegisterIndex()) = false;
+				InstructionIndex++;
+				break;
+			}
 			case ERigVMOpCode::Copy:
 			{
 				const FRigVMCopyOp& Op = ByteCode.GetOpAt<FRigVMCopyOp>(Instructions[InstructionIndex]);
-				Storage[0]->Copy(Op.Source.GetRegisterIndex(), Op.Target.GetRegisterIndex(), Storage[Op.Source.GetStorageIndex()], Op.SourceOffset, Op.TargetOffset, Op.NumBytes);
+				Storage[Op.Target.GetStorageIndex()]->Copy(Op.Source.GetRegisterIndex(), Op.Target.GetRegisterIndex(), Storage[Op.Source.GetStorageIndex()], Op.SourceOffset, Op.TargetOffset, Op.NumBytes);
 				InstructionIndex++;
 				break;
 			}
 			case ERigVMOpCode::Increment:
 			{
 				const FRigVMIncrementOp& Op = ByteCode.GetOpAt<FRigVMIncrementOp>(Instructions[InstructionIndex]);
-				Storage[0]->GetRef<int32>(Op.Arg.GetRegisterIndex())++;
+				Storage[Op.Arg.GetStorageIndex()]->GetRef<int32>(Op.Arg.GetRegisterIndex())++;
 				InstructionIndex++;
 				break;
 			}
 			case ERigVMOpCode::Decrement:
 			{
 				const FRigVMDecrementOp& Op = ByteCode.GetOpAt<FRigVMDecrementOp>(Instructions[InstructionIndex]);
-				Storage[0]->GetRef<int32>(Op.Arg.GetRegisterIndex())--;
+				Storage[Op.Arg.GetStorageIndex()]->GetRef<int32>(Op.Arg.GetRegisterIndex())--;
 				InstructionIndex++;
 				break;
 			}
 			case ERigVMOpCode::Equals:
+			case ERigVMOpCode::NotEquals:
 			{
 				const FRigVMEqualsOp& Op = ByteCode.GetOpAt<FRigVMEqualsOp>(Instructions[InstructionIndex]);
-				ensureMsgf(false, TEXT("TO BE IMPLEMENTED"));
+				const FRigVMRegister& RegisterA = (*Storage[Op.A.GetStorageIndex()])[Op.A.GetRegisterIndex()];
+				const FRigVMRegister& RegisterB = (*Storage[Op.B.GetStorageIndex()])[Op.B.GetRegisterIndex()];
+				uint16 BytesA = RegisterA.GetNumBytesPerSlice();
+				uint16 BytesB = RegisterB.GetNumBytesPerSlice();
+				
+				bool Result = false;
+				if (BytesA == BytesB && RegisterA.Type == RegisterB.Type && RegisterA.ScriptStructIndex == RegisterB.ScriptStructIndex)
+				{
+					switch (RegisterA.Type)
+					{
+						case ERigVMRegisterType::Plain:
+						case ERigVMRegisterType::Name:
+						{
+							void * DataA = Storage[Op.A.GetStorageIndex()]->GetData(Op.A.GetRegisterIndex());
+							void * DataB = Storage[Op.B.GetStorageIndex()]->GetData(Op.B.GetRegisterIndex());
+							Result = FMemory::Memcmp(DataA, DataB, BytesA) == 0;
+							break;
+						}
+						case ERigVMRegisterType::String:
+						{
+							TArrayView<FString> StringsA = Storage[Op.A.GetStorageIndex()]->GetArray<FString>(Op.A.GetRegisterIndex());
+							TArrayView<FString> StringsB = Storage[Op.B.GetStorageIndex()]->GetArray<FString>(Op.B.GetRegisterIndex());
+
+							Result = true;
+							for (int32 StringIndex = 0; StringIndex < StringsA.Num(); StringIndex++)
+							{
+								if (StringsA[StringIndex] != StringsB[StringIndex])
+								{
+									Result = false;
+									break;
+								}
+							}
+							break;
+						}
+						case ERigVMRegisterType::Struct:
+						{
+							UScriptStruct* ScriptStruct = Storage[Op.A.GetStorageIndex()]->GetScriptStruct(RegisterA.ScriptStructIndex);
+
+							uint8* DataA = (uint8*)Storage[Op.A.GetStorageIndex()]->GetData(Op.A.GetRegisterIndex());
+							uint8* DataB = (uint8*)Storage[Op.B.GetStorageIndex()]->GetData(Op.B.GetRegisterIndex());
+							
+							Result = true;
+							for (int32 ElementIndex = 0; ElementIndex < RegisterA.ElementCount; ElementIndex++)
+							{
+								if (!ScriptStruct->CompareScriptStruct(DataA, DataB, 0))
+								{
+									Result = false;
+									break;
+								}
+								DataA += RegisterA.ElementSize;
+								DataB += RegisterB.ElementSize;
+							}
+
+							break;
+						}
+						case ERigVMRegisterType::Invalid:
+						{
+							break;
+						}
+					}
+				}
+				if (Instructions[InstructionIndex].OpCode == ERigVMOpCode::NotEquals)
+				{
+					Result = !Result;
+				}
+				Storage[Op.Result.GetStorageIndex()]->GetRef<bool>(Op.Result.GetRegisterIndex()) = Result;
+				InstructionIndex++;
 				break;
 			}
-			case ERigVMOpCode::NotEquals:
+			case ERigVMOpCode::JumpAbsolute:
 			{
-				const FRigVMNotEqualsOp& Op = ByteCode.GetOpAt<FRigVMNotEqualsOp>(Instructions[InstructionIndex]);
-				ensureMsgf(false, TEXT("TO BE IMPLEMENTED"));
+				const FRigVMJumpOp& Op = ByteCode.GetOpAt<FRigVMJumpOp>(Instructions[InstructionIndex]);
+				InstructionIndex = Op.InstructionIndex;
 				break;
 			}
-			case ERigVMOpCode::Jump:
+			case ERigVMOpCode::JumpForward:
 			{
 				const FRigVMJumpOp& Op = ByteCode.GetOpAt<FRigVMJumpOp>(Instructions[InstructionIndex]);
-				ensureMsgf(false, TEXT("TO BE IMPLEMENTED"));
+				InstructionIndex += Op.InstructionIndex;
+				break;
+			}
+			case ERigVMOpCode::JumpBackward:
+			{
+				const FRigVMJumpOp& Op = ByteCode.GetOpAt<FRigVMJumpOp>(Instructions[InstructionIndex]);
+				InstructionIndex -= Op.InstructionIndex;
+				break;
+			}
+			case ERigVMOpCode::JumpAbsoluteIf:
+			{
+				const FRigVMJumpIfOp& Op = ByteCode.GetOpAt<FRigVMJumpIfOp>(Instructions[InstructionIndex]);
+				const bool Condition = Storage[0]->GetRef<bool>(Op.ConditionArg.GetRegisterIndex());
+				if (Condition == Op.Condition)
+				{
+					InstructionIndex = Op.InstructionIndex;
+				}
+				else
+				{
+					InstructionIndex++;
+				}
 				break;
 			}
-			case ERigVMOpCode::JumpIfTrue:
+			case ERigVMOpCode::JumpForwardIf:
 			{
-				const FRigVMJumpIfTrueOp& Op = ByteCode.GetOpAt<FRigVMJumpIfTrueOp>(Instructions[InstructionIndex]);
-				ensureMsgf(false, TEXT("TO BE IMPLEMENTED"));
+				const FRigVMJumpIfOp& Op = ByteCode.GetOpAt<FRigVMJumpIfOp>(Instructions[InstructionIndex]);
+				const bool Condition = Storage[0]->GetRef<bool>(Op.ConditionArg.GetRegisterIndex());
+				if (Condition == Op.Condition)
+				{
+					InstructionIndex += Op.InstructionIndex;
+				}
+				else
+				{
+					InstructionIndex++;
+				}
 				break;
 			}
-			case ERigVMOpCode::JumpIfFalse:
+			case ERigVMOpCode::JumpBackwardIf:
 			{
-				const FRigVMJumpIfFalseOp& Op = ByteCode.GetOpAt<FRigVMJumpIfFalseOp>(Instructions[InstructionIndex]);
-				ensureMsgf(false, TEXT("TO BE IMPLEMENTED"));
+				const FRigVMJumpIfOp& Op = ByteCode.GetOpAt<FRigVMJumpIfOp>(Instructions[InstructionIndex]);
+				const bool Condition = Storage[0]->GetRef<bool>(Op.ConditionArg.GetRegisterIndex());
+				if (Condition == Op.Condition)
+				{
+					InstructionIndex -= Op.InstructionIndex;
+				}
+				else
+				{
+					InstructionIndex++;
+				}
 				break;
 			}
 			case ERigVMOpCode::Exit:
diff --git a/Engine/Source/Runtime/RigVM/Private/RigVMCore/RigVMByteCode.cpp b/Engine/Source/Runtime/RigVM/Private/RigVMCore/RigVMByteCode.cpp
--- a/Engine/Source/Runtime/RigVM/Private/RigVMCore/RigVMByteCode.cpp
+++ b/Engine/Source/Runtime/RigVM/Private/RigVMCore/RigVMByteCode.cpp
@@ -119,6 +119,18 @@ uint64 FRigVMByteCode::GetOpNumBytesAt(uint64 InByteCodeIndex, bool bIncludeArgu
 		{
 			return (uint64)sizeof(FRigVMCopyOp);
 		}
+		case ERigVMOpCode::Zero:
+		{
+			return (uint64)sizeof(FRigVMZeroOp);
+		}
+		case ERigVMOpCode::BoolFalse:
+		{
+			return (uint64)sizeof(FRigVMFalseOp);
+		}
+		case ERigVMOpCode::BoolTrue:
+		{
+			return (uint64)sizeof(FRigVMTrueOp);
+		}
 		case ERigVMOpCode::Increment:
 		{
 			return (uint64)sizeof(FRigVMIncrementOp);
@@ -135,17 +147,17 @@ uint64 FRigVMByteCode::GetOpNumBytesAt(uint64 InByteCodeIndex, bool bIncludeArgu
 		{
 			return (uint64)sizeof(FRigVMNotEqualsOp);
 		}
-		case ERigVMOpCode::Jump:
+		case ERigVMOpCode::JumpAbsolute:
+		case ERigVMOpCode::JumpForward:
+		case ERigVMOpCode::JumpBackward:
 		{
 			return (uint64)sizeof(FRigVMJumpOp);
 		}
-		case ERigVMOpCode::JumpIfTrue:
+		case ERigVMOpCode::JumpAbsoluteIf:
+		case ERigVMOpCode::JumpForwardIf:
+		case ERigVMOpCode::JumpBackwardIf:
 		{
-			return (uint64)sizeof(FRigVMJumpIfTrueOp);
-		}
-		case ERigVMOpCode::JumpIfFalse:
-		{
-			return (uint64)sizeof(FRigVMJumpIfFalseOp);
+			return (uint64)sizeof(FRigVMJumpIfOp);
 		}
 		case ERigVMOpCode::Exit:
 		{
@@ -160,6 +172,24 @@ uint64 FRigVMByteCode::GetOpNumBytesAt(uint64 InByteCodeIndex, bool bIncludeArgu
 	return 0;
 }
 
+uint64 FRigVMByteCode::AddZeroOp(const FRigVMArgument& InArg)
+{
+	FRigVMZeroOp Op(InArg);
+	return AddOp(Op);
+}
+
+uint64 FRigVMByteCode::AddFalseOp(const FRigVMArgument& InArg)
+{
+	FRigVMFalseOp Op(InArg);
+	return AddOp(Op);
+}
+
+uint64 FRigVMByteCode::AddTrueOp(const FRigVMArgument& InArg)
+{
+	FRigVMTrueOp Op(InArg);
+	return AddOp(Op);
+}
+
 uint64 FRigVMByteCode::AddCopyOp(const FRigVMArgument& InSource, const FRigVMArgument& InTarget, int32 InSourceOffset, int32 InTargetOffset, int32 InNumBytes)
 {
 	FRigVMCopyOp Op(InSource, InTarget, InSourceOffset, InTargetOffset, InNumBytes);
@@ -190,21 +220,15 @@ uint64 FRigVMByteCode::AddNotEqualsOp(const FRigVMArgument& InA, const FRigVMArg
 	return AddOp(Op);
 }
 
-uint64 FRigVMByteCode::AddJumpOp(uint64 InByteCodeIndex)
-{
-	FRigVMJumpOp Op(InByteCodeIndex);
-	return AddOp(Op);
-}
-
-uint64 FRigVMByteCode::AddJumpIfTrueOp(uint64 InByteCodeIndex, const FRigVMArgument& InCondition)
+uint64 FRigVMByteCode::AddJumpOp(ERigVMOpCode InOpCode, uint16 InInstructionIndex)
 {
-	FRigVMJumpIfTrueOp Op(InByteCodeIndex, InCondition);
+	FRigVMJumpOp Op(InOpCode, InInstructionIndex);
 	return AddOp(Op);
 }
 
-uint64 FRigVMByteCode::AddJumpIfFalseOp(uint64 InByteCodeIndex, const FRigVMArgument& InCondition)
+uint64 FRigVMByteCode::AddJumpIfOp(ERigVMOpCode InOpCode, uint16 InInstructionIndex, const FRigVMArgument& InConditionArg, bool bInCondition)
 {
-	FRigVMJumpIfFalseOp Op(InByteCodeIndex, InCondition);
+	FRigVMJumpIfOp Op(InOpCode, InInstructionIndex, InConditionArg, bInCondition);
 	return AddOp(Op);
 }
 
diff --git a/Engine/Source/Runtime/RigVM/Public/RigVMCore/RigVMByteCode.h b/Engine/Source/Runtime/RigVM/Public/RigVMCore/RigVMByteCode.h
--- a/Engine/Source/Runtime/RigVM/Public/RigVMCore/RigVMByteCode.h
+++ b/Engine/Source/Runtime/RigVM/Public/RigVMCore/RigVMByteCode.h
@@ -76,14 +76,20 @@ enum class ERigVMOpCode : uint8
 	Execute_62_Args,
 	Execute_63_Args,
 	Execute_64_Args,
+	Zero,
+	BoolFalse,
+	BoolTrue,
 	Copy,
 	Increment,
 	Decrement,
 	Equals,
 	NotEquals,
-	Jump,
-	JumpIfTrue,
-	JumpIfFalse,
+	JumpAbsolute,
+	JumpForward,
+	JumpBackward,
+	JumpAbsoluteIf,
+	JumpForwardIf,
+	JumpBackwardIf,
 	Exit,
 	Invalid
 };
@@ -117,6 +123,57 @@ struct RIGVM_API FRigVMExecuteOp : public FRigVMBaseOp
 	FORCEINLINE uint8 GetArgumentCount() const { return uint8(OpCode) - uint8(ERigVMOpCode::Execute_0_Args); }
 };
 
+struct RIGVM_API FRigVMZeroOp : public FRigVMBaseOp
+{
+	FRigVMZeroOp()
+		: FRigVMBaseOp(ERigVMOpCode::Zero)
+		, Arg()
+	{
+	}
+
+	FRigVMZeroOp(FRigVMArgument InArg)
+		: FRigVMBaseOp(ERigVMOpCode::Zero)
+		, Arg(InArg)
+	{
+	}
+
+	FRigVMArgument Arg;
+};
+
+struct RIGVM_API FRigVMFalseOp : public FRigVMBaseOp
+{
+	FRigVMFalseOp()
+		: FRigVMBaseOp(ERigVMOpCode::BoolFalse)
+		, Arg()
+	{
+	}
+
+	FRigVMFalseOp(FRigVMArgument InArg)
+		: FRigVMBaseOp(ERigVMOpCode::BoolFalse)
+		, Arg(InArg)
+	{
+	}
+
+	FRigVMArgument Arg;
+};
+
+struct RIGVM_API FRigVMTrueOp : public FRigVMBaseOp
+{
+	FRigVMTrueOp()
+		: FRigVMBaseOp(ERigVMOpCode::BoolTrue)
+		, Arg()
+	{
+	}
+
+	FRigVMTrueOp(FRigVMArgument InArg)
+		: FRigVMBaseOp(ERigVMOpCode::BoolTrue)
+		, Arg(InArg)
+	{
+	}
+
+	FRigVMArgument Arg;
+};
+
 struct RIGVM_API FRigVMCopyOp : public FRigVMBaseOp
 {
 	FRigVMCopyOp()
@@ -243,58 +300,45 @@ struct RIGVM_API FRigVMNotEqualsOp : public FRigVMBaseOp
 struct RIGVM_API FRigVMJumpOp : public FRigVMBaseOp
 {
 	FRigVMJumpOp()
-	: FRigVMBaseOp(ERigVMOpCode::Jump)
+	: FRigVMBaseOp(ERigVMOpCode::Invalid)
 	, InstructionIndex(INDEX_NONE)
 	{
 	}
 
-	FRigVMJumpOp(int32 InInstructionIndex)
-	: FRigVMBaseOp(ERigVMOpCode::Jump)
+	FRigVMJumpOp(ERigVMOpCode InOpCode, int32 InInstructionIndex)
+	: FRigVMBaseOp(InOpCode)
 	, InstructionIndex(InInstructionIndex)
 	{
+		ensure(uint8(InOpCode) >= uint8(ERigVMOpCode::JumpAbsolute));
+		ensure(uint8(InOpCode) <= uint8(ERigVMOpCode::JumpBackward));
 	}
 
 	int32 InstructionIndex;
 };
 
-struct RIGVM_API FRigVMJumpIfTrueOp : public FRigVMBaseOp
+struct RIGVM_API FRigVMJumpIfOp : public FRigVMBaseOp
 {
-	FRigVMJumpIfTrueOp()
-	: FRigVMBaseOp(ERigVMOpCode::JumpIfTrue)
-	, InstructionIndex(INDEX_NONE)
-	, Condition()
+	FRigVMJumpIfOp()
+		: FRigVMBaseOp(ERigVMOpCode::Invalid)
+		, InstructionIndex(INDEX_NONE)
+		, ConditionArg()
+		, Condition(true)
 	{
 	}
 
-	FRigVMJumpIfTrueOp(int32 InInstructionIndex, FRigVMArgument InCondition)
-	: FRigVMBaseOp(ERigVMOpCode::JumpIfTrue)
-	, InstructionIndex(InInstructionIndex)
-	, Condition(InCondition)
+	FRigVMJumpIfOp(ERigVMOpCode InOpCode, int32 InInstructionIndex, FRigVMArgument InConditionArg, bool InCondition = false)
+		: FRigVMBaseOp(InOpCode)
+		, InstructionIndex(InInstructionIndex)
+		, ConditionArg(InConditionArg)
+		, Condition(InCondition)
 	{
+		ensure(uint8(InOpCode) >= uint8(ERigVMOpCode::JumpAbsoluteIf));
+		ensure(uint8(InOpCode) <= uint8(ERigVMOpCode::JumpBackwardIf));
 	}
 
 	int32 InstructionIndex;
-	FRigVMArgument Condition;
-};
-
-struct RIGVM_API FRigVMJumpIfFalseOp : public FRigVMBaseOp
-{
-	FRigVMJumpIfFalseOp()
-	: FRigVMBaseOp(ERigVMOpCode::JumpIfFalse)
-	, InstructionIndex(INDEX_NONE)
-	, Condition()
-	{
-	}
-
-	FRigVMJumpIfFalseOp(int32 InInstructionIndex, FRigVMArgument InCondition)
-	: FRigVMBaseOp(ERigVMOpCode::JumpIfFalse)
-	, InstructionIndex(InInstructionIndex)
-	, Condition(InCondition)
-	{
-	}
-
-	int32 InstructionIndex;
-	FRigVMArgument Condition;
+	FRigVMArgument ConditionArg;
+	bool Condition;
 };
 
 struct RIGVM_API FRigVMExitOp : public FRigVMBaseOp
@@ -360,15 +404,17 @@ public:
 	void Reset();
 	uint64 Num() const;
 
+	uint64 AddExecuteOp(uint16 InFunctionIndex, const TArrayView<FRigVMArgument>& InArguments);
+	uint64 AddZeroOp(const FRigVMArgument& InArg);
+	uint64 AddFalseOp(const FRigVMArgument& InArg);
+	uint64 AddTrueOp(const FRigVMArgument& InArg);
 	uint64 AddCopyOp(const FRigVMArgument& InSource, const FRigVMArgument& InTarget, int32 InSourceOffset = INDEX_NONE, int32 InTargetOffset = INDEX_NONE, int32 InNumBytes = INDEX_NONE);
 	uint64 AddIncrementOp(const FRigVMArgument& InArg);
 	uint64 AddDecrementOp(const FRigVMArgument& InArg);
 	uint64 AddEqualsOp(const FRigVMArgument& InA, const FRigVMArgument& InB, const FRigVMArgument& InResult);
 	uint64 AddNotEqualsOp(const FRigVMArgument& InA, const FRigVMArgument& InB, const FRigVMArgument& InResult);
-	uint64 AddJumpOp(uint64 InByteCodeIndex);
-	uint64 AddJumpIfTrueOp(uint64 InByteCodeIndex, const FRigVMArgument& InCondition);
-	uint64 AddJumpIfFalseOp(uint64 InByteCodeIndex, const FRigVMArgument& InCondition);
-	uint64 AddExecuteOp(uint16 InFunctionIndex, const TArrayView<FRigVMArgument>& InArguments);
+	uint64 AddJumpOp(ERigVMOpCode InOpCode, uint16 InInstructionIndex);
+	uint64 AddJumpIfOp(ERigVMOpCode InOpCode, uint16 InInstructionIndex, const FRigVMArgument& InConditionArg, bool bInCondition = false);
 	uint64 AddExitOp();
 
 	FORCEINLINE FRigVMByteCodeTable GetTable() const
