RigVM: Support for changetype opcode to reuse registers for smaller work loads

diff --git a/Engine/Source/Runtime/RigVM/Private/RigVMCore/RigVM.cpp b/Engine/Source/Runtime/RigVM/Private/RigVMCore/RigVM.cpp
--- a/Engine/Source/Runtime/RigVM/Private/RigVMCore/RigVM.cpp
+++ b/Engine/Source/Runtime/RigVM/Private/RigVMCore/RigVM.cpp
@@ -333,6 +333,13 @@ bool URigVM::Execute(FRigVMStoragePtrArray Storage, TArrayView<void*> Additional
 				}
 				break;
 			}
+			case ERigVMOpCode::ChangeType:
+			{
+				const FRigVMChangeTypeOp& Op = ByteCode.GetOpAt<FRigVMChangeTypeOp>(Instructions[InstructionIndex]);
+				Storage[Op.Arg.GetStorageIndex()]->ChangeRegisterType(Op.Arg.GetRegisterIndex(), Op.Type, Op.ElementSize, nullptr, Op.ElementCount, Op.SliceCount);
+				InstructionIndex++;
+				break;
+			}
 			case ERigVMOpCode::Exit:
 			case ERigVMOpCode::Invalid:
 			{
diff --git a/Engine/Source/Runtime/RigVM/Private/RigVMCore/RigVMByteCode.cpp b/Engine/Source/Runtime/RigVM/Private/RigVMCore/RigVMByteCode.cpp
--- a/Engine/Source/Runtime/RigVM/Private/RigVMCore/RigVMByteCode.cpp
+++ b/Engine/Source/Runtime/RigVM/Private/RigVMCore/RigVMByteCode.cpp
@@ -159,6 +159,10 @@ uint64 FRigVMByteCode::GetOpNumBytesAt(uint64 InByteCodeIndex, bool bIncludeArgu
 		{
 			return (uint64)sizeof(FRigVMJumpIfOp);
 		}
+		case ERigVMOpCode::ChangeType:
+		{
+			return (uint64)sizeof(FRigVMChangeTypeOp);
+		}
 		case ERigVMOpCode::Exit:
 		{
 			return (uint64)sizeof(FRigVMExitOp);
@@ -232,6 +236,12 @@ uint64 FRigVMByteCode::AddJumpIfOp(ERigVMOpCode InOpCode, uint16 InInstructionIn
 	return AddOp(Op);
 }
 
+uint64 FRigVMByteCode::AddChangeTypeOp(FRigVMArgument InArg, ERigVMRegisterType InType, uint16 InElementSize, uint16 InElementCount, uint16 InSliceCount)
+{
+	FRigVMChangeTypeOp Op(InArg, InType, InElementSize, InElementCount, InSliceCount);
+	return AddOp(Op);
+}
+
 uint64 FRigVMByteCode::AddExecuteOp(uint16 InFunctionIndex, const TArrayView<FRigVMArgument>& InArguments)
 {
 	FRigVMExecuteOp Op(InFunctionIndex, (uint8)InArguments.Num());
diff --git a/Engine/Source/Runtime/RigVM/Public/RigVMCore/RigVMByteCode.h b/Engine/Source/Runtime/RigVM/Public/RigVMCore/RigVMByteCode.h
--- a/Engine/Source/Runtime/RigVM/Public/RigVMCore/RigVMByteCode.h
+++ b/Engine/Source/Runtime/RigVM/Public/RigVMCore/RigVMByteCode.h
@@ -90,6 +90,7 @@ enum class ERigVMOpCode : uint8
 	JumpAbsoluteIf,
 	JumpForwardIf,
 	JumpBackwardIf,
+	ChangeType,
 	Exit,
 	Invalid
 };
@@ -341,6 +342,35 @@ struct RIGVM_API FRigVMJumpIfOp : public FRigVMBaseOp
 	bool Condition;
 };
 
+struct RIGVM_API FRigVMChangeTypeOp : public FRigVMBaseOp
+{
+	FRigVMChangeTypeOp()
+		: FRigVMBaseOp(ERigVMOpCode::Invalid)
+		, Arg()
+		, Type(ERigVMRegisterType::Invalid)
+		, ElementSize(0)
+		, ElementCount(0)
+		, SliceCount(0)
+	{
+	}
+
+	FRigVMChangeTypeOp(FRigVMArgument InArg, ERigVMRegisterType InType, uint16 InElementSize, uint16 InElementCount, uint16 InSliceCount)
+		: FRigVMBaseOp(ERigVMOpCode::ChangeType)
+		, Arg(InArg)
+		, Type(InType)
+		, ElementSize(InElementSize)
+		, ElementCount(InElementCount)
+		, SliceCount(InSliceCount)
+	{
+	}
+
+	FRigVMArgument Arg;
+	ERigVMRegisterType Type;
+	uint16 ElementSize;
+	uint16 ElementCount;
+	uint16 SliceCount;
+};
+
 struct RIGVM_API FRigVMExitOp : public FRigVMBaseOp
 {
 	FRigVMExitOp()
@@ -415,6 +445,7 @@ public:
 	uint64 AddNotEqualsOp(const FRigVMArgument& InA, const FRigVMArgument& InB, const FRigVMArgument& InResult);
 	uint64 AddJumpOp(ERigVMOpCode InOpCode, uint16 InInstructionIndex);
 	uint64 AddJumpIfOp(ERigVMOpCode InOpCode, uint16 InInstructionIndex, const FRigVMArgument& InConditionArg, bool bInCondition = false);
+	uint64 AddChangeTypeOp(FRigVMArgument InArg, ERigVMRegisterType InType, uint16 InElementSize, uint16 InElementCount, uint16 InSliceCount = 1);
 	uint64 AddExitOp();
 
 	FORCEINLINE FRigVMByteCodeTable GetTable() const
