RigVM: Fix linux build breaks

diff --git a/Engine/Source/Runtime/RigVM/Private/RigVMCore/RigVM.cpp b/Engine/Source/Runtime/RigVM/Private/RigVMCore/RigVM.cpp
--- a/Engine/Source/Runtime/RigVM/Private/RigVMCore/RigVM.cpp
+++ b/Engine/Source/Runtime/RigVM/Private/RigVMCore/RigVM.cpp
@@ -50,7 +50,7 @@ void URigVM::ResolveFunctionsIfRequired()
 		Functions.Reset();
 		Functions.SetNumZeroed(FunctionNames.Num());
 
-		for (int32 FunctionIndex = 0; FunctionIndex < FunctionNames.Num(); FunctionIndex)
+		for (int32 FunctionIndex = 0; FunctionIndex < FunctionNames.Num(); FunctionIndex++)
 		{
 			Functions[FunctionIndex] = FRigVMRegistry::Get().Find(*FunctionNames[FunctionIndex]);
 		}
diff --git a/Engine/Source/Runtime/RigVM/Public/RigVMModel/RigVMNotifications.h b/Engine/Source/Runtime/RigVM/Public/RigVMModel/RigVMNotifications.h
--- a/Engine/Source/Runtime/RigVM/Public/RigVMModel/RigVMNotifications.h
+++ b/Engine/Source/Runtime/RigVM/Public/RigVMModel/RigVMNotifications.h
@@ -6,6 +6,8 @@
 #include "UObject/ObjectMacros.h"
 #include "RigVMNotifications.generated.h"
 
+class URigVMGraph;
+
 UENUM()
 enum class ERigVMGraphNotifType : uint8
 {
