RigVM: fix osx build breaks for template specialization

diff --git a/Engine/Source/Runtime/RigVM/Public/RigVMCore/RigVMStorage.h b/Engine/Source/Runtime/RigVM/Public/RigVMCore/RigVMStorage.h
--- a/Engine/Source/Runtime/RigVM/Public/RigVMCore/RigVMStorage.h
+++ b/Engine/Source/Runtime/RigVM/Public/RigVMCore/RigVMStorage.h
@@ -347,58 +347,6 @@ public:
 		return Allocate(InNewName, sizeof(T), InCount, InSliceCount, (const void*)InDataPtr);
 	}
 
-	template<>
-	FORCEINLINE int32 AddPlainArray(const FName& InNewName, int32 InCount, const FName* InDataPtr, int32 InSliceCount)
-	{
-		int32 Register = Allocate(InNewName, sizeof(FName), InCount, InSliceCount, nullptr);
-		Registers[Register].Type = ERigVMRegisterType::Name;
-
-		Construct(Register);
-
-		if (InDataPtr)
-		{
-			Registers[Register].MoveToFirstSlice();
-			for (uint16 SliceIndex = 0; SliceIndex < Registers[Register].SliceCount; SliceIndex++)
-			{
-				FName* DataPtr = (FName*)GetData(Register);
-				for (int32 Index = 0; Index < InCount; Index++)
-				{
-					DataPtr[Index] = InDataPtr[Index];
-				}
-				Registers[Register].MoveToNextSlice();
-			}
-			Registers[Register].MoveToFirstSlice();
-		}
-
-		return Register;
-	}
-
-	template<>
-	FORCEINLINE int32 AddPlainArray(const FName& InNewName, int32 InCount, const FString* InDataPtr, int32 InSliceCount)
-	{
-		int32 Register = Allocate(InNewName, sizeof(FString), InCount, InSliceCount, nullptr);
-		Registers[Register].Type = ERigVMRegisterType::String;
-
-		Construct(Register);
-
-		if (InDataPtr)
-		{
-			Registers[Register].MoveToFirstSlice();
-			for (uint16 SliceIndex = 0; SliceIndex < Registers[Register].SliceCount; SliceIndex++)
-			{
-				FString* DataPtr = (FString*)GetData(Register);
-				for (int32 Index = 0; Index < InCount; Index++)
-				{
-					DataPtr[Index] = InDataPtr[Index];
-				}
-				Registers[Register].MoveToNextSlice();
-			}
-			Registers[Register].MoveToFirstSlice();
-		}
-
-		return Register;
-	}
-
 	template<class T>
 	FORCEINLINE int32 AddPlainArray(const FName& InNewName, const TArray<T>& InArray, int32 InSliceCount = 1)
 	{
@@ -531,58 +479,6 @@ public:
 		return ChangeRegisterType(InRegisterIndex, ERigVMRegisterType::Plain, sizeof(T), InDataPtr, InNewElementCount, InNewSliceCount);
 	}
 
-	template<>
-	FORCEINLINE bool ChangeRegisterType(int32 InRegisterIndex, const FName* InDataPtr, int32 InNewElementCount, int32 InNewSliceCount)
-	{
-		if (!ChangeRegisterType(InRegisterIndex, ERigVMRegisterType::Name, sizeof(FName), nullptr, InNewElementCount, InNewSliceCount))
-		{
-			return false;
-		}
-
-		if (InDataPtr)
-		{
-			Registers[InRegisterIndex].MoveToFirstSlice();
-			for (uint16 SliceIndex = 0; SliceIndex < Registers[InRegisterIndex].SliceCount; SliceIndex++)
-			{
-				FName* DataPtr = (FName*)GetData(InRegisterIndex);
-				for (int32 Index = 0; Index < InNewElementCount; Index++)
-				{
-					DataPtr[Index] = InDataPtr[Index];
-				}
-				Registers[InRegisterIndex].MoveToNextSlice();
-			}
-			Registers[InRegisterIndex].MoveToFirstSlice();
-		}
-
-		return true;
-	}
-
-	template<>
-	FORCEINLINE bool ChangeRegisterType(int32 InRegisterIndex, const FString* InDataPtr, int32 InNewElementCount, int32 InNewSliceCount)
-	{
-		if (!ChangeRegisterType(InRegisterIndex, ERigVMRegisterType::String, sizeof(FString), nullptr, InNewElementCount, InNewSliceCount))
-		{
-			return false;
-		}
-
-		if (InDataPtr)
-		{
-			Registers[InRegisterIndex].MoveToFirstSlice();
-			for (uint16 SliceIndex = 0; SliceIndex < Registers[InRegisterIndex].SliceCount; SliceIndex++)
-			{
-				FString* DataPtr = (FString*)GetData(InRegisterIndex);
-				for (int32 Index = 0; Index < InNewElementCount; Index++)
-				{
-					DataPtr[Index] = InDataPtr[Index];
-				}
-				Registers[InRegisterIndex].MoveToNextSlice();
-			}
-			Registers[InRegisterIndex].MoveToFirstSlice();
-		}
-
-		return true;
-	}
-
 	void UpdateRegisters();
 
 private:
@@ -614,5 +510,109 @@ private:
 	TMap<FName, int32> NameMap;
 };
 
+template<>
+FORCEINLINE bool FRigVMStorage::ChangeRegisterType(int32 InRegisterIndex, const FName* InDataPtr, int32 InNewElementCount, int32 InNewSliceCount)
+{
+	if (!ChangeRegisterType(InRegisterIndex, ERigVMRegisterType::Name, sizeof(FName), nullptr, InNewElementCount, InNewSliceCount))
+	{
+		return false;
+	}
+
+	if (InDataPtr)
+	{
+		Registers[InRegisterIndex].MoveToFirstSlice();
+		for (uint16 SliceIndex = 0; SliceIndex < Registers[InRegisterIndex].SliceCount; SliceIndex++)
+		{
+			FName* DataPtr = (FName*)GetData(InRegisterIndex);
+			for (int32 Index = 0; Index < InNewElementCount; Index++)
+			{
+				DataPtr[Index] = InDataPtr[Index];
+			}
+			Registers[InRegisterIndex].MoveToNextSlice();
+		}
+		Registers[InRegisterIndex].MoveToFirstSlice();
+	}
+
+	return true;
+}
+
+template<>
+FORCEINLINE bool FRigVMStorage::ChangeRegisterType(int32 InRegisterIndex, const FString* InDataPtr, int32 InNewElementCount, int32 InNewSliceCount)
+{
+	if (!ChangeRegisterType(InRegisterIndex, ERigVMRegisterType::String, sizeof(FString), nullptr, InNewElementCount, InNewSliceCount))
+	{
+		return false;
+	}
+
+	if (InDataPtr)
+	{
+		Registers[InRegisterIndex].MoveToFirstSlice();
+		for (uint16 SliceIndex = 0; SliceIndex < Registers[InRegisterIndex].SliceCount; SliceIndex++)
+		{
+			FString* DataPtr = (FString*)GetData(InRegisterIndex);
+			for (int32 Index = 0; Index < InNewElementCount; Index++)
+			{
+				DataPtr[Index] = InDataPtr[Index];
+			}
+			Registers[InRegisterIndex].MoveToNextSlice();
+		}
+		Registers[InRegisterIndex].MoveToFirstSlice();
+	}
+
+	return true;
+}
+
+template<>
+FORCEINLINE int32 FRigVMStorage::AddPlainArray(const FName& InNewName, int32 InCount, const FName* InDataPtr, int32 InSliceCount)
+{
+	int32 Register = Allocate(InNewName, sizeof(FName), InCount, InSliceCount, nullptr);
+	Registers[Register].Type = ERigVMRegisterType::Name;
+
+	Construct(Register);
+
+	if (InDataPtr)
+	{
+		Registers[Register].MoveToFirstSlice();
+		for (uint16 SliceIndex = 0; SliceIndex < Registers[Register].SliceCount; SliceIndex++)
+		{
+			FName* DataPtr = (FName*)GetData(Register);
+			for (int32 Index = 0; Index < InCount; Index++)
+			{
+				DataPtr[Index] = InDataPtr[Index];
+			}
+			Registers[Register].MoveToNextSlice();
+		}
+		Registers[Register].MoveToFirstSlice();
+	}
+
+	return Register;
+}
+
+template<>
+FORCEINLINE int32 FRigVMStorage::AddPlainArray(const FName& InNewName, int32 InCount, const FString* InDataPtr, int32 InSliceCount)
+{
+	int32 Register = Allocate(InNewName, sizeof(FString), InCount, InSliceCount, nullptr);
+	Registers[Register].Type = ERigVMRegisterType::String;
+
+	Construct(Register);
+
+	if (InDataPtr)
+	{
+		Registers[Register].MoveToFirstSlice();
+		for (uint16 SliceIndex = 0; SliceIndex < Registers[Register].SliceCount; SliceIndex++)
+		{
+			FString* DataPtr = (FString*)GetData(Register);
+			for (int32 Index = 0; Index < InCount; Index++)
+			{
+				DataPtr[Index] = InDataPtr[Index];
+			}
+			Registers[Register].MoveToNextSlice();
+		}
+		Registers[Register].MoveToFirstSlice();
+	}
+
+	return Register;
+}
+
 typedef FRigVMStorage* FRigVMStoragePtr;
 typedef TArrayView<FRigVMStoragePtr> FRigVMStoragePtrArray;
\ No newline at end of file
