RigVM : Storage - support for serialization + unit test

diff --git a/Engine/Source/Runtime/Core/Private/UObject/DevObjectVersion.cpp b/Engine/Source/Runtime/Core/Private/UObject/DevObjectVersion.cpp
--- a/Engine/Source/Runtime/Core/Private/UObject/DevObjectVersion.cpp
+++ b/Engine/Source/Runtime/Core/Private/UObject/DevObjectVersion.cpp
@@ -24,6 +24,7 @@
 #include "UObject/AutomationObjectVersion.h"
 #include "UObject/NiagaraObjectVersion.h"
 #include "UObject/DestructionObjectVersion.h"
+#include "UObject/RigVMObjectVersion.h"
 
 DEFINE_LOG_CATEGORY_STATIC(LogDevObjectVersion, Log, All);
 
@@ -163,3 +164,8 @@ FDevVersionRegistration GRegisterNiagaraObjectVersion(FNiagaraObjectVersion::GUI
 const FGuid FDestructionObjectVersion::GUID(0x174F1F0B, 0xB4C645A5, 0xB13F2EE8, 0xD0FB917D);
 // Register Destruction custom version with Core
 FDevVersionRegistration GRegisterDestructionObjectVersion(FDestructionObjectVersion::GUID, FDestructionObjectVersion::LatestVersion, TEXT("Dev-Destruction"));
+
+// Unique RigVM Object version id
+const FGuid FRigVMObjectVersion::GUID(0x155EEB8F, 0xBB004582, 0xA79D3FBC, 0xEFF209E4);
+// Register RigVM custom version with Core
+FDevVersionRegistration GRegisterRigVMObjectVersion(FRigVMObjectVersion::GUID, FRigVMObjectVersion::LatestVersion, TEXT("Dev-RigVM"));
diff --git a/Engine/Source/Runtime/Core/Public/UObject/RigVMObjectVersion.h b/Engine/Source/Runtime/Core/Public/UObject/RigVMObjectVersion.h
new file mode 100644
--- /dev/null
+++ b/Engine/Source/Runtime/Core/Public/UObject/RigVMObjectVersion.h
@@ -0,0 +1,28 @@
+// Copyright 1998-2019 Epic Games, Inc. All Rights Reserved.
+#pragma once
+
+#include "CoreTypes.h"
+#include "Misc/Guid.h"
+
+// Custom serialization version for changes made in RigVM
+struct CORE_API FRigVMObjectVersion
+{
+	enum Type
+	{
+		// Before any version changes were made
+		BeforeCustomVersionWasAdded = 0,
+
+		// Using custom serialization for the data array
+		FirstVersionOfCustomDataSerialization = 1,
+
+		// -----<new versions can be added above this line>-------------------------------------------------
+		VersionPlusOne,
+		LatestVersion = VersionPlusOne - 1
+	};
+
+	// The GUID for this custom version number
+	const static FGuid GUID;
+
+private:
+	FRigVMObjectVersion() {}
+};
diff --git a/Engine/Source/Runtime/RigVM/Private/RigVMCore/RigVMStorage.cpp b/Engine/Source/Runtime/RigVM/Private/RigVMCore/RigVMStorage.cpp
--- a/Engine/Source/Runtime/RigVM/Private/RigVMCore/RigVMStorage.cpp
+++ b/Engine/Source/Runtime/RigVM/Private/RigVMCore/RigVMStorage.cpp
@@ -1,6 +1,31 @@
 // Copyright 1998-2019 Epic Games, Inc. All Rights Reserved.
 
 #include "RigVMCore/RigVMStorage.h"
+#include "Uobject/RigVMObjectVersion.h"
+#include "UObject/PropertyPortFlags.h"
+
+bool FRigVMRegister::Serialize(FArchive& Ar)
+{
+	Ar.UsingCustomVersion(FRigVMObjectVersion::GUID);
+
+	if (Ar.CustomVer(FRigVMObjectVersion::GUID) < FRigVMObjectVersion::FirstVersionOfCustomDataSerialization)
+	{
+		return false;
+	}
+
+	Ar << Type;
+	Ar << ByteIndex;
+	Ar << ElementSize;
+	Ar << ElementCount;
+	Ar << SliceIndex;
+	Ar << SliceCount;
+	Ar << AlignmentBytes;
+	Ar << TrailingBytes;
+	Ar << Name;
+	Ar << ScriptStructIndex;
+
+	return true;
+}
 
 FRigVMStorage::FRigVMStorage(bool bInUseNames)
 	: bUseNameMap(bInUseNames)
@@ -38,6 +63,176 @@ FRigVMStorage& FRigVMStorage::operator= (const FRigVMStorage &InOther)
 	return *this;
 }
 
+bool FRigVMStorage::Serialize(FArchive& Ar)
+{
+	Ar.UsingCustomVersion(FRigVMObjectVersion::GUID);
+
+ 	if (Ar.CustomVer(FRigVMObjectVersion::GUID) < FRigVMObjectVersion::FirstVersionOfCustomDataSerialization)
+ 	{
+ 		return false;
+ 	}
+
+	if (Ar.IsLoading())
+	{
+		for (int32 RegisterIndex = 0; RegisterIndex < Registers.Num(); RegisterIndex++)
+		{
+			Registers[RegisterIndex].MoveToFirstSlice();
+			Destroy(RegisterIndex);
+		}
+	}
+	else
+	{
+		for (int32 RegisterIndex = 0; RegisterIndex < Registers.Num(); RegisterIndex++)
+		{
+			Registers[RegisterIndex].MoveToFirstSlice();
+		}
+	}
+
+	Ar << bUseNameMap;
+	Ar << StorageType;
+	Ar << Registers;
+
+	if (Ar.IsLoading())
+	{
+		ScriptStructs.Reset();
+		TArray<FString> ScriptStructPaths;
+		Ar << ScriptStructPaths;
+		
+		for (const FString& ScriptStructPath : ScriptStructPaths)
+		{
+			UScriptStruct* ScriptStruct = FindObject<UScriptStruct>(nullptr, *ScriptStructPath);
+			ensure(ScriptStruct != nullptr);
+			ScriptStructs.Add(ScriptStruct);
+		}
+
+		uint64 TotalBytes = 0;
+		Ar << TotalBytes;
+		
+		Data.SetNumZeroed(TotalBytes);
+
+		for (int32 RegisterIndex = 0; RegisterIndex < Registers.Num(); RegisterIndex++)
+		{
+			Registers[RegisterIndex].MoveToFirstSlice();
+			Construct(RegisterIndex);
+		}
+
+		for (const FRigVMRegister& Register : Registers)
+		{
+			switch (Register.Type)
+			{
+				case ERigVMRegisterType::Plain:
+				{
+					TArray<uint8> View;
+					Ar << View;
+					ensure(View.Num() == Register.GetAllocatedBytes());
+					FMemory::Memcpy(&Data[Register.GetStorageByteIndex()], View.GetData(), View.Num());
+					break;
+				}
+				case ERigVMRegisterType::Name:
+				{
+					TArray<FName> View;
+					Ar << View;
+					ensure(View.Num() == Register.GetTotalElementCount());
+					for (uint16 ElementIndex = 0; ElementIndex < Register.GetTotalElementCount(); ElementIndex++)
+					{
+						*((FName*)&Data[Register.GetWorkByteIndex() + Register.ElementSize * ElementIndex]) = View[ElementIndex];
+					}
+					break;
+				}
+				case ERigVMRegisterType::String:
+				{
+					TArray<FString> View;
+					Ar << View;
+					ensure(View.Num() == Register.GetTotalElementCount());
+					for (uint16 ElementIndex = 0; ElementIndex < Register.GetTotalElementCount(); ElementIndex++)
+					{
+						*((FString*)&Data[Register.GetWorkByteIndex() + Register.ElementSize * ElementIndex]) = View[ElementIndex];
+					}
+					break;
+				}
+				case ERigVMRegisterType::Struct:
+				{
+					TArray<FString> View;
+					Ar << View;
+					ensure(View.Num() == Register.GetTotalElementCount());
+
+					uint8* DataPtr = &Data[Register.GetWorkByteIndex()];
+					UScriptStruct* ScriptStruct = ScriptStructs[Register.ScriptStructIndex];
+					for (uint16 ElementIndex = 0; ElementIndex < Register.GetTotalElementCount(); ElementIndex++)
+					{
+						ScriptStruct->ImportText(*View[ElementIndex], DataPtr, nullptr, PPF_None, nullptr, ScriptStruct->GetName());
+						DataPtr += Register.ElementSize;
+					}
+					break;
+				}
+			}
+		}
+
+		UpdateRegisters();
+	}
+	else
+	{
+		TArray<FString> ScriptStructPaths;
+		for (UScriptStruct* ScriptStruct : ScriptStructs)
+		{
+			ScriptStructPaths.Add(ScriptStruct->GetPathName());
+		}
+		Ar << ScriptStructPaths;
+
+		uint64 TotalBytes = Data.Num();
+		Ar << TotalBytes;
+
+		for (FRigVMRegister& Register : Registers)
+		{
+			Register.MoveToFirstSlice();
+
+			switch(Register.Type)
+			{
+				case ERigVMRegisterType::Plain:
+				{
+					TArray<uint8> View;
+					View.Append(&Data[Register.GetStorageByteIndex()], Register.GetAllocatedBytes());
+					Ar << View;
+					break;
+				}
+				case ERigVMRegisterType::Name:
+				{
+					TArray<FName> View;
+					View.Append((FName*)&Data[Register.GetWorkByteIndex()], Register.GetTotalElementCount());
+					Ar << View;
+					break;
+				}
+				case ERigVMRegisterType::String:
+				{
+					TArray<FString> View;
+					View.Append((FString*)&Data[Register.GetWorkByteIndex()], Register.GetTotalElementCount());
+					Ar << View;
+					break;
+				}
+				case ERigVMRegisterType::Struct:
+				{
+					uint8* DataPtr = &Data[Register.GetWorkByteIndex()];
+					UScriptStruct* ScriptStruct = ScriptStructs[Register.ScriptStructIndex];
+
+					TArray<FString> View;
+					for (uint16 ElementIndex = 0; ElementIndex < Register.GetTotalElementCount(); ElementIndex++)
+					{
+						FString Value;
+						ScriptStruct->ExportText(Value, DataPtr, nullptr, nullptr, PPF_None, nullptr);
+						View.Add(Value);
+						DataPtr += Register.ElementSize;
+					}
+
+					Ar << View;
+					break;
+				}
+			}
+		}
+	}
+
+	return true;
+}
+
 void FRigVMStorage::Reset()
 {
 	for (int32 Index = 0; Index < Registers.Num(); Index++)
diff --git a/Engine/Source/Runtime/RigVM/Public/RigVMCore/RigVMStorage.h b/Engine/Source/Runtime/RigVM/Public/RigVMCore/RigVMStorage.h
--- a/Engine/Source/Runtime/RigVM/Public/RigVMCore/RigVMStorage.h
+++ b/Engine/Source/Runtime/RigVM/Public/RigVMCore/RigVMStorage.h
@@ -59,9 +59,9 @@ enum class ERigVMRegisterType : uint8
 USTRUCT()
 struct RIGVM_API FRigVMRegister
 {
-	GENERATED_USTRUCT_BODY()
+	GENERATED_BODY()
 
-		FRigVMRegister()
+	FRigVMRegister()
 		: Type(ERigVMRegisterType::Invalid)
 		, ByteIndex(INDEX_NONE)
 		, ElementSize(0)
@@ -105,6 +105,13 @@ struct RIGVM_API FRigVMRegister
 	UPROPERTY()
 	int32 ScriptStructIndex;
 
+	bool Serialize(FArchive& Ar);
+	FORCEINLINE friend FArchive& operator<<(FArchive& Ar, FRigVMRegister& P)
+	{
+		P.Serialize(Ar);
+		return Ar;
+	}
+
 	FORCEINLINE uint64 GetWorkByteIndex() const { return ByteIndex; }
 	FORCEINLINE uint64 GetStorageByteIndex() const { return ByteIndex - (uint64)AlignmentBytes - (uint64)(SliceIndex * GetNumBytesPerSlice()); }
 	FORCEINLINE uint8 GetAlignmentBytes() const { return AlignmentBytes; }
@@ -139,7 +146,7 @@ typedef TArrayView<FRigVMRegister> FRigVMRegisterArray;
 USTRUCT()
 struct RIGVM_API FRigVMStorage
 {
-	GENERATED_USTRUCT_BODY()
+	GENERATED_BODY()
 	
 public:
 
@@ -165,6 +172,13 @@ public:
 	FORCEINLINE TArray<FRigVMRegister>::RangedForIteratorType      end() { return Registers.end(); }
 	FORCEINLINE TArray<FRigVMRegister>::RangedForConstIteratorType end() const { return Registers.end(); }
 
+	bool Serialize(FArchive& Ar);
+	FORCEINLINE friend FArchive& operator<<(FArchive& Ar, FRigVMStorage& P)
+	{
+		P.Serialize(Ar);
+		return Ar;
+	}
+
 	FORCEINLINE FRigVMArgument GetArgument(int32 InRegisterIndex) const
 	{
 		ensure(Registers.IsValidIndex(InRegisterIndex));
@@ -500,7 +514,7 @@ private:
 	UPROPERTY()
 	TArray<FRigVMRegister> Registers;
 
-	UPROPERTY()
+	UPROPERTY(transient)
 	TArray<uint8> Data;
 
 	UPROPERTY()
