RigVM: Remove RigVMObjectVersion and use AnimObjectVersion

diff --git a/Engine/Source/Runtime/Core/Private/UObject/DevObjectVersion.cpp b/Engine/Source/Runtime/Core/Private/UObject/DevObjectVersion.cpp
--- a/Engine/Source/Runtime/Core/Private/UObject/DevObjectVersion.cpp
+++ b/Engine/Source/Runtime/Core/Private/UObject/DevObjectVersion.cpp
@@ -24,7 +24,6 @@
 #include "UObject/AutomationObjectVersion.h"
 #include "UObject/NiagaraObjectVersion.h"
 #include "UObject/DestructionObjectVersion.h"
-#include "UObject/RigVMObjectVersion.h"
 
 DEFINE_LOG_CATEGORY_STATIC(LogDevObjectVersion, Log, All);
 
@@ -164,8 +163,3 @@ FDevVersionRegistration GRegisterNiagaraObjectVersion(FNiagaraObjectVersion::GUI
 const FGuid FDestructionObjectVersion::GUID(0x174F1F0B, 0xB4C645A5, 0xB13F2EE8, 0xD0FB917D);
 // Register Destruction custom version with Core
 FDevVersionRegistration GRegisterDestructionObjectVersion(FDestructionObjectVersion::GUID, FDestructionObjectVersion::LatestVersion, TEXT("Dev-Destruction"));
-
-// Unique RigVM Object version id
-const FGuid FRigVMObjectVersion::GUID(0x155EEB8F, 0xBB004582, 0xA79D3FBC, 0xEFF209E4);
-// Register RigVM custom version with Core
-FDevVersionRegistration GRegisterRigVMObjectVersion(FRigVMObjectVersion::GUID, FRigVMObjectVersion::LatestVersion, TEXT("Dev-RigVM"));
diff --git a/Engine/Source/Runtime/Core/Public/UObject/RigVMObjectVersion.h b/Engine/Source/Runtime/Core/Public/UObject/RigVMObjectVersion.h
deleted file mode 100644
diff --git a/Engine/Source/Runtime/RigVM/Private/RigVMCore/RigVMStorage.cpp b/Engine/Source/Runtime/RigVM/Private/RigVMCore/RigVMStorage.cpp
--- a/Engine/Source/Runtime/RigVM/Private/RigVMCore/RigVMStorage.cpp
+++ b/Engine/Source/Runtime/RigVM/Private/RigVMCore/RigVMStorage.cpp
@@ -1,14 +1,14 @@
 // Copyright 1998-2019 Epic Games, Inc. All Rights Reserved.
 
 #include "RigVMCore/RigVMStorage.h"
-#include "Uobject/RigVMObjectVersion.h"
+#include "Uobject/AnimObjectVersion.h"
 #include "UObject/PropertyPortFlags.h"
 
 bool FRigVMRegister::Serialize(FArchive& Ar)
 {
-	Ar.UsingCustomVersion(FRigVMObjectVersion::GUID);
+	Ar.UsingCustomVersion(FAnimObjectVersion::GUID);
 
-	if (Ar.CustomVer(FRigVMObjectVersion::GUID) < FRigVMObjectVersion::FirstVersionOfCustomDataSerialization)
+	if (Ar.CustomVer(FAnimObjectVersion::GUID) < FAnimObjectVersion::StoreMarkerNamesOnSkeleton)
 	{
 		return false;
 	}
@@ -65,9 +65,9 @@ FRigVMStorage& FRigVMStorage::operator= (const FRigVMStorage &InOther)
 
 bool FRigVMStorage::Serialize(FArchive& Ar)
 {
-	Ar.UsingCustomVersion(FRigVMObjectVersion::GUID);
+	Ar.UsingCustomVersion(FAnimObjectVersion::GUID);
 
- 	if (Ar.CustomVer(FRigVMObjectVersion::GUID) < FRigVMObjectVersion::FirstVersionOfCustomDataSerialization)
+ 	if (Ar.CustomVer(FAnimObjectVersion::GUID) < FAnimObjectVersion::StoreMarkerNamesOnSkeleton)
  	{
  		return false;
  	}
diff --git a/Engine/Source/Runtime/RigVM/Public/RigVMCore/RigVMStorage.h b/Engine/Source/Runtime/RigVM/Public/RigVMCore/RigVMStorage.h
--- a/Engine/Source/Runtime/RigVM/Public/RigVMCore/RigVMStorage.h
+++ b/Engine/Source/Runtime/RigVM/Public/RigVMCore/RigVMStorage.h
@@ -6,6 +6,9 @@
 #include "UObject/ObjectMacros.h"
 #include "RigVMStorage.generated.h"
 
+/**
+ * 
+ */
 UENUM()
 enum class ERigVMStorageType: uint8
 {
