RigVM: Renamed RigVMStorage to RigVMMemoryContainer and documented headers

diff --git a/Engine/Source/Programs/UnrealHeaderTool/Private/CodeGenerator.cpp b/Engine/Source/Programs/UnrealHeaderTool/Private/CodeGenerator.cpp
--- a/Engine/Source/Programs/UnrealHeaderTool/Private/CodeGenerator.cpp
+++ b/Engine/Source/Programs/UnrealHeaderTool/Private/CodeGenerator.cpp
@@ -3337,7 +3337,7 @@ void FNativeClassHeaderGenerator::ExportGeneratedStructBodyMacros(FOutputDevice&
 				VariableType = Parameter.TypeOriginal(true);
 			}
 
-			RigVMStubProlog.Add(FString::Printf(TEXT("%s %s = RigVMStorage[%s_Arg.GetStorageIndex()]->%s<%s>(%s_Arg.GetRegisterIndex(), true);"),
+			RigVMStubProlog.Add(FString::Printf(TEXT("%s %s = RigVMMemoryContainer[%s_Arg.GetContainerIndex()]->%s<%s>(%s_Arg.GetRegisterIndex(), true);"),
 				*VariableType,
 				*Parameter.NameOriginal(false),
 				*Parameter.NameOriginal(false),
@@ -3383,7 +3383,7 @@ void FNativeClassHeaderGenerator::ExportGeneratedStructBodyMacros(FOutputDevice&
 			{
 				FString ParameterSuffix = MethodInfo.Parameters.Declarations(true, TEXT(",\r\n\t\t"));
 				RigVMMethodsDeclarations += FString::Printf(TEXT("\tstatic %s Static%s(\r\n\t\t%s%s\r\n\t);\r\n"), *MethodInfo.ReturnType, *MethodInfo.Name, *StructMembers, *ParameterSuffix);
-				RigVMMethodsDeclarations += FString::Printf(TEXT("\tstatic %s RigVM%s(\r\n\t\tconst FRigVMArgumentArray& RigVMArguments,\r\n\t\tFRigVMStoragePtrArray& RigVMStorage,\r\n\t\tconst FRigVMUserDataArray& RigVMUserData\r\n\t);\r\n"), *MethodInfo.ReturnType, *MethodInfo.Name);
+				RigVMMethodsDeclarations += FString::Printf(TEXT("\tstatic %s RigVM%s(\r\n\t\tconst FRigVMArgumentArray& RigVMArguments,\r\n\t\tFRigVMMemoryContainerPtrArray& RigVMMemoryContainer,\r\n\t\tconst FRigVMUserDataArray& RigVMUserData\r\n\t);\r\n"), *MethodInfo.ReturnType, *MethodInfo.Name);
 			}
 		}
 
@@ -3638,7 +3638,7 @@ void FNativeClassHeaderGenerator::ExportGeneratedStructBodyMacros(FOutputDevice&
 			Out.Log(TEXT("\r\n"));
 
 			// implement stub method body
-			Out.Logf(TEXT("%s %s::RigVM%s(const FRigVMArgumentArray& RigVMArguments, FRigVMStoragePtrArray& RigVMStorage, const FRigVMUserDataArray& RigVMUserData)\r\n"), *MethodInfo.ReturnType, StructNameCPP, *MethodInfo.Name);
+			Out.Logf(TEXT("%s %s::RigVM%s(const FRigVMArgumentArray& RigVMArguments, FRigVMMemoryContainerPtrArray& RigVMMemoryContainer, const FRigVMUserDataArray& RigVMUserData)\r\n"), *MethodInfo.ReturnType, StructNameCPP, *MethodInfo.Name);
 			Out.Log(TEXT("{\r\n"));
 
 			if (MethodInfo.Parameters.Num() > 0)
diff --git a/Engine/Source/Runtime/RigVM/Private/RigVMCore/RigVM.cpp b/Engine/Source/Runtime/RigVM/Private/RigVMCore/RigVM.cpp
--- a/Engine/Source/Runtime/RigVM/Private/RigVMCore/RigVM.cpp
+++ b/Engine/Source/Runtime/RigVM/Private/RigVMCore/RigVM.cpp
@@ -4,8 +4,8 @@
 
 URigVM::URigVM()
 {
-	WorkStorage.SetStorageType(ERigVMStorageType::Work);
-	LiteralStorage.SetStorageType(ERigVMStorageType::Literal);
+	WorkMemory.SetMemoryType(ERigVMMemoryType::Work);
+	LiteralMemory.SetMemoryType(ERigVMMemoryType::Literal);
 }
 
 URigVM::~URigVM()
@@ -15,8 +15,8 @@ URigVM::~URigVM()
 
 void URigVM::Reset()
 {
-	WorkStorage.Reset();
-	LiteralStorage.Reset();
+	WorkMemory.Reset();
+	LiteralMemory.Reset();
 	FunctionNames.Reset();
 	Functions.Reset();
 	ByteCode.Reset();
@@ -65,7 +65,7 @@ void URigVM::RefreshInstructionsIfRequired()
 	}
 }
 
-bool URigVM::Execute(FRigVMStoragePtrArray Storage, TArrayView<void*> AdditionalArgs)
+bool URigVM::Execute(FRigVMMemoryContainerPtrArray Memory, TArrayView<void*> AdditionalArgs)
 {
 	ResolveFunctionsIfRequired();
 	RefreshInstructionsIfRequired();
@@ -75,10 +75,10 @@ bool URigVM::Execute(FRigVMStoragePtrArray Storage, TArrayView<void*> Additional
 		return true;
 	}
 
-	FRigVMStorage* LocalStorage[] = { &WorkStorage, &LiteralStorage }; 
-	if (Storage.Num() == 0)
+	FRigVMMemoryContainer* LocalMemory[] = { &WorkMemory, &LiteralMemory }; 
+	if (Memory.Num() == 0)
 	{
-		Storage = FRigVMStoragePtrArray(LocalStorage, 2);
+		Memory = FRigVMMemoryContainerPtrArray(LocalMemory, 2);
 	}
 
 	uint16 InstructionIndex = 0;
@@ -154,49 +154,49 @@ bool URigVM::Execute(FRigVMStoragePtrArray Storage, TArrayView<void*> Additional
 			{
 				const FRigVMExecuteOp& Op = ByteCode.GetOpAt<FRigVMExecuteOp>(Instructions[InstructionIndex]);
 				FRigVMArgumentArray Args = ByteCode.GetArgumentsForExecuteOp(Instructions[InstructionIndex]);
-				(*Functions[Op.FunctionIndex])(Args, Storage, AdditionalArgs);
+				(*Functions[Op.FunctionIndex])(Args, Memory, AdditionalArgs);
 				InstructionIndex++;
 				break;
 			}
 			case ERigVMOpCode::Zero:
 			{
 				const FRigVMZeroOp& Op = ByteCode.GetOpAt<FRigVMZeroOp>(Instructions[InstructionIndex]);
-				Storage[Op.Arg.GetStorageIndex()]->GetRef<int32>(Op.Arg.GetRegisterIndex()) = 0;
+				Memory[Op.Arg.GetContainerIndex()]->GetRef<int32>(Op.Arg.GetRegisterIndex()) = 0;
 				InstructionIndex++;
 				break;
 			}
 			case ERigVMOpCode::BoolFalse:
 			{
 				const FRigVMFalseOp& Op = ByteCode.GetOpAt<FRigVMFalseOp>(Instructions[InstructionIndex]);
-				Storage[Op.Arg.GetStorageIndex()]->GetRef<bool>(Op.Arg.GetRegisterIndex()) = false;
+				Memory[Op.Arg.GetContainerIndex()]->GetRef<bool>(Op.Arg.GetRegisterIndex()) = false;
 				InstructionIndex++;
 				break;
 			}
 			case ERigVMOpCode::BoolTrue:
 			{
 				const FRigVMTrueOp& Op = ByteCode.GetOpAt<FRigVMTrueOp>(Instructions[InstructionIndex]);
-				Storage[Op.Arg.GetStorageIndex()]->GetRef<bool>(Op.Arg.GetRegisterIndex()) = false;
+				Memory[Op.Arg.GetContainerIndex()]->GetRef<bool>(Op.Arg.GetRegisterIndex()) = false;
 				InstructionIndex++;
 				break;
 			}
 			case ERigVMOpCode::Copy:
 			{
 				const FRigVMCopyOp& Op = ByteCode.GetOpAt<FRigVMCopyOp>(Instructions[InstructionIndex]);
-				Storage[Op.Target.GetStorageIndex()]->Copy(Op.Source.GetRegisterIndex(), Op.Target.GetRegisterIndex(), Storage[Op.Source.GetStorageIndex()], Op.SourceOffset, Op.TargetOffset, Op.NumBytes);
+				Memory[Op.Target.GetContainerIndex()]->Copy(Op.Source.GetRegisterIndex(), Op.Target.GetRegisterIndex(), Memory[Op.Source.GetContainerIndex()], Op.SourceOffset, Op.TargetOffset, Op.NumBytes);
 				InstructionIndex++;
 				break;
 			}
 			case ERigVMOpCode::Increment:
 			{
 				const FRigVMIncrementOp& Op = ByteCode.GetOpAt<FRigVMIncrementOp>(Instructions[InstructionIndex]);
-				Storage[Op.Arg.GetStorageIndex()]->GetRef<int32>(Op.Arg.GetRegisterIndex())++;
+				Memory[Op.Arg.GetContainerIndex()]->GetRef<int32>(Op.Arg.GetRegisterIndex())++;
 				InstructionIndex++;
 				break;
 			}
 			case ERigVMOpCode::Decrement:
 			{
 				const FRigVMDecrementOp& Op = ByteCode.GetOpAt<FRigVMDecrementOp>(Instructions[InstructionIndex]);
-				Storage[Op.Arg.GetStorageIndex()]->GetRef<int32>(Op.Arg.GetRegisterIndex())--;
+				Memory[Op.Arg.GetContainerIndex()]->GetRef<int32>(Op.Arg.GetRegisterIndex())--;
 				InstructionIndex++;
 				break;
 			}
@@ -204,8 +204,8 @@ bool URigVM::Execute(FRigVMStoragePtrArray Storage, TArrayView<void*> Additional
 			case ERigVMOpCode::NotEquals:
 			{
 				const FRigVMEqualsOp& Op = ByteCode.GetOpAt<FRigVMEqualsOp>(Instructions[InstructionIndex]);
-				const FRigVMRegister& RegisterA = (*Storage[Op.A.GetStorageIndex()])[Op.A.GetRegisterIndex()];
-				const FRigVMRegister& RegisterB = (*Storage[Op.B.GetStorageIndex()])[Op.B.GetRegisterIndex()];
+				const FRigVMRegister& RegisterA = (*Memory[Op.A.GetContainerIndex()])[Op.A.GetRegisterIndex()];
+				const FRigVMRegister& RegisterB = (*Memory[Op.B.GetContainerIndex()])[Op.B.GetRegisterIndex()];
 				uint16 BytesA = RegisterA.GetNumBytesPerSlice();
 				uint16 BytesB = RegisterB.GetNumBytesPerSlice();
 				
@@ -217,15 +217,15 @@ bool URigVM::Execute(FRigVMStoragePtrArray Storage, TArrayView<void*> Additional
 						case ERigVMRegisterType::Plain:
 						case ERigVMRegisterType::Name:
 						{
-							void * DataA = Storage[Op.A.GetStorageIndex()]->GetData(Op.A.GetRegisterIndex());
-							void * DataB = Storage[Op.B.GetStorageIndex()]->GetData(Op.B.GetRegisterIndex());
+							void * DataA = Memory[Op.A.GetContainerIndex()]->GetData(Op.A.GetRegisterIndex());
+							void * DataB = Memory[Op.B.GetContainerIndex()]->GetData(Op.B.GetRegisterIndex());
 							Result = FMemory::Memcmp(DataA, DataB, BytesA) == 0;
 							break;
 						}
 						case ERigVMRegisterType::String:
 						{
-							TArrayView<FString> StringsA = Storage[Op.A.GetStorageIndex()]->GetArray<FString>(Op.A.GetRegisterIndex());
-							TArrayView<FString> StringsB = Storage[Op.B.GetStorageIndex()]->GetArray<FString>(Op.B.GetRegisterIndex());
+							TArrayView<FString> StringsA = Memory[Op.A.GetContainerIndex()]->GetArray<FString>(Op.A.GetRegisterIndex());
+							TArrayView<FString> StringsB = Memory[Op.B.GetContainerIndex()]->GetArray<FString>(Op.B.GetRegisterIndex());
 
 							Result = true;
 							for (int32 StringIndex = 0; StringIndex < StringsA.Num(); StringIndex++)
@@ -240,10 +240,10 @@ bool URigVM::Execute(FRigVMStoragePtrArray Storage, TArrayView<void*> Additional
 						}
 						case ERigVMRegisterType::Struct:
 						{
-							UScriptStruct* ScriptStruct = Storage[Op.A.GetStorageIndex()]->GetScriptStruct(RegisterA.ScriptStructIndex);
+							UScriptStruct* ScriptStruct = Memory[Op.A.GetContainerIndex()]->GetScriptStruct(RegisterA.ScriptStructIndex);
 
-							uint8* DataA = (uint8*)Storage[Op.A.GetStorageIndex()]->GetData(Op.A.GetRegisterIndex());
-							uint8* DataB = (uint8*)Storage[Op.B.GetStorageIndex()]->GetData(Op.B.GetRegisterIndex());
+							uint8* DataA = (uint8*)Memory[Op.A.GetContainerIndex()]->GetData(Op.A.GetRegisterIndex());
+							uint8* DataB = (uint8*)Memory[Op.B.GetContainerIndex()]->GetData(Op.B.GetRegisterIndex());
 							
 							Result = true;
 							for (int32 ElementIndex = 0; ElementIndex < RegisterA.ElementCount; ElementIndex++)
@@ -269,7 +269,7 @@ bool URigVM::Execute(FRigVMStoragePtrArray Storage, TArrayView<void*> Additional
 				{
 					Result = !Result;
 				}
-				Storage[Op.Result.GetStorageIndex()]->GetRef<bool>(Op.Result.GetRegisterIndex()) = Result;
+				Memory[Op.Result.GetContainerIndex()]->GetRef<bool>(Op.Result.GetRegisterIndex()) = Result;
 				InstructionIndex++;
 				break;
 			}
@@ -294,7 +294,7 @@ bool URigVM::Execute(FRigVMStoragePtrArray Storage, TArrayView<void*> Additional
 			case ERigVMOpCode::JumpAbsoluteIf:
 			{
 				const FRigVMJumpIfOp& Op = ByteCode.GetOpAt<FRigVMJumpIfOp>(Instructions[InstructionIndex]);
-				const bool Condition = Storage[0]->GetRef<bool>(Op.ConditionArg.GetRegisterIndex());
+				const bool Condition = Memory[0]->GetRef<bool>(Op.ConditionArg.GetRegisterIndex());
 				if (Condition == Op.Condition)
 				{
 					InstructionIndex = Op.InstructionIndex;
@@ -308,7 +308,7 @@ bool URigVM::Execute(FRigVMStoragePtrArray Storage, TArrayView<void*> Additional
 			case ERigVMOpCode::JumpForwardIf:
 			{
 				const FRigVMJumpIfOp& Op = ByteCode.GetOpAt<FRigVMJumpIfOp>(Instructions[InstructionIndex]);
-				const bool Condition = Storage[0]->GetRef<bool>(Op.ConditionArg.GetRegisterIndex());
+				const bool Condition = Memory[0]->GetRef<bool>(Op.ConditionArg.GetRegisterIndex());
 				if (Condition == Op.Condition)
 				{
 					InstructionIndex += Op.InstructionIndex;
@@ -322,7 +322,7 @@ bool URigVM::Execute(FRigVMStoragePtrArray Storage, TArrayView<void*> Additional
 			case ERigVMOpCode::JumpBackwardIf:
 			{
 				const FRigVMJumpIfOp& Op = ByteCode.GetOpAt<FRigVMJumpIfOp>(Instructions[InstructionIndex]);
-				const bool Condition = Storage[0]->GetRef<bool>(Op.ConditionArg.GetRegisterIndex());
+				const bool Condition = Memory[0]->GetRef<bool>(Op.ConditionArg.GetRegisterIndex());
 				if (Condition == Op.Condition)
 				{
 					InstructionIndex -= Op.InstructionIndex;
@@ -336,7 +336,7 @@ bool URigVM::Execute(FRigVMStoragePtrArray Storage, TArrayView<void*> Additional
 			case ERigVMOpCode::ChangeType:
 			{
 				const FRigVMChangeTypeOp& Op = ByteCode.GetOpAt<FRigVMChangeTypeOp>(Instructions[InstructionIndex]);
-				Storage[Op.Arg.GetStorageIndex()]->ChangeRegisterType(Op.Arg.GetRegisterIndex(), Op.Type, Op.ElementSize, nullptr, Op.ElementCount, Op.SliceCount);
+				Memory[Op.Arg.GetContainerIndex()]->ChangeRegisterType(Op.Arg.GetRegisterIndex(), Op.Type, Op.ElementSize, nullptr, Op.ElementCount, Op.SliceCount);
 				InstructionIndex++;
 				break;
 			}
diff --git a/Engine/Source/Runtime/RigVM/Private/RigVMCore/RigVMMemory.cpp b/Engine/Source/Runtime/RigVM/Private/RigVMCore/RigVMMemory.cpp
new file mode 100644
--- /dev/null
+++ b/Engine/Source/Runtime/RigVM/Private/RigVMCore/RigVMMemory.cpp
@@ -0,0 +1,780 @@
+// Copyright 1998-2019 Epic Games, Inc. All Rights Reserved.
+
+#include "RigVMCore/RigVMMemory.h"
+#include "Uobject/AnimObjectVersion.h"
+#include "UObject/PropertyPortFlags.h"
+
+bool FRigVMRegister::Serialize(FArchive& Ar)
+{
+	Ar.UsingCustomVersion(FAnimObjectVersion::GUID);
+
+	if (Ar.CustomVer(FAnimObjectVersion::GUID) < FAnimObjectVersion::StoreMarkerNamesOnSkeleton)
+	{
+		return false;
+	}
+
+	Ar << Type;
+	Ar << ByteIndex;
+	Ar << ElementSize;
+	Ar << ElementCount;
+	Ar << SliceIndex;
+	Ar << SliceCount;
+	Ar << AlignmentBytes;
+	Ar << TrailingBytes;
+	Ar << Name;
+	Ar << ScriptStructIndex;
+
+	return true;
+}
+
+FRigVMMemoryContainer::FRigVMMemoryContainer(bool bInUseNames)
+	: bUseNameMap(bInUseNames)
+	, MemoryType(ERigVMMemoryType::Work)
+{
+}
+
+FRigVMMemoryContainer::FRigVMMemoryContainer(const FRigVMMemoryContainer& Other)
+{
+	*this = Other;
+}
+
+FRigVMMemoryContainer::~FRigVMMemoryContainer()
+{
+	Reset();
+}
+
+FRigVMMemoryContainer& FRigVMMemoryContainer::operator= (const FRigVMMemoryContainer &InOther)
+{
+	Reset();
+
+	bUseNameMap = InOther.bUseNameMap;
+	Data.Append(InOther.Data);
+	Registers.Append(InOther.Registers);
+	ScriptStructs.Append(InOther.ScriptStructs);
+
+	UpdateRegisters();
+
+	for (int32 Index = 0; Index < Registers.Num(); Index++)
+	{
+		Construct(Index);
+		Copy(Index, Index, &InOther);
+	}
+
+	return *this;
+}
+
+bool FRigVMMemoryContainer::Serialize(FArchive& Ar)
+{
+	Ar.UsingCustomVersion(FAnimObjectVersion::GUID);
+
+ 	if (Ar.CustomVer(FAnimObjectVersion::GUID) < FAnimObjectVersion::StoreMarkerNamesOnSkeleton)
+ 	{
+ 		return false;
+ 	}
+
+	if (Ar.IsLoading())
+	{
+		for (int32 RegisterIndex = 0; RegisterIndex < Registers.Num(); RegisterIndex++)
+		{
+			Registers[RegisterIndex].MoveToFirstSlice();
+			Destroy(RegisterIndex);
+		}
+	}
+	else
+	{
+		for (int32 RegisterIndex = 0; RegisterIndex < Registers.Num(); RegisterIndex++)
+		{
+			Registers[RegisterIndex].MoveToFirstSlice();
+		}
+	}
+
+	Ar << bUseNameMap;
+	Ar << MemoryType;
+	Ar << Registers;
+
+	if (Ar.IsLoading())
+	{
+		ScriptStructs.Reset();
+		TArray<FString> ScriptStructPaths;
+		Ar << ScriptStructPaths;
+		
+		for (const FString& ScriptStructPath : ScriptStructPaths)
+		{
+			UScriptStruct* ScriptStruct = FindObject<UScriptStruct>(nullptr, *ScriptStructPath);
+			ensure(ScriptStruct != nullptr);
+			ScriptStructs.Add(ScriptStruct);
+		}
+
+		uint64 TotalBytes = 0;
+		Ar << TotalBytes;
+		
+		Data.SetNumZeroed(TotalBytes);
+
+		for (int32 RegisterIndex = 0; RegisterIndex < Registers.Num(); RegisterIndex++)
+		{
+			Registers[RegisterIndex].MoveToFirstSlice();
+			Construct(RegisterIndex);
+		}
+
+		for (const FRigVMRegister& Register : Registers)
+		{
+			switch (Register.Type)
+			{
+				case ERigVMRegisterType::Plain:
+				{
+					TArray<uint8> View;
+					Ar << View;
+					ensure(View.Num() == Register.GetAllocatedBytes());
+					FMemory::Memcpy(&Data[Register.GetFirstAllocatedByte()], View.GetData(), View.Num());
+					break;
+				}
+				case ERigVMRegisterType::Name:
+				{
+					TArray<FName> View;
+					Ar << View;
+					ensure(View.Num() == Register.GetTotalElementCount());
+					for (uint16 ElementIndex = 0; ElementIndex < Register.GetTotalElementCount(); ElementIndex++)
+					{
+						*((FName*)&Data[Register.GetWorkByteIndex() + Register.ElementSize * ElementIndex]) = View[ElementIndex];
+					}
+					break;
+				}
+				case ERigVMRegisterType::String:
+				{
+					TArray<FString> View;
+					Ar << View;
+					ensure(View.Num() == Register.GetTotalElementCount());
+					for (uint16 ElementIndex = 0; ElementIndex < Register.GetTotalElementCount(); ElementIndex++)
+					{
+						*((FString*)&Data[Register.GetWorkByteIndex() + Register.ElementSize * ElementIndex]) = View[ElementIndex];
+					}
+					break;
+				}
+				case ERigVMRegisterType::Struct:
+				{
+					TArray<FString> View;
+					Ar << View;
+					ensure(View.Num() == Register.GetTotalElementCount());
+
+					uint8* DataPtr = &Data[Register.GetWorkByteIndex()];
+					UScriptStruct* ScriptStruct = ScriptStructs[Register.ScriptStructIndex];
+					for (uint16 ElementIndex = 0; ElementIndex < Register.GetTotalElementCount(); ElementIndex++)
+					{
+						ScriptStruct->ImportText(*View[ElementIndex], DataPtr, nullptr, PPF_None, nullptr, ScriptStruct->GetName());
+						DataPtr += Register.ElementSize;
+					}
+					break;
+				}
+			}
+		}
+
+		UpdateRegisters();
+	}
+	else
+	{
+		TArray<FString> ScriptStructPaths;
+		for (UScriptStruct* ScriptStruct : ScriptStructs)
+		{
+			ScriptStructPaths.Add(ScriptStruct->GetPathName());
+		}
+		Ar << ScriptStructPaths;
+
+		uint64 TotalBytes = Data.Num();
+		Ar << TotalBytes;
+
+		for (FRigVMRegister& Register : Registers)
+		{
+			Register.MoveToFirstSlice();
+
+			switch(Register.Type)
+			{
+				case ERigVMRegisterType::Plain:
+				{
+					TArray<uint8> View;
+					View.Append(&Data[Register.GetFirstAllocatedByte()], Register.GetAllocatedBytes());
+					Ar << View;
+					break;
+				}
+				case ERigVMRegisterType::Name:
+				{
+					TArray<FName> View;
+					View.Append((FName*)&Data[Register.GetWorkByteIndex()], Register.GetTotalElementCount());
+					Ar << View;
+					break;
+				}
+				case ERigVMRegisterType::String:
+				{
+					TArray<FString> View;
+					View.Append((FString*)&Data[Register.GetWorkByteIndex()], Register.GetTotalElementCount());
+					Ar << View;
+					break;
+				}
+				case ERigVMRegisterType::Struct:
+				{
+					uint8* DataPtr = &Data[Register.GetWorkByteIndex()];
+					UScriptStruct* ScriptStruct = ScriptStructs[Register.ScriptStructIndex];
+
+					TArray<FString> View;
+					for (uint16 ElementIndex = 0; ElementIndex < Register.GetTotalElementCount(); ElementIndex++)
+					{
+						FString Value;
+						ScriptStruct->ExportText(Value, DataPtr, nullptr, nullptr, PPF_None, nullptr);
+						View.Add(Value);
+						DataPtr += Register.ElementSize;
+					}
+
+					Ar << View;
+					break;
+				}
+			}
+		}
+	}
+
+	return true;
+}
+
+void FRigVMMemoryContainer::Reset()
+{
+	for (int32 Index = 0; Index < Registers.Num(); Index++)
+	{
+		Destroy(Index);
+	}
+
+	Data.Reset();
+	Registers.Reset();
+	ScriptStructs.Reset();
+	NameMap.Reset();
+}
+
+bool FRigVMMemoryContainer::Copy(
+	int32 InSourceRegisterIndex,
+	int32 InTargetRegisterIndex,
+	const FRigVMMemoryContainer* InSourceMemory,
+	int32 InSourceByteOffset,
+	int32 InTargetByteOffset,
+	int32 InNumBytes)
+{
+	if (InSourceMemory == nullptr)
+	{
+		InSourceMemory = this;
+	}
+
+	ensure(InSourceMemory->Registers.IsValidIndex(InSourceRegisterIndex));
+	ensure(Registers.IsValidIndex(InTargetRegisterIndex));
+
+	if (InSourceRegisterIndex == InTargetRegisterIndex && InSourceByteOffset == InTargetByteOffset && this == InSourceMemory)
+	{
+		return false;
+	}
+
+	const FRigVMRegister& Source = InSourceMemory->Registers[InSourceRegisterIndex];
+	const FRigVMRegister& Target = Registers[InTargetRegisterIndex];
+
+	int32 SourceStartByte = Source.GetWorkByteIndex();
+	int32 SourceNumBytes = Source.GetNumBytesAllSlices();
+	if(InSourceByteOffset != INDEX_NONE)
+	{
+		ensure(InNumBytes != INDEX_NONE);
+		ensure(InSourceByteOffset >= 0 && InSourceByteOffset < Source.GetNumBytesAllSlices());
+		ensure(InNumBytes > 0 && InSourceByteOffset + InNumBytes <= Source.GetNumBytesAllSlices());
+
+		SourceStartByte += InSourceByteOffset;
+		SourceNumBytes = InNumBytes;
+	}
+
+	int32 TargetStartByte = Target.GetWorkByteIndex();
+	int32 TargetNumBytes = Target.GetNumBytesAllSlices();
+	if(InTargetByteOffset != INDEX_NONE)
+	{
+		ensure(InNumBytes != INDEX_NONE);
+		ensure(InTargetByteOffset >= 0 && InTargetByteOffset < Target.GetNumBytesAllSlices());
+		ensure(InNumBytes > 0 && InTargetByteOffset + InNumBytes <= Target.GetNumBytesAllSlices());
+
+		TargetStartByte += InTargetByteOffset;
+		TargetNumBytes = InNumBytes;
+	}
+
+	if (SourceNumBytes != TargetNumBytes)
+	{
+		return false;
+	}
+
+	switch (Target.Type)
+	{
+		case ERigVMRegisterType::Plain:
+		{
+			FMemory::Memcpy(&Data[TargetStartByte], &InSourceMemory->Data[SourceStartByte], TargetNumBytes);
+			break;
+		}
+		case ERigVMRegisterType::Struct:
+		{
+			UScriptStruct* ScriptStruct = GetScriptStruct(InTargetRegisterIndex);
+			int32 NumStructs = TargetNumBytes / ScriptStruct->GetStructureSize();
+			ScriptStruct->CopyScriptStruct(&Data[TargetStartByte], &InSourceMemory->Data[SourceStartByte], NumStructs);
+			break;
+		}
+		case ERigVMRegisterType::Name:
+		{
+			int32 NumNames = TargetNumBytes / sizeof(FName);
+			TArrayView<FName> TargetNames((FName*)&Data[TargetStartByte], NumNames);
+			TArrayView<FName> SourceNames((FName*)&InSourceMemory->Data[SourceStartByte], NumNames);
+			for (int32 Index = 0; Index < NumNames; Index++)
+			{
+				TargetNames[Index] = SourceNames[Index];
+			}
+			break;
+		}
+		case ERigVMRegisterType::String:
+		{
+			int32 NumStrings = TargetNumBytes / sizeof(FString);
+			TArrayView<FString> TargetStrings((FString*)&Data[TargetStartByte], NumStrings);
+			TArrayView<FString> SourceStrings((FString*)&InSourceMemory->Data[SourceStartByte], NumStrings);
+			for (int32 Index = 0; Index < NumStrings; Index++)
+			{
+				TargetStrings[Index] = SourceStrings[Index];
+			}
+			break;
+		}
+		case ERigVMRegisterType::Invalid:
+		{
+			return false;
+		}
+	}
+
+	return true;
+}
+
+bool FRigVMMemoryContainer::Copy(
+	const FName& InSourceName,
+	const FName& InTargetName,
+	const FRigVMMemoryContainer* InSourceMemory,
+	int32 InSourceByteOffset,
+	int32 InTargetByteOffset,
+	int32 InNumBytes)
+{
+	ensure(bUseNameMap);
+
+	int32 SourceRegisterIndex = GetIndex(InSourceName);
+	int32 TargetRegisterIndex = GetIndex(InTargetName);
+
+	if(SourceRegisterIndex == INDEX_NONE || TargetRegisterIndex == INDEX_NONE)
+	{
+		return false;
+	}
+
+	return Copy(SourceRegisterIndex, TargetRegisterIndex, InSourceMemory, InSourceByteOffset, InTargetByteOffset, InNumBytes);
+}
+
+int32 FRigVMMemoryContainer::Allocate(const FName& InNewName, int32 InElementSize, int32 InElementCount, int32 InSliceCount, const void* InDataPtr, bool bUpdateRegisters)
+{
+	FName Name = InNewName;
+	if (bUseNameMap && InNewName == NAME_None)
+	{
+		const TCHAR* RegisterPrefix = TEXT("Register");
+		int32 RegisterSuffix = 0;
+		do
+		{
+			Name = FName(*FString::Printf(TEXT("%s_%d"), RegisterPrefix, RegisterSuffix++));
+		} while (!IsNameAvailable(Name));
+	}
+
+	ensure(InElementSize > 0 && InElementCount > 0 && InSliceCount > 0);
+
+	if (bUseNameMap)
+	{
+		if (!IsNameAvailable(Name))
+		{
+			return INDEX_NONE;
+		}
+	}
+
+	FRigVMRegister NewRegister;
+	NewRegister.ByteIndex = Data.Num();
+	if (bUseNameMap)
+	{
+		NewRegister.Name = Name;
+	}
+	NewRegister.ElementSize = (uint16)InElementSize;
+	NewRegister.ElementCount = (uint16)InElementCount;
+	NewRegister.SliceCount = (uint16)InSliceCount;
+	NewRegister.Type = ERigVMRegisterType::Plain;
+
+	Data.AddZeroed(NewRegister.GetAllocatedBytes());
+
+	if (InDataPtr != nullptr)
+	{
+		NewRegister.MoveToFirstSlice();
+		for (uint16 SliceIndex = 0; SliceIndex < NewRegister.SliceCount; SliceIndex++)
+		{
+			FMemory::Memcpy(&Data[NewRegister.GetWorkByteIndex()], InDataPtr, NewRegister.GetNumBytesPerSlice());
+			NewRegister.MoveToNextSlice();
+		}
+		NewRegister.MoveToFirstSlice();
+	}
+
+	int32 RegisterIndex = Registers.Num();
+	Registers.Add(NewRegister);
+
+	if (bUpdateRegisters)
+	{
+		UpdateRegisters();
+	}
+	return RegisterIndex;
+}
+
+int32 FRigVMMemoryContainer::Allocate(int32 InElementSize, int32 InElementCount, int32 InSliceCount, const void* InDataPtr, bool bUpdateRegisters)
+{
+	return Allocate(NAME_None, InElementSize, InElementCount, InSliceCount, InDataPtr, bUpdateRegisters);
+}
+
+bool FRigVMMemoryContainer::Construct(int32 InRegisterIndex, int32 InElementIndex)
+{
+	ensure(Registers.IsValidIndex(InRegisterIndex));
+
+	const FRigVMRegister& Register = Registers[InRegisterIndex];
+	switch (Register.Type)
+	{
+		case ERigVMRegisterType::Struct:
+		{
+			void* DataPtr = (void*)&Data[InElementIndex == INDEX_NONE ? Register.GetWorkByteIndex() : Register.GetWorkByteIndex() + InElementIndex * Register.ElementSize];
+			int32 Count = InElementIndex == INDEX_NONE ? Register.GetTotalElementCount() : 1;
+
+			UScriptStruct* ScriptStruct = GetScriptStruct(InRegisterIndex);
+			ScriptStruct->InitializeStruct(DataPtr, Count);
+			break;
+		}
+		case ERigVMRegisterType::String:
+		{
+			FString* DataPtr = (FString*)&Data[InElementIndex == INDEX_NONE ? Register.GetWorkByteIndex() : Register.GetWorkByteIndex() + InElementIndex * Register.ElementSize];
+			int32 Count = InElementIndex == INDEX_NONE ? Register.GetTotalElementCount() : 1;
+
+			FMemory::Memzero(DataPtr, Count * Register.ElementSize);
+			for (int32 Index = 0; Index < Count; Index++)
+			{
+				DataPtr[Index] = FString();
+			}
+			break;
+		}
+		case ERigVMRegisterType::Name:
+		{
+			FName* DataPtr = (FName*)&Data[InElementIndex == INDEX_NONE ? Register.GetWorkByteIndex() : Register.GetWorkByteIndex() + InElementIndex * Register.ElementSize];
+			int32 Count = InElementIndex == INDEX_NONE ? Register.GetTotalElementCount() : 1;
+
+			FMemory::Memzero(DataPtr, Count * Register.ElementSize);
+			for (int32 Index = 0; Index < Count; Index++)
+			{
+				DataPtr[Index] = FName();
+			}
+			break;
+		}
+		default:
+		{
+			return false;
+		}
+	}
+
+	return true;
+}
+
+bool FRigVMMemoryContainer::Destroy(int32 InRegisterIndex, int32 InElementIndex)
+{
+	ensure(Registers.IsValidIndex(InRegisterIndex));
+
+	FRigVMRegister& Register = Registers[InRegisterIndex];
+
+	if (InElementIndex == INDEX_NONE)
+	{
+		Register.MoveToFirstSlice();
+	}
+
+	switch (Register.Type)
+	{
+		case ERigVMRegisterType::Struct:
+		{
+			void* DataPtr = (void*)&Data[InElementIndex == INDEX_NONE ? Register.GetWorkByteIndex() : Register.GetWorkByteIndex() + InElementIndex * Register.ElementSize];
+			int32 Count = InElementIndex == INDEX_NONE ? Register.GetTotalElementCount() : 1;
+
+			UScriptStruct* ScriptStruct = GetScriptStruct(InRegisterIndex);
+			ScriptStruct->DestroyStruct(DataPtr, Count);
+			break;
+		}
+		case ERigVMRegisterType::String:
+		{
+			FString* DataPtr = (FString*)&Data[InElementIndex == INDEX_NONE ? Register.GetWorkByteIndex() : Register.GetWorkByteIndex() + InElementIndex * Register.ElementSize];
+			int32 Count = InElementIndex == INDEX_NONE ? Register.GetTotalElementCount() : 1;
+
+			for (int32 Index = 0; Index < Count; Index++)
+			{
+				DataPtr[Index] = FString();
+			}
+			break;
+		}
+		case ERigVMRegisterType::Name:
+		{
+			FName* DataPtr = (FName*)&Data[InElementIndex == INDEX_NONE ? Register.GetWorkByteIndex() : Register.GetWorkByteIndex() + InElementIndex * Register.ElementSize];
+			int32 Count = InElementIndex == INDEX_NONE ? Register.GetTotalElementCount() : 1;
+
+			for (int32 Index = 0; Index < Count; Index++)
+			{
+				DataPtr[Index] = FName();
+			}
+			break;
+		}
+		default:
+		{
+			return false;
+		}
+	}
+
+	return true;
+}
+
+bool FRigVMMemoryContainer::Remove(int32 InRegisterIndex)
+{
+	if (InRegisterIndex < 0 || InRegisterIndex >= Registers.Num())
+	{
+		return false;
+	}
+
+	Destroy(InRegisterIndex);
+
+	FRigVMRegister RegisterToRemove = Registers[InRegisterIndex];
+	Data.RemoveAt(RegisterToRemove.ByteIndex, RegisterToRemove.GetAllocatedBytes());
+	Registers.RemoveAt(InRegisterIndex);
+
+	for (int32 Index = InRegisterIndex; Index < Registers.Num(); Index++)
+	{
+		Registers[Index].ByteIndex -= RegisterToRemove.GetAllocatedBytes();
+	}
+
+	UpdateRegisters();
+	return true;
+}
+
+bool FRigVMMemoryContainer::Remove(const FName& InRegisterName)
+{
+	ensure(bUseNameMap);
+	return Remove(GetIndex(InRegisterName));
+}
+
+FName FRigVMMemoryContainer::Rename(int32 InRegisterIndex, const FName& InNewName)
+{
+	if (Registers[InRegisterIndex].Name == InNewName)
+	{
+		return Registers[InRegisterIndex].Name;
+	}
+
+	if (!IsNameAvailable(InNewName))
+	{
+		return Registers[InRegisterIndex].Name;
+	}
+
+	Registers[InRegisterIndex].Name = InNewName;
+	UpdateRegisters();
+
+	return InNewName;
+}
+
+FName FRigVMMemoryContainer::Rename(const FName& InOldName, const FName& InNewName)
+{
+	ensure(bUseNameMap);
+
+	int32 RegisterIndex = GetIndex(InOldName);
+	if (RegisterIndex == INDEX_NONE)
+	{
+		return NAME_None;
+	}
+
+	return Rename(RegisterIndex, InNewName);
+}
+
+bool FRigVMMemoryContainer::Resize(int32 InRegisterIndex, int32 InNewElementCount, int32 InNewSliceCount)
+{
+	ensure(Registers.IsValidIndex(InRegisterIndex));
+	ensure(Registers[InRegisterIndex].TrailingBytes == 0);
+
+	if (InNewElementCount <= 0 || InNewSliceCount < 0)
+	{
+		return Remove(InRegisterIndex);
+	}
+
+	uint32 NewTotalCount = (uint32)InNewElementCount * (uint32)InNewSliceCount;
+	if (Registers[InRegisterIndex].GetTotalElementCount() == NewTotalCount)
+	{
+		return false;
+	}
+
+	FRigVMRegister& Register = Registers[InRegisterIndex];
+	Register.MoveToFirstSlice();
+
+	if (Register.GetTotalElementCount() > NewTotalCount) // shrink
+	{
+		int32 ElementsToRemove = Register.GetTotalElementCount() - NewTotalCount;
+		int32 NumBytesToRemove = Register.ElementSize * ElementsToRemove;
+		int32 FirstByteToRemove = Register.GetWorkByteIndex() + Register.ElementSize * NewTotalCount;
+
+		for (uint32 ElementIndex = NewTotalCount; ElementIndex < Register.GetTotalElementCount(); ElementIndex++)
+		{
+			Destroy(InRegisterIndex, (uint32)ElementIndex);
+		}
+
+		Data.RemoveAt(FirstByteToRemove, NumBytesToRemove);
+		Register.ElementCount = (uint32)InNewElementCount;
+		Register.SliceCount = (uint32)InNewSliceCount;
+
+		for (int32 RegisterIndex = InRegisterIndex + 1; RegisterIndex < Registers.Num(); RegisterIndex++)
+		{
+			Registers[RegisterIndex].ByteIndex -= NumBytesToRemove;
+		}
+	}
+	else // grow
+	{
+		int32 OldElementCount = Register.GetTotalElementCount();
+		int32 ElementsToAdd = NewTotalCount - Register.GetTotalElementCount();
+		int32 NumBytesToAdd = Register.ElementSize * ElementsToAdd;
+		int32 FirstByteToAdd = Register.GetWorkByteIndex() + Register.ElementSize * Register.GetTotalElementCount();
+
+		Data.InsertZeroed(FirstByteToAdd, NumBytesToAdd);
+		Register.ElementCount = (uint32)InNewElementCount;
+		Register.SliceCount = (uint32)InNewSliceCount;
+
+		for (uint32 ElementIndex = OldElementCount; ElementIndex < NewTotalCount; ElementIndex++)
+		{
+			Construct(InRegisterIndex, (int32)ElementIndex);
+		}
+
+		for (int32 RegisterIndex = InRegisterIndex + 1; RegisterIndex < Registers.Num(); RegisterIndex++)
+		{
+			Registers[RegisterIndex].ByteIndex += NumBytesToAdd;
+		}
+	}
+
+	UpdateRegisters();
+	return true;
+}
+
+bool FRigVMMemoryContainer::Resize(const FName& InRegisterName, int32 InNewElementCount, int32 InNewSliceCount)
+{
+	ensure(bUseNameMap);
+
+	int32 RegisterIndex = GetIndex(InRegisterName);
+	if (RegisterIndex == INDEX_NONE)
+	{
+		return false;
+	}
+
+	return Resize(RegisterIndex, InNewElementCount, InNewSliceCount);
+}
+
+bool FRigVMMemoryContainer::ChangeRegisterType(int32 InRegisterIndex, ERigVMRegisterType InNewType, int32 InElementSize, const void* InDataPtr, int32 InNewElementCount, int32 InNewSliceCount)
+{
+	ensure(Registers.IsValidIndex(InRegisterIndex));
+	
+	FRigVMRegister& Register = Registers[InRegisterIndex];
+	ensure(Register.ScriptStructIndex == INDEX_NONE);
+	ensure(Register.AlignmentBytes == 0);
+	ensure(InNewType == ERigVMRegisterType::Name || InNewType == ERigVMRegisterType::String || InNewType == ERigVMRegisterType::Plain);
+
+	Register.MoveToFirstSlice();
+
+	if (Register.Type == InNewType && Register.ElementSize == InElementSize && Register.ElementCount == InNewElementCount && Register.SliceCount == InNewSliceCount)
+	{
+		return false;
+	}
+
+	Destroy(InRegisterIndex);
+
+	uint16 OldAllocatedBytes = Register.GetAllocatedBytes();
+	uint16 NewAllocatedBytes = (uint16)InElementSize * (uint16)InNewElementCount * (uint16)InNewSliceCount;
+	ensure(NewAllocatedBytes <= OldAllocatedBytes);
+
+	Register.Type = InNewType;
+	Register.ElementSize = (uint16)InElementSize;
+	Register.ElementCount = (uint16)InNewElementCount;
+	Register.SliceCount = (uint16)InNewSliceCount;
+	Register.TrailingBytes = OldAllocatedBytes - NewAllocatedBytes;
+
+	Construct(InRegisterIndex);
+
+	if (InDataPtr != nullptr)
+	{
+		for (int32 SliceIndex = 0; SliceIndex < Register.SliceCount; SliceIndex++)
+		{
+			FMemory::Memcpy(GetData(InRegisterIndex), InDataPtr, Register.GetNumBytesPerSlice());
+			Register.MoveToNextSlice();
+		}
+		Register.MoveToFirstSlice();
+	}
+
+	return true;
+}
+
+void FRigVMMemoryContainer::UpdateRegisters()
+{
+	int32 AlignmentShift = 0;
+	for (int32 RegisterIndex = 0; RegisterIndex < Registers.Num(); RegisterIndex++)
+	{
+		FRigVMRegister& Register = Registers[RegisterIndex];
+		Register.ByteIndex += AlignmentShift;
+		Register.MoveToFirstSlice();
+
+		UScriptStruct* ScriptStruct = GetScriptStruct(RegisterIndex);
+		if (ScriptStruct != nullptr)
+		{
+			UScriptStruct::ICppStructOps* TheCppStructOps = ScriptStruct->GetCppStructOps();
+			if (TheCppStructOps != NULL)
+			{
+				if (!TheCppStructOps->HasZeroConstructor())
+				{
+					void* Pointer = GetData(RegisterIndex);
+
+					if (Register.AlignmentBytes > 0)
+					{
+						if (!IsAligned(Pointer, TheCppStructOps->GetAlignment()))
+						{
+							Data.RemoveAt(Register.GetFirstAllocatedByte(), Register.AlignmentBytes);
+							AlignmentShift -= Register.AlignmentBytes;
+							Register.ByteIndex -= Register.AlignmentBytes;
+							Register.AlignmentBytes = 0;
+							Pointer = GetData(RegisterIndex);
+						}
+					}
+
+					while (!IsAligned(Pointer, TheCppStructOps->GetAlignment()))
+					{
+						Data.InsertZeroed(Register.GetFirstAllocatedByte(), 1);
+						Register.AlignmentBytes++;
+						Register.ByteIndex++;
+						AlignmentShift++;
+						Pointer = GetData(RegisterIndex);
+					}
+				}
+			}
+		}
+	}
+
+	if (bUseNameMap)
+	{
+		NameMap.Reset();
+		for (int32 Index = 0; Index < Registers.Num(); Index++)
+		{
+			NameMap.Add(Registers[Index].Name, Index);
+		}
+	}
+}
+
+void FRigVMMemoryContainer::FillWithZeroes(int32 InRegisterIndex)
+{
+	ensure(Registers.IsValidIndex(InRegisterIndex));
+	FMemory::Memzero(GetData(InRegisterIndex), Registers[InRegisterIndex].GetNumBytesAllSlices());
+}
+
+int32 FRigVMMemoryContainer::FindOrAddScriptStruct(UScriptStruct* InScriptStruct)
+{
+	int32 StructIndex = INDEX_NONE;
+	if (ScriptStructs.Find(InScriptStruct, StructIndex))
+	{
+		return StructIndex;
+	}
+	return ScriptStructs.Add(InScriptStruct);
+}
diff --git a/Engine/Source/Runtime/RigVM/Private/RigVMCore/RigVMStorage.cpp b/Engine/Source/Runtime/RigVM/Private/RigVMCore/RigVMStorage.cpp
deleted file mode 100644
diff --git a/Engine/Source/Runtime/RigVM/Public/RigVMCore/RigVM.h b/Engine/Source/Runtime/RigVM/Public/RigVMCore/RigVM.h
--- a/Engine/Source/Runtime/RigVM/Public/RigVMCore/RigVM.h
+++ b/Engine/Source/Runtime/RigVM/Public/RigVMCore/RigVM.h
@@ -3,7 +3,7 @@
 #pragma once
 
 #include "CoreMinimal.h"
-#include "RigVMStorage.h"
+#include "RigVMMemory.h"
 #include "RigVMRegistry.h"
 #include "RigVMByteCode.h"
 #include "RigVM.generated.h"
@@ -20,16 +20,16 @@ public:
 
 	void Reset();
 
-	bool Execute(FRigVMStoragePtrArray Storage = FRigVMStoragePtrArray(), TArrayView<void*> AdditionalArgs = TArrayView<void*>());
+	bool Execute(FRigVMMemoryContainerPtrArray Memory = FRigVMMemoryContainerPtrArray(), TArrayView<void*> AdditionalArgs = TArrayView<void*>());
 
 	UFUNCTION()
 	int32 AddRigVMFunction(UScriptStruct* InRigVMStruct, const FName& InFunctionName);
 
 	UPROPERTY()
-	FRigVMStorage WorkStorage;
+	FRigVMMemoryContainer WorkMemory;
 
 	UPROPERTY()
-	FRigVMStorage LiteralStorage;
+	FRigVMMemoryContainer LiteralMemory;
 
 	UPROPERTY()
 	FRigVMByteCode ByteCode;
diff --git a/Engine/Source/Runtime/RigVM/Public/RigVMCore/RigVMByteCode.h b/Engine/Source/Runtime/RigVM/Public/RigVMCore/RigVMByteCode.h
--- a/Engine/Source/Runtime/RigVM/Public/RigVMCore/RigVMByteCode.h
+++ b/Engine/Source/Runtime/RigVM/Public/RigVMCore/RigVMByteCode.h
@@ -508,7 +508,7 @@ private:
 		return ByteIndex;
 	}
 
-	// storage for all functions
+	// memory for all functions
 	UPROPERTY()
 	TArray<uint8> ByteCode;
 };
diff --git a/Engine/Source/Runtime/RigVM/Public/RigVMCore/RigVMMemory.h b/Engine/Source/Runtime/RigVM/Public/RigVMCore/RigVMMemory.h
new file mode 100644
--- /dev/null
+++ b/Engine/Source/Runtime/RigVM/Public/RigVMCore/RigVMMemory.h
@@ -0,0 +1,888 @@
+// Copyright 1998-2019 Epic Games, Inc. All Rights Reserved.
+
+#pragma once
+
+#include "CoreMinimal.h"
+#include "UObject/ObjectMacros.h"
+#include "RigVMMemory.generated.h"
+
+/**
+ * The type of memory used. Typically we differentiate between
+ * Work (Mutable) and Literal (Constant) memory.
+ */
+UENUM()
+enum class ERigVMMemoryType: uint8
+{
+	Work, // Mutable state
+	Literal, // Const / fixed state
+	Invalid
+};
+
+/**
+ * The FRigVMArgument represents an argument used for an operator
+ * within the virtual machine. Arguments provide information about
+ * which memory needs to be referred to, which register within the
+ * memory all the way to the actual byte address in memory.
+ * The FRigVMArgument is a light weight address for a register in
+ * a FRigVMMemoryContainer.
+ */
+struct RIGVM_API FRigVMArgument
+{
+public:
+
+	FRigVMArgument()
+		: MemoryType(ERigVMMemoryType::Work)
+		, RegisterIndex(INDEX_NONE)
+		, ByteIndex(INDEX_NONE)
+	{
+	}
+
+	FRigVMArgument(ERigVMMemoryType InMemoryType, int32 InRegisterIndex, int32 InByteIndex)
+		: MemoryType(InMemoryType)
+		, RegisterIndex(InRegisterIndex)
+		, ByteIndex(InByteIndex)
+	{
+	}
+
+	// returns the memory type of this argument
+	FORCEINLINE ERigVMMemoryType GetMemoryType() const { return MemoryType; }
+
+	// returns the index of the container of this argument
+	FORCEINLINE int32 GetContainerIndex() const { return (int32)MemoryType; }
+
+	// returns the index of the register of this argument
+	FORCEINLINE uint16 GetRegisterIndex() const { return RegisterIndex; }
+
+	// returns the index of the byte inside of the memory of this argument
+	FORCEINLINE uint64 GetByteIndex() const { return ByteIndex; }
+
+private:
+
+	ERigVMMemoryType MemoryType;
+	uint16 RegisterIndex;
+	uint64 ByteIndex;
+};
+
+typedef TArrayView<FRigVMArgument> FRigVMArgumentArray;
+
+/**
+ * The type of register within the memory.
+ */
+UENUM()
+enum class ERigVMRegisterType : uint8
+{
+	Plain, // bool, int32, float, FVector etc.
+	String, // FString
+	Name, // FName
+	Struct, // Any USTRUCT
+	Invalid
+};
+
+// The register represents an address with the VM's memory. Within a register
+// we can store arbitrary data, so it provides a series of properties to
+// describe the memory location.
+// Registers also support the notion of slices. A slice is a complete copy of
+// the memory - so for example if you register stores 4 Vectors, then a slice 
+// would contain 48 bytes (4 * 3 * 4). The register can however store multiple
+// slices / copies of that if needed. Slices can be used to provide 
+// per-invocation memory to functions within the same register.
+// An integrator for example that needs to store a simulated position
+// might want access to a separate memory memory per loop iteration.
+USTRUCT()
+struct RIGVM_API FRigVMRegister
+{
+	GENERATED_BODY()
+
+	FRigVMRegister()
+		: Type(ERigVMRegisterType::Invalid)
+		, ByteIndex(INDEX_NONE)
+		, ElementSize(0)
+		, ElementCount(0)
+		, SliceIndex(0)
+		, SliceCount(1)
+		, AlignmentBytes(0)
+		, TrailingBytes(0)
+		, Name(NAME_None)
+		, ScriptStructIndex(INDEX_NONE)
+	{
+	}
+
+	// The type of register (plain, name, string, etc.)
+	UPROPERTY()
+	ERigVMRegisterType Type;
+
+	// The index of the first work byte
+	UPROPERTY()
+	uint32 ByteIndex;
+
+	// The size of each store element
+	UPROPERTY()
+	uint16 ElementSize;
+
+	// The number of elements in this register
+	UPROPERTY()
+	uint16 ElementCount;
+
+	// The currently active slice index
+	UPROPERTY()
+	uint16 SliceIndex;
+
+	// The number of slices (complete copies)
+	UPROPERTY()
+	uint16 SliceCount;
+
+	// The number of leading bytes for alignment
+	UPROPERTY()
+	uint8 AlignmentBytes;
+
+	// The numbber of trailing bytes.
+	// These originate after shrinking a register.
+	UPROPERTY()
+	uint16 TrailingBytes;
+
+	// The name of the register (can be None)
+	UPROPERTY()
+	FName Name;
+
+	// For struct registers this is the index of the
+	// struct used - otherwise INDEX_NONE
+	UPROPERTY()
+	int32 ScriptStructIndex;
+
+	bool Serialize(FArchive& Ar);
+	FORCEINLINE friend FArchive& operator<<(FArchive& Ar, FRigVMRegister& P)
+	{
+		P.Serialize(Ar);
+		return Ar;
+	}
+
+	// returns the current address of the register within the data byte array.
+	// this can change over time - as the register is moving through slices.
+	// use GetFirstAllocatedByte to get the fixed first byte.
+	FORCEINLINE uint64 GetWorkByteIndex() const { return ByteIndex; }
+
+	// returns the first allocated byte in the data byte array
+	FORCEINLINE uint64 GetFirstAllocatedByte() const { return ByteIndex - (uint64)AlignmentBytes - (uint64)(SliceIndex * GetNumBytesPerSlice()); }
+
+	// Returns the leading alignment bytes
+	FORCEINLINE uint8 GetAlignmentBytes() const { return AlignmentBytes; }
+
+	// Returns true if the register stores more than one element
+	FORCEINLINE bool IsArray() const { return ElementCount > 1; }
+
+	// Returns true if the register stores plain / shallow memory
+	FORCEINLINE bool IsPlain() const { return ScriptStructIndex == INDEX_NONE; }
+
+	// Returns the number of allocated bytes (inclusing alignment + trailing bytes)
+	FORCEINLINE uint16 GetAllocatedBytes() const { return ElementCount * ElementSize * SliceCount + (uint16)AlignmentBytes + TrailingBytes; }
+
+	// Returns the number of bytes for a complete slice
+	FORCEINLINE uint16 GetNumBytesPerSlice() const { return ElementCount * ElementSize; }
+
+	// Returns the number of bytes for all slices
+	FORCEINLINE uint16 GetNumBytesAllSlices() const { return ElementCount * ElementSize * SliceCount; }
+
+	// Returns the total number of elements (elementcount * slicecount) in the register
+	FORCEINLINE uint32 GetTotalElementCount() const { return (uint32)ElementCount * (uint32)SliceCount; }
+	
+	// Moves the register address to the first slice
+	FORCEINLINE void MoveToFirstSlice() { ByteIndex -= SliceIndex * GetNumBytesPerSlice(); SliceIndex = 0; }
+
+	// Moves the register address to the next slice. This also wraps around - so when the register reaches 
+	// the end it will move to the first slice again.
+	FORCEINLINE void MoveToNextSlice()
+	{
+		if(SliceCount <= 1)
+		{
+			return;
+		}
+		
+		if(SliceIndex == SliceCount - 1)
+		{
+			MoveToFirstSlice();
+		}
+		else
+		{
+			ByteIndex += GetNumBytesPerSlice();
+			SliceIndex++;
+		}
+	}
+};
+
+typedef TArrayView<FRigVMRegister> FRigVMRegisterArray;
+
+/**
+ * The FRigVMMemoryContainer provides a heterogeneous memory container to store arbitrary
+ * data. Each element stored can be referred to using a FRigVMRegister.
+ * Elements can be accessed by index (index of the register), FRigVMArgument or by name.
+ * Name access is optional and is specified upon construction of the container.
+ * The memory container provides a series of templated functions to add and get data.
+ * Plain types (shallow types without the need for construction) can be added an retrieved
+ * using the methods suffixed with 'Plain'.
+ *
+ * For example:
+ * 		int32 Index = Container.AddPlain<float>(4.f);
+ *      float& ValueRef = Container.GetRef<float>(Index);
+ *
+ * This can also be done with arrays:
+ *      TArray<float> MyArray = {3.f, 4.f, 5.f};
+ * 		int32 Index = Container.AddPlainArray<float>(MyArray);
+ *      TArrayView<float> ArrayView = Container.GetArray<float>(Index);
+ *
+ * More complex data structures such as USTRUCT types need to be store with the
+ * 'Struct' suffixed methods - since the require a constructor / destructor to be called.
+ *
+ * For example:
+ *
+ *		FMyStruct MyStruct;
+ *		int32 Index = Container.AddStruct<FMyStruct>(MyStruct);
+ *      FMyStruct& ValueRef = Container.GetRef<FMyStruct>(Index);
+ */
+USTRUCT()
+struct RIGVM_API FRigVMMemoryContainer
+{
+	GENERATED_BODY()
+	
+public:
+
+	FRigVMMemoryContainer(bool bInUseNames = true);
+	FRigVMMemoryContainer(const FRigVMMemoryContainer& Other);
+	~FRigVMMemoryContainer();
+
+	FRigVMMemoryContainer& operator= (const FRigVMMemoryContainer &InOther);
+
+	// returns the memory type of this container
+	FORCEINLINE ERigVMMemoryType GetMemoryType() const { return MemoryType;  }
+
+	// sets the memory type. should only be used when the container is empty
+	FORCEINLINE void SetMemoryType(ERigVMMemoryType InMemoryType) { MemoryType = InMemoryType; }
+
+	// returns true if this container supports name based lookup
+	FORCEINLINE bool SupportsNames() const { return bUseNameMap;  }
+
+	// returns the number of registers in this container
+	FORCEINLINE int32 Num() const { return Registers.Num(); }
+
+	// resets the container and removes all storage.
+	void Reset();
+
+	// const accessor for a register based on index
+	FORCEINLINE const FRigVMRegister& operator[](int32 InIndex) const { return Registers[InIndex]; }
+
+	// accessor for a register based on index
+	FORCEINLINE FRigVMRegister& operator[](int32 InIndex) { return Registers[InIndex]; }
+
+	// const accessor for a register based on an argument
+	FORCEINLINE const FRigVMRegister& operator[](const FRigVMArgument& InArg) const { return Registers[InArg.GetRegisterIndex()]; }
+
+	// accessor for a register based on an argument
+	FORCEINLINE FRigVMRegister& operator[](const FRigVMArgument& InArg) { return Registers[InArg.GetRegisterIndex()]; }
+
+	// const accessor for a register based on a a name. note: only works if SupportsNames() == true
+	FORCEINLINE const FRigVMRegister& operator[](const FName& InName) const { return Registers[GetIndex(InName)]; }
+
+	// accessor for a register based on a a name. note: only works if SupportsNames() == true
+	FORCEINLINE FRigVMRegister& operator[](const FName& InName) { return Registers[GetIndex(InName)]; }
+
+	FORCEINLINE TArray<FRigVMRegister>::RangedForIteratorType      begin() { return Registers.begin(); }
+	FORCEINLINE TArray<FRigVMRegister>::RangedForConstIteratorType begin() const { return Registers.begin(); }
+	FORCEINLINE TArray<FRigVMRegister>::RangedForIteratorType      end() { return Registers.end(); }
+	FORCEINLINE TArray<FRigVMRegister>::RangedForConstIteratorType end() const { return Registers.end(); }
+
+	bool Serialize(FArchive& Ar);
+	FORCEINLINE friend FArchive& operator<<(FArchive& Ar, FRigVMMemoryContainer& P)
+	{
+		P.Serialize(Ar);
+		return Ar;
+	}
+
+	// Returns an argument for a given register.
+	// This is typically used to store a light weight address for use within a VM.
+	FORCEINLINE FRigVMArgument GetArgument(int32 InRegisterIndex) const
+	{
+		ensure(Registers.IsValidIndex(InRegisterIndex));
+		return FRigVMArgument(MemoryType, InRegisterIndex, Registers[InRegisterIndex].ByteIndex);
+	}
+
+	// Returns the current const data pointer for a given register.
+	// Note: This refers to the active slice - and can change over time.
+	FORCEINLINE const uint8* GetData(const FRigVMRegister& Register) const
+	{
+		ensure(Register.ElementCount > 0);
+		return &Data[Register.GetWorkByteIndex()];
+	}
+
+	// Returns the current const data pointer for a given register index.
+	// Note: This refers to the active slice - and can change over time.
+	FORCEINLINE const uint8* GetData(int32 InRegisterIndex) const
+	{
+		ensure(Registers.IsValidIndex(InRegisterIndex));
+		const FRigVMRegister& Register = Registers[InRegisterIndex];
+		return GetData(Register);
+	}
+
+	// Returns the current data pointer for a given register.
+	// Note: This refers to the active slice - and can change over time.
+	FORCEINLINE uint8* GetData(FRigVMRegister& Register, bool bMoveToNextSlice = false)
+	{
+		ensure(Register.ElementCount > 0);
+		if(bMoveToNextSlice)
+		{
+			Register.MoveToNextSlice();
+		}
+		return &Data[Register.GetWorkByteIndex()];
+	}
+
+	// Returns the current data pointer for a given register index.
+	// Note: This refers to the active slice - and can change over time.
+	FORCEINLINE uint8* GetData(int32 InRegisterIndex, bool bMoveToNextSlice = false)
+	{
+		ensure(Registers.IsValidIndex(InRegisterIndex));
+		FRigVMRegister& Register = Registers[InRegisterIndex];
+		return GetData(Register, bMoveToNextSlice);
+	}
+
+	// Returns the current const typed pointer for a given register.
+	// Note: This refers to the active slice - and can change over time.
+	template<class T>
+	FORCEINLINE const T* Get(const FRigVMRegister& Register) const
+	{
+		ensure(Register.ElementCount > 0);
+		return (const T*)&Data[Register.GetWorkByteIndex()];
+	}
+
+	// Returns the current const typed pointer for a given register index.
+	// Note: This refers to the active slice - and can change over time.
+	template<class T>
+	FORCEINLINE const T* Get(int32 InRegisterIndex) const
+	{
+		ensure(Registers.IsValidIndex(InRegisterIndex));
+		const FRigVMRegister& Register = Registers[InRegisterIndex];
+		return Get<T>(Register);
+	}
+
+	// Returns the current const typed pointer for a given argument
+	// Note: This refers to the active slice - and can change over time.
+	template<class T>
+	FORCEINLINE const T* Get(const FRigVMArgument& InArgument) const
+	{
+		return Get<T>(InArgument.GetRegisterIndex());
+	}
+
+	// Returns the current const typed reference for a given register
+	// Note: This refers to the active slice - and can change over time.
+	template<class T>
+	FORCEINLINE const T& GetRef(const FRigVMRegister& Register) const
+	{
+		return *Get<T>(Register);
+	}
+
+	// Returns the current const typed reference for a given register index
+	// Note: This refers to the active slice - and can change over time.
+	template<class T>
+	FORCEINLINE const T& GetRef(int32 InRegisterIndex) const
+	{
+		return *Get<T>(InRegisterIndex);
+	}
+
+	// Returns the current const typed reference for a given argument
+	// Note: This refers to the active slice - and can change over time.
+	template<class T>
+	FORCEINLINE const T& GetRef(const FRigVMArgument& InArgument) const
+	{
+		return GetRef<T>(InArgument.GetRegisterIndex());
+	}
+
+	// Returns the current typed pointer for a given register
+	// Note: This refers to the active slice - and can change over time.
+	template<class T>
+	FORCEINLINE T* Get(FRigVMRegister& Register, bool bMoveToNextSlice = false)
+	{
+		ensure(Register.ElementCount > 0);
+		if(bMoveToNextSlice)
+		{
+			Register.MoveToNextSlice();
+		}
+		return (T*)&Data[Register.GetWorkByteIndex()];
+	}
+
+	// Returns the current typed pointer for a given register index
+	// Note: This refers to the active slice - and can change over time.
+	template<class T>
+	FORCEINLINE T* Get(int32 InRegisterIndex, bool bMoveToNextSlice = false)
+	{
+		ensure(Registers.IsValidIndex(InRegisterIndex));
+		FRigVMRegister& Register = Registers[InRegisterIndex];
+		return Get<T>(Register, bMoveToNextSlice);
+	}
+
+	// Returns the current typed pointer for a given argument
+	// Note: This refers to the active slice - and can change over time.
+	template<class T>
+	FORCEINLINE T* Get(const FRigVMArgument& InArgument, bool bMoveToNextSlice = false)
+	{
+		return Get<T>(InArgument.GetRegisterIndex(), bMoveToNextSlice);
+	}
+
+	// Returns the current typed reference for a given register
+	// Note: This refers to the active slice - and can change over time.
+	template<class T>
+	FORCEINLINE T& GetRef(FRigVMRegister& Register, bool bMoveToNextSlice = false)
+	{
+		return *Get<T>(Register, bMoveToNextSlice);
+	}
+
+	// Returns the current typed reference for a given register index
+	// Note: This refers to the active slice - and can change over time.
+	template<class T>
+	FORCEINLINE T& GetRef(int32 InRegisterIndex, bool bMoveToNextSlice = false)
+	{
+		return *Get<T>(InRegisterIndex, bMoveToNextSlice);
+	}
+
+	// Returns the current typed reference for a given argument
+	// Note: This refers to the active slice - and can change over time.
+	template<class T>
+	FORCEINLINE T& GetRef(const FRigVMArgument& InArgument, bool bMoveToNextSlice = false)
+	{
+		return GetRef<T>(InArgument.GetRegisterIndex(), bMoveToNextSlice);
+	}
+
+	// Returns an array view for all elements of the current slice for a given register.
+	template<class T>
+	FORCEINLINE TArrayView<T> GetArray(FRigVMRegister& Register, bool bMoveToNextSlice = false)
+	{
+		ensure(Register.ElementCount > 0);
+		if(bMoveToNextSlice)
+		{
+			Register.MoveToNextSlice();
+		}
+		return TArrayView<T>((T*)&Data[Register.GetWorkByteIndex()], Register.ElementCount);
+	}
+
+	// Returns an array view for all elements of the current slice for a given register index.
+	template<class T>
+	FORCEINLINE TArrayView<T> GetArray(int32 InRegisterIndex, bool bMoveToNextSlice = false)
+	{
+		ensure(Registers.IsValidIndex(InRegisterIndex));
+		FRigVMRegister& Register = Registers[InRegisterIndex];
+		return GetArray<T>(Register, bMoveToNextSlice);
+	}
+	
+	// Returns an array view for all elements of the current slice for a given argument.
+	template<class T>
+	FORCEINLINE TArrayView<T> GetArray(const FRigVMArgument& InArgument, bool bMoveToNextSlice = false)
+	{
+		return GetArray<T>(InArgument.GetRegisterIndex(), bMoveToNextSlice);
+	}
+
+	// Returns the script struct used for a given register (can be nullptr for non-struct-registers).
+	FORCEINLINE UScriptStruct* GetScriptStruct(const FRigVMRegister& Register) const
+	{
+		if (Register.ScriptStructIndex != INDEX_NONE)
+		{
+			ensure(ScriptStructs.IsValidIndex(Register.ScriptStructIndex));
+			return ScriptStructs[Register.ScriptStructIndex];
+		}
+		return nullptr;
+	}
+
+	// Returns the script struct used for a given register index (can be nullptr for non-struct-registers).
+	FORCEINLINE UScriptStruct* GetScriptStruct(int32 InRegisterIndex) const
+	{
+		ensure(Registers.IsValidIndex(InRegisterIndex));
+		const FRigVMRegister& Register = Registers[InRegisterIndex];
+		return GetScriptStruct(Register);
+	}
+
+	// Copies the content of a source register to a target register.
+	// The source register can optionally be referencing a specific source memory container.
+	bool Copy(
+		int32 InSourceRegisterIndex,
+		int32 InTargetRegisterIndex,
+		const FRigVMMemoryContainer* InSourceMemory = nullptr,
+		int32 InSourceByteOffset = INDEX_NONE,
+		int32 InTargetByteOffset = INDEX_NONE,
+		int32 InNumBytes = INDEX_NONE);
+
+	// Copies the content of a source register to a target register.
+	// The source register can optionally be referencing a specific source memory container.
+	// Note: This only works if SupportsNames() == true
+	bool Copy(
+		const FName& InSourceName,
+		const FName& InTargetName,
+		const FRigVMMemoryContainer* InSourceMemory = nullptr,
+		int32 InSourceByteOffset = INDEX_NONE,
+		int32 InTargetByteOffset = INDEX_NONE,
+		int32 InNumBytes = INDEX_NONE);
+
+	// Returns the index of a register based on the register name.
+	// Note: This only works if SupportsNames() == true
+	FORCEINLINE int32 GetIndex(const FName& InName) const
+	{
+		if (!bUseNameMap)
+		{
+			return INDEX_NONE;
+		}
+
+		if (NameMap.Num() != Registers.Num())
+		{
+			for (int32 Index = 0; Index < Registers.Num(); Index++)
+			{
+				if (Registers[Index].Name == InName)
+				{
+					return Index;
+				}
+			}
+		}
+		else
+		{
+			const int32* Index = NameMap.Find(InName);
+			if (Index != nullptr)
+			{
+				return *Index;
+			}
+		}
+
+		return INDEX_NONE;
+	}
+
+	// Returns true if a given name is available for a new register.
+	FORCEINLINE bool IsNameAvailable(const FName& InPotentialNewName) const
+	{
+		if (!bUseNameMap)
+		{
+			return false;
+		}
+		return GetIndex(InPotentialNewName) == INDEX_NONE;
+	}
+
+	// Adds a new named register for a plain array from a data pointer.
+	FORCEINLINE int32 AddPlainArray(const FName& InNewName, int32 InElementSize, int32 InCount, const void* InDataPtr = nullptr, int32 InSliceCount = 1)
+	{
+		return Allocate(InNewName, InElementSize, InCount, InSliceCount, InDataPtr);
+	}
+
+	// Adds a new named register for a typed plain array from a data pointer.
+	template<class T>
+	FORCEINLINE int32 AddPlainArray(const FName& InNewName, int32 InCount, const T* InDataPtr = nullptr, int32 InSliceCount = 1)
+	{
+		return Allocate(InNewName, sizeof(T), InCount, InSliceCount, (const void*)InDataPtr);
+	}
+
+	// Adds a new named register for a typed plain array from a TArray.
+	template<class T>
+	FORCEINLINE int32 AddPlainArray(const FName& InNewName, const TArray<T>& InArray, int32 InSliceCount = 1)
+	{
+		return AddPlainArray<T>(InNewName, InArray.Num(), InArray.GetData(), InSliceCount);
+	}
+
+	// Adds a new unnamed register for a typed plain array from a TArray.
+	template<class T>
+	FORCEINLINE int32 AddPlainArray(const TArray<T>& InArray, int32 InSliceCount = 1)
+	{
+		return AddPlainArray<T>(NAME_None, InArray, InSliceCount);
+	}
+
+	// Adds a new named register for a plain value from a data pointer.
+	FORCEINLINE int32 AddPlain(const FName& InNewName, int32 InElementSize, const void* InValuePtr, int32 InSliceCount = 1)
+	{
+		return AddPlainArray(InNewName, InElementSize, 1, InValuePtr, InSliceCount);
+	}
+
+	// Adds a new unnamed register for a  plain value from a data pointer.
+	FORCEINLINE int32 AddPlain(int32 InElementSize, const void* InValuePtr, int32 InSliceCount = 1)
+	{
+		return AddPlain(NAME_None, InElementSize, InValuePtr, InSliceCount);
+	}
+
+	// Adds a new named register for a typed plain value from a value reference.
+	template<class T>
+	FORCEINLINE int32 AddPlain(const FName& InNewName, const T& InValue, int32 InSliceCount = 1)
+	{
+		return AddPlainArray<T>(InNewName, 1, &InValue, InSliceCount);
+	}
+
+	// Adds a new unnamed register for a typed plain value from a value reference.
+	template<class T>
+	FORCEINLINE int32 AddPlain(const T& InValue, int32 InSliceCount = 1)
+	{
+		return AddPlain<T>(NAME_None, InValue, InSliceCount);
+	}
+
+	// Adds a new named register for a struct given a scriptstruct and a data pointer.
+	FORCEINLINE int32 AddStructArray(const FName& InNewName, UScriptStruct* InScriptStruct, int32 InCount, const void* InDataPtr = nullptr, int32 InSliceCount = 1)
+	{
+		int32 Register = Allocate(InNewName, InScriptStruct->GetStructureSize(), InCount, InSliceCount, nullptr, false);
+		if (Register == INDEX_NONE)
+		{
+			return INDEX_NONE;
+		}
+
+		Registers[Register].Type = ERigVMRegisterType::Struct;
+		Registers[Register].ScriptStructIndex = FindOrAddScriptStruct(InScriptStruct);
+
+		UpdateRegisters();
+
+		// construct the content
+		Construct(Register);
+
+		// copy values from the provided data
+		if (InDataPtr != nullptr)
+		{
+			Registers[Register].MoveToFirstSlice();
+			for (uint16 SliceIndex = 0; SliceIndex < Registers[Register].SliceCount; SliceIndex++)
+			{
+				InScriptStruct->CopyScriptStruct(GetData(Register), InDataPtr, InCount);
+				Registers[Register].MoveToNextSlice();
+			}
+			Registers[Register].MoveToFirstSlice();
+		}
+
+		return Register;
+	}
+
+	// Adds a new unnamed register for a struct given a scriptstruct and a data pointer.
+	FORCEINLINE int32 AddStructArray(UScriptStruct* InScriptStruct, int32 InCount, const void* InDataPtr = nullptr, int32 InSliceCount = 1)
+	{
+		return AddStructArray(NAME_None, InScriptStruct, InCount, InDataPtr, InSliceCount);
+	}
+
+	// Adds a new named register for a typed struct given a data pointer.
+	template<class T>
+	FORCEINLINE int32 AddStructArray(const FName& InNewName, int32 InCount, const T* InDataPtr = nullptr, int32 InSliceCount = 1)
+	{
+		// if you are hitting this - you might need to use AddPlainArray instead!
+		UScriptStruct* Struct = T::StaticStruct();
+		if (Struct == nullptr)
+		{
+			return INDEX_NONE;
+		}
+	
+		return AddStructArray(InNewName, Struct, InCount, InDataPtr, InSliceCount);
+	}
+
+	// Adds a new named register for a typed struct array given a value reference.
+	template<class T>
+	FORCEINLINE int32 AddStructArray(const FName& InNewName, const TArray<T>& InArray, int32 InSliceCount = 1)
+	{
+		return AddStructArray<T>(InNewName, InArray.Num(), InArray.GetData(), InSliceCount);
+	}
+
+	// Adds a new named register for a typed struct array given a TArray
+	template<class T>
+	FORCEINLINE int32 AddStructArray(const TArray<T>& InArray, int32 InSliceCount = 1)
+	{
+		return AddStructArray<T>(NAME_None, InArray, InSliceCount);
+	}
+
+	// Adds a new named register for a struct given a scriptstruct and a value ptr
+	FORCEINLINE int32 AddStruct(const FName& InNewName, UScriptStruct* InScriptStruct, const void* InValuePtr, int32 InSliceCount = 1)
+	{
+		return AddStructArray(InNewName, InScriptStruct, 1, InValuePtr, InSliceCount);
+	}
+
+	// Adds a new unnamed register for a struct given a scriptstruct and value ptr
+	FORCEINLINE int32 AddStruct(UScriptStruct* InScriptStruct, const void* InValuePtr, int32 InSliceCount = 1)
+	{
+		return AddStruct(NAME_None, InScriptStruct, InValuePtr, InSliceCount);
+	}
+
+	// Adds a new named register for a typed struct given a value reference
+	template<class T>
+	FORCEINLINE int32 AddStruct(const FName& InNewName, const T& InValue, int32 InSliceCount = 1)
+	{
+		return AddStructArray<T>(InNewName, 1, &InValue, InSliceCount);
+	}
+
+	// Adds a new unnamed register for a typed struct given a value reference
+	template<class T>
+	FORCEINLINE int32 AddStruct(const T& InValue, int32 InSliceCount = 1)
+	{
+		return AddStruct<T>(NAME_None, InValue, InSliceCount);
+	}
+
+	// Remove a register given its index
+	bool Remove(int32 InRegisterIndex);
+
+	// Remove a register given its name
+	// Note: This only works if SupportsNames() == true
+	bool Remove(const FName& InRegisterName);
+
+	// Remove a register given its index
+	// Note: This only works if SupportsNames() == true
+	FName Rename(int32 InRegisterIndex, const FName& InNewName);
+
+	// Remove a register given its old name
+	// Note: This only works if SupportsNames() == true
+	FName Rename(const FName& InOldName, const FName& InNewName);
+
+	// Resize a register given its index
+	bool Resize(int32 InRegisterIndex, int32 InNewElementCount, int32 InNewSliceCount = 1);
+
+	// Resize a register given its name
+	// Note: This only works if SupportsNames() == true
+	bool Resize(const FName& InRegisterName, int32 InNewElementCount, int32 InNewSliceCount = 1);
+
+	// Changes the type of a register provided a new type and a data pointer.
+	// This can be used to store different content a register for re-using the memory,
+	// but this only works for new elements smaller or of the same size (since the register can't grow).
+	bool ChangeRegisterType(int32 InRegisterIndex, ERigVMRegisterType InNewType, int32 InElementSize, const void* InDataPtr = nullptr, int32 InNewElementCount = 1, int32 InNewSliceCount = 1);
+
+	// Changes the type of a register provided a new type and a data pointer.
+	// This can be used to store different content a register for re-using the memory,
+	// but this only works for new elements smaller or of the same size (since the register can't grow).
+	template<class T>
+	FORCEINLINE bool ChangeRegisterType(int32 InRegisterIndex, const T* InDataPtr = nullptr, int32 InNewElementCount = 1, int32 InNewSliceCount = 1)
+	{
+		return ChangeRegisterType(InRegisterIndex, ERigVMRegisterType::Plain, sizeof(T), InDataPtr, InNewElementCount, InNewSliceCount);
+	}
+
+private:
+
+	// Updates internal data for topological changes
+	void UpdateRegisters();
+
+	// Allocates a new named register
+	int32 Allocate(const FName& InNewName, int32 InElementSize, int32 InElementCount, int32 InSliceCount, const void* InDataPtr = nullptr, bool bUpdateRegisters = true);
+
+	// Allocates a new unnamed register
+	int32 Allocate(int32 InElementSize, int32 InElementCount, int32 InSliceCount, const void* InDataPtr = nullptr, bool bUpdateRegisters = true);
+
+	// Performs optional construction of data within a struct register
+	bool Construct(int32 InRegisterIndex, int32 InElementIndex = INDEX_NONE);
+
+	// Performs optional destruction of data within a struct register
+	bool Destroy(int32 InRegisterIndex, int32 InElementIndex = INDEX_NONE);
+
+	// Fills a register with zero memory
+	void FillWithZeroes(int32 InRegisterIndex);
+
+	// Ensures to add a script struct to the internal map if needed
+	int32 FindOrAddScriptStruct(UScriptStruct* InScriptStruct);
+
+	UPROPERTY()
+	bool bUseNameMap;
+
+	UPROPERTY()
+	ERigVMMemoryType MemoryType;
+
+	UPROPERTY()
+	TArray<FRigVMRegister> Registers;
+
+	UPROPERTY(transient)
+	TArray<uint8> Data;
+
+	UPROPERTY()
+	TArray<UScriptStruct*> ScriptStructs;
+
+	UPROPERTY(transient)
+	TMap<FName, int32> NameMap;
+};
+
+template<>
+FORCEINLINE bool FRigVMMemoryContainer::ChangeRegisterType(int32 InRegisterIndex, const FName* InDataPtr, int32 InNewElementCount, int32 InNewSliceCount)
+{
+	if (!ChangeRegisterType(InRegisterIndex, ERigVMRegisterType::Name, sizeof(FName), nullptr, InNewElementCount, InNewSliceCount))
+	{
+		return false;
+	}
+
+	if (InDataPtr)
+	{
+		Registers[InRegisterIndex].MoveToFirstSlice();
+		for (uint16 SliceIndex = 0; SliceIndex < Registers[InRegisterIndex].SliceCount; SliceIndex++)
+		{
+			FName* DataPtr = (FName*)GetData(InRegisterIndex);
+			for (int32 Index = 0; Index < InNewElementCount; Index++)
+			{
+				DataPtr[Index] = InDataPtr[Index];
+			}
+			Registers[InRegisterIndex].MoveToNextSlice();
+		}
+		Registers[InRegisterIndex].MoveToFirstSlice();
+	}
+
+	return true;
+}
+
+template<>
+FORCEINLINE bool FRigVMMemoryContainer::ChangeRegisterType(int32 InRegisterIndex, const FString* InDataPtr, int32 InNewElementCount, int32 InNewSliceCount)
+{
+	if (!ChangeRegisterType(InRegisterIndex, ERigVMRegisterType::String, sizeof(FString), nullptr, InNewElementCount, InNewSliceCount))
+	{
+		return false;
+	}
+
+	if (InDataPtr)
+	{
+		Registers[InRegisterIndex].MoveToFirstSlice();
+		for (uint16 SliceIndex = 0; SliceIndex < Registers[InRegisterIndex].SliceCount; SliceIndex++)
+		{
+			FString* DataPtr = (FString*)GetData(InRegisterIndex);
+			for (int32 Index = 0; Index < InNewElementCount; Index++)
+			{
+				DataPtr[Index] = InDataPtr[Index];
+			}
+			Registers[InRegisterIndex].MoveToNextSlice();
+		}
+		Registers[InRegisterIndex].MoveToFirstSlice();
+	}
+
+	return true;
+}
+
+template<>
+FORCEINLINE int32 FRigVMMemoryContainer::AddPlainArray(const FName& InNewName, int32 InCount, const FName* InDataPtr, int32 InSliceCount)
+{
+	int32 Register = Allocate(InNewName, sizeof(FName), InCount, InSliceCount, nullptr);
+	Registers[Register].Type = ERigVMRegisterType::Name;
+
+	Construct(Register);
+
+	if (InDataPtr)
+	{
+		Registers[Register].MoveToFirstSlice();
+		for (uint16 SliceIndex = 0; SliceIndex < Registers[Register].SliceCount; SliceIndex++)
+		{
+			FName* DataPtr = (FName*)GetData(Register);
+			for (int32 Index = 0; Index < InCount; Index++)
+			{
+				DataPtr[Index] = InDataPtr[Index];
+			}
+			Registers[Register].MoveToNextSlice();
+		}
+		Registers[Register].MoveToFirstSlice();
+	}
+
+	return Register;
+}
+
+template<>
+FORCEINLINE int32 FRigVMMemoryContainer::AddPlainArray(const FName& InNewName, int32 InCount, const FString* InDataPtr, int32 InSliceCount)
+{
+	int32 Register = Allocate(InNewName, sizeof(FString), InCount, InSliceCount, nullptr);
+	Registers[Register].Type = ERigVMRegisterType::String;
+
+	Construct(Register);
+
+	if (InDataPtr)
+	{
+		Registers[Register].MoveToFirstSlice();
+		for (uint16 SliceIndex = 0; SliceIndex < Registers[Register].SliceCount; SliceIndex++)
+		{
+			FString* DataPtr = (FString*)GetData(Register);
+			for (int32 Index = 0; Index < InCount; Index++)
+			{
+				DataPtr[Index] = InDataPtr[Index];
+			}
+			Registers[Register].MoveToNextSlice();
+		}
+		Registers[Register].MoveToFirstSlice();
+	}
+
+	return Register;
+}
+
+typedef FRigVMMemoryContainer* FRigVMMemoryContainerPtr;
+typedef TArrayView<FRigVMMemoryContainerPtr> FRigVMMemoryContainerPtrArray;
\ No newline at end of file
diff --git a/Engine/Source/Runtime/RigVM/Public/RigVMCore/RigVMRegistry.h b/Engine/Source/Runtime/RigVM/Public/RigVMCore/RigVMRegistry.h
--- a/Engine/Source/Runtime/RigVM/Public/RigVMCore/RigVMRegistry.h
+++ b/Engine/Source/Runtime/RigVM/Public/RigVMCore/RigVMRegistry.h
@@ -3,10 +3,10 @@
 #pragma once
 
 #include "CoreMinimal.h"
-#include "RigVMStorage.h"
+#include "RigVMMemory.h"
 
 typedef TArrayView<void*> FRigVMUserDataArray;
-typedef void (*FRigVMFunctionPtr)(const FRigVMArgumentArray&, FRigVMStoragePtrArray&, const FRigVMUserDataArray&);
+typedef void (*FRigVMFunctionPtr)(const FRigVMArgumentArray&, FRigVMMemoryContainerPtrArray&, const FRigVMUserDataArray&);
 
 struct RIGVM_API FRigVMFunction
 {
@@ -45,7 +45,7 @@ private:
 	// disable assignment operator
 	FRigVMRegistry& operator= (const FRigVMRegistry &InOther) { return *this; }
 
-	// storage for all functions
+	// memory for all functions
 	TArray<FRigVMFunction> Functions;
 
 	static FRigVMRegistry s_RigVMRegistry;
diff --git a/Engine/Source/Runtime/RigVM/Public/RigVMCore/RigVMStorage.h b/Engine/Source/Runtime/RigVM/Public/RigVMCore/RigVMStorage.h
deleted file mode 100644
