RigVM: Renamed a couple of structures for clarity + documented all headers.

diff --git a/Engine/Source/Runtime/RigVM/Private/RigVMCore/RigVM.cpp b/Engine/Source/Runtime/RigVM/Private/RigVMCore/RigVM.cpp
--- a/Engine/Source/Runtime/RigVM/Private/RigVMCore/RigVM.cpp
+++ b/Engine/Source/Runtime/RigVM/Private/RigVMCore/RigVM.cpp
@@ -23,10 +23,10 @@ void URigVM::Reset()
 	Instructions.Reset();
 }
 
-int32 URigVM::AddRigVMFunction(UScriptStruct* InRigVMStruct, const FName& InFunctionName)
+int32 URigVM::AddRigVMFunction(UScriptStruct* InRigVMStruct, const FName& InMethodName)
 {
 	check(InRigVMStruct);
-	FString FunctionKey = FString::Printf(TEXT("F%s::%s"), *InRigVMStruct->GetName(), *InFunctionName.ToString());
+	FString FunctionKey = FString::Printf(TEXT("F%s::%s"), *InRigVMStruct->GetName(), *InMethodName.ToString());
 	int32 FunctionIndex = FunctionNames.Find(FunctionKey);
 	if (FunctionIndex != INDEX_NONE)
 	{
@@ -43,6 +43,12 @@ int32 URigVM::AddRigVMFunction(UScriptStruct* InRigVMStruct, const FName& InFunc
 	return Functions.Add(Function);
 }
 
+const FRigVMInstructionArray& URigVM::GetInstructions()
+{
+	RefreshInstructionsIfRequired();
+	return Instructions;
+}
+
 void URigVM::ResolveFunctionsIfRequired()
 {
 	if (Functions.Num() != FunctionNames.Num())
@@ -61,7 +67,7 @@ void URigVM::RefreshInstructionsIfRequired()
 {
 	if (Instructions.Num() == 0)
 	{
-		Instructions = ByteCode.GetTable();
+		Instructions = ByteCode.GetInstructions();
 	}
 }
 
diff --git a/Engine/Source/Runtime/RigVM/Private/RigVMCore/RigVMByteCode.cpp b/Engine/Source/Runtime/RigVM/Private/RigVMCore/RigVMByteCode.cpp
--- a/Engine/Source/Runtime/RigVM/Private/RigVMCore/RigVMByteCode.cpp
+++ b/Engine/Source/Runtime/RigVM/Private/RigVMCore/RigVMByteCode.cpp
@@ -2,24 +2,24 @@
 
 #include "RigVMCore/RigVMByteCode.h"
 
-FRigVMByteCodeTable::FRigVMByteCodeTable()
+FRigVMInstructionArray::FRigVMInstructionArray()
 {
 }
 
-FRigVMByteCodeTable::FRigVMByteCodeTable(const FRigVMByteCode& InByteCode)
+FRigVMInstructionArray::FRigVMInstructionArray(const FRigVMByteCode& InByteCode)
 {
 	uint64 ByteIndex = 0;
 	while (ByteIndex < InByteCode.Num())
 	{
 		ERigVMOpCode OpCode = InByteCode.GetOpCodeAt(ByteIndex);
-		Entries.Add(FRigVMByteCodeTableEntry(OpCode, ByteIndex));
+		Instructions.Add(FRigVMInstruction(OpCode, ByteIndex));
 		ByteIndex += InByteCode.GetOpNumBytesAt(ByteIndex);
 	}
 }
 
-void FRigVMByteCodeTable::Reset()
+void FRigVMInstructionArray::Reset()
 {
-	Entries.Reset();
+	Instructions.Reset();
 }
 
 FRigVMByteCode::FRigVMByteCode()
@@ -230,9 +230,9 @@ uint64 FRigVMByteCode::AddJumpOp(ERigVMOpCode InOpCode, uint16 InInstructionInde
 	return AddOp(Op);
 }
 
-uint64 FRigVMByteCode::AddJumpIfOp(ERigVMOpCode InOpCode, uint16 InInstructionIndex, const FRigVMArgument& InConditionArg, bool bInCondition)
+uint64 FRigVMByteCode::AddJumpIfOp(ERigVMOpCode InOpCode, uint16 InInstructionIndex, const FRigVMArgument& InConditionArg, bool bJumpWhenConditionIs)
 {
-	FRigVMJumpIfOp Op(InOpCode, InInstructionIndex, InConditionArg, bInCondition);
+	FRigVMJumpIfOp Op(InOpCode, InInstructionIndex, InConditionArg, bJumpWhenConditionIs);
 	return AddOp(Op);
 }
 
diff --git a/Engine/Source/Runtime/RigVM/Public/RigVMCore/RigVM.h b/Engine/Source/Runtime/RigVM/Public/RigVMCore/RigVM.h
--- a/Engine/Source/Runtime/RigVM/Public/RigVMCore/RigVM.h
+++ b/Engine/Source/Runtime/RigVM/Public/RigVMCore/RigVM.h
@@ -8,6 +8,11 @@
 #include "RigVMByteCode.h"
 #include "RigVM.generated.h"
 
+/**
+ * The RigVM is the main object for evaluating FRigVMByteCode instructions.
+ * It combines the byte code, a list of required function pointers for 
+ * execute instructions and required memory in one class.
+ */
 UCLASS()
 class RIGVM_API URigVM : public UObject
 {
@@ -18,30 +23,42 @@ public:
 	URigVM();
 	virtual ~URigVM();
 
+	// resets the container and removes all memory
 	void Reset();
 
+	// Executes the VM.
+	// You can optionally provide external memory to the execution
+	// and provide optional additional arguments.
 	bool Execute(FRigVMMemoryContainerPtrArray Memory = FRigVMMemoryContainerPtrArray(), TArrayView<void*> AdditionalArgs = TArrayView<void*>());
 
+	// Add a function for execute instructions to this VM.
+	// Execute instructions can then refer to the function by index.
 	UFUNCTION()
-	int32 AddRigVMFunction(UScriptStruct* InRigVMStruct, const FName& InFunctionName);
+	int32 AddRigVMFunction(UScriptStruct* InRigVMStruct, const FName& InMethodName);
 
+	// The default mutable work memory
 	UPROPERTY()
 	FRigVMMemoryContainer WorkMemory;
 
+	// The default const literal memory
 	UPROPERTY()
 	FRigVMMemoryContainer LiteralMemory;
 
+	// The byte code of the VM
 	UPROPERTY()
 	FRigVMByteCode ByteCode;
 
-	UPROPERTY()
-	FRigVMByteCodeTable Instructions;
-
+	// Returns the instructions of the VM
+	const FRigVMInstructionArray& GetInstructions();
+	
 private:
 
 	void ResolveFunctionsIfRequired();
 	void RefreshInstructionsIfRequired();
 
+	UPROPERTY(transient)
+	FRigVMInstructionArray Instructions;
+
 	UPROPERTY()
 	TArray<FString> FunctionNames;
 
diff --git a/Engine/Source/Runtime/RigVM/Public/RigVMCore/RigVMByteCode.h b/Engine/Source/Runtime/RigVM/Public/RigVMCore/RigVMByteCode.h
--- a/Engine/Source/Runtime/RigVM/Public/RigVMCore/RigVMByteCode.h
+++ b/Engine/Source/Runtime/RigVM/Public/RigVMCore/RigVMByteCode.h
@@ -8,93 +8,95 @@
 
 struct FRigVMByteCode;
 
+// The code for a single operation within the RigVM
 UENUM()
 enum class ERigVMOpCode : uint8
 {
-	Execute_0_Args,
-	Execute_1_Args,
-	Execute_2_Args,
-	Execute_3_Args,
-	Execute_4_Args,
-	Execute_5_Args,
-	Execute_6_Args,
-	Execute_7_Args,
-	Execute_8_Args,
-	Execute_9_Args,
-	Execute_10_Args,
-	Execute_11_Args,
-	Execute_12_Args,
-	Execute_13_Args,
-	Execute_14_Args,
-	Execute_15_Args,
-	Execute_16_Args,
-	Execute_17_Args,
-	Execute_18_Args,
-	Execute_19_Args,
-	Execute_20_Args,
-	Execute_21_Args,
-	Execute_22_Args,
-	Execute_23_Args,
-	Execute_24_Args,
-	Execute_25_Args,
-	Execute_26_Args,
-	Execute_27_Args,
-	Execute_28_Args,
-	Execute_29_Args,
-	Execute_30_Args,
-	Execute_31_Args,
-	Execute_32_Args,
-	Execute_33_Args,
-	Execute_34_Args,
-	Execute_35_Args,
-	Execute_36_Args,
-	Execute_37_Args,
-	Execute_38_Args,
-	Execute_39_Args,
-	Execute_40_Args,
-	Execute_41_Args,
-	Execute_42_Args,
-	Execute_43_Args,
-	Execute_44_Args,
-	Execute_45_Args,
-	Execute_46_Args,
-	Execute_47_Args,
-	Execute_48_Args,
-	Execute_49_Args,
-	Execute_50_Args,
-	Execute_51_Args,
-	Execute_52_Args,
-	Execute_53_Args,
-	Execute_54_Args,
-	Execute_55_Args,
-	Execute_56_Args,
-	Execute_57_Args,
-	Execute_58_Args,
-	Execute_59_Args,
-	Execute_60_Args,
-	Execute_61_Args,
-	Execute_62_Args,
-	Execute_63_Args,
-	Execute_64_Args,
-	Zero,
-	BoolFalse,
-	BoolTrue,
-	Copy,
-	Increment,
-	Decrement,
-	Equals,
-	NotEquals,
-	JumpAbsolute,
-	JumpForward,
-	JumpBackward,
-	JumpAbsoluteIf,
-	JumpForwardIf,
-	JumpBackwardIf,
-	ChangeType,
-	Exit,
+	Execute_0_Args, // execute a rig function with 0 arguments
+	Execute_1_Args, // execute a rig function with 1 arguments
+	Execute_2_Args, // execute a rig function with 2 arguments
+	Execute_3_Args, // execute a rig function with 3 arguments
+	Execute_4_Args, // execute a rig function with 4 arguments
+	Execute_5_Args, // execute a rig function with 5 arguments
+	Execute_6_Args, // execute a rig function with 6 arguments
+	Execute_7_Args, // execute a rig function with 7 arguments
+	Execute_8_Args, // execute a rig function with 8 arguments
+	Execute_9_Args, // execute a rig function with 9 arguments
+	Execute_10_Args, // execute a rig function with 10 arguments
+	Execute_11_Args, // execute a rig function with 11 arguments
+	Execute_12_Args, // execute a rig function with 12 arguments
+	Execute_13_Args, // execute a rig function with 13 arguments
+	Execute_14_Args, // execute a rig function with 14 arguments
+	Execute_15_Args, // execute a rig function with 15 arguments
+	Execute_16_Args, // execute a rig function with 16 arguments
+	Execute_17_Args, // execute a rig function with 17 arguments
+	Execute_18_Args, // execute a rig function with 18 arguments
+	Execute_19_Args, // execute a rig function with 19 arguments
+	Execute_20_Args, // execute a rig function with 20 arguments
+	Execute_21_Args, // execute a rig function with 21 arguments
+	Execute_22_Args, // execute a rig function with 22 arguments
+	Execute_23_Args, // execute a rig function with 23 arguments
+	Execute_24_Args, // execute a rig function with 24 arguments
+	Execute_25_Args, // execute a rig function with 25 arguments
+	Execute_26_Args, // execute a rig function with 26 arguments
+	Execute_27_Args, // execute a rig function with 27 arguments
+	Execute_28_Args, // execute a rig function with 28 arguments
+	Execute_29_Args, // execute a rig function with 29 arguments
+	Execute_30_Args, // execute a rig function with 30 arguments
+	Execute_31_Args, // execute a rig function with 31 arguments
+	Execute_32_Args, // execute a rig function with 32 arguments
+	Execute_33_Args, // execute a rig function with 33 arguments
+	Execute_34_Args, // execute a rig function with 34 arguments
+	Execute_35_Args, // execute a rig function with 35 arguments
+	Execute_36_Args, // execute a rig function with 36 arguments
+	Execute_37_Args, // execute a rig function with 37 arguments
+	Execute_38_Args, // execute a rig function with 38 arguments
+	Execute_39_Args, // execute a rig function with 39 arguments
+	Execute_40_Args, // execute a rig function with 40 arguments
+	Execute_41_Args, // execute a rig function with 41 arguments
+	Execute_42_Args, // execute a rig function with 42 arguments
+	Execute_43_Args, // execute a rig function with 43 arguments
+	Execute_44_Args, // execute a rig function with 44 arguments
+	Execute_45_Args, // execute a rig function with 45 arguments
+	Execute_46_Args, // execute a rig function with 46 arguments
+	Execute_47_Args, // execute a rig function with 47 arguments
+	Execute_48_Args, // execute a rig function with 48 arguments
+	Execute_49_Args, // execute a rig function with 49 arguments
+	Execute_50_Args, // execute a rig function with 50 arguments
+	Execute_51_Args, // execute a rig function with 51 arguments
+	Execute_52_Args, // execute a rig function with 52 arguments
+	Execute_53_Args, // execute a rig function with 53 arguments
+	Execute_54_Args, // execute a rig function with 54 arguments
+	Execute_55_Args, // execute a rig function with 55 arguments
+	Execute_56_Args, // execute a rig function with 56 arguments
+	Execute_57_Args, // execute a rig function with 57 arguments
+	Execute_58_Args, // execute a rig function with 58 arguments
+	Execute_59_Args, // execute a rig function with 59 arguments
+	Execute_60_Args, // execute a rig function with 60 arguments
+	Execute_61_Args, // execute a rig function with 61 arguments
+	Execute_62_Args, // execute a rig function with 62 arguments
+	Execute_63_Args, // execute a rig function with 63 arguments
+	Execute_64_Args, // execute a rig function with 64 arguments
+	Zero, // zero the memory of a given register
+	BoolFalse, // set a given register to false
+	BoolTrue, // set a given register to true
+	Copy, // copy the content of one register to another
+	Increment, // increment a int32 register
+	Decrement, // decrement a int32 register
+	Equals, // fill a bool register with the result of (A == B)
+	NotEquals, // fill a bool register with the result of (A != B)
+	JumpAbsolute, // jump to an absolute instruction index
+	JumpForward, // jump forwards given a relative instruction index offset
+	JumpBackward, // jump backwards given a relative instruction index offset
+	JumpAbsoluteIf, // jump to an absolute instruction index based on a condition register
+	JumpForwardIf, // jump forwards given a relative instruction index offset based on a condition register
+	JumpBackwardIf, // jump backwards given a relative instruction index offset based on a condition register
+	ChangeType, // change the type of a register
+	Exit, // exit the execution loop
 	Invalid
 };
 
+// Base class for all VM operations
 struct RIGVM_API FRigVMBaseOp
 {
 	FRigVMBaseOp(ERigVMOpCode InOpCode = ERigVMOpCode::Invalid)
@@ -105,6 +107,8 @@ struct RIGVM_API FRigVMBaseOp
 	ERigVMOpCode OpCode;
 };
 
+
+// execute a function
 struct RIGVM_API FRigVMExecuteOp : public FRigVMBaseOp
 {
 	FRigVMExecuteOp()
@@ -124,6 +128,7 @@ struct RIGVM_API FRigVMExecuteOp : public FRigVMBaseOp
 	FORCEINLINE uint8 GetArgumentCount() const { return uint8(OpCode) - uint8(ERigVMOpCode::Execute_0_Args); }
 };
 
+// zero the memory of a given register
 struct RIGVM_API FRigVMZeroOp : public FRigVMBaseOp
 {
 	FRigVMZeroOp()
@@ -141,6 +146,7 @@ struct RIGVM_API FRigVMZeroOp : public FRigVMBaseOp
 	FRigVMArgument Arg;
 };
 
+// set a given register to false
 struct RIGVM_API FRigVMFalseOp : public FRigVMBaseOp
 {
 	FRigVMFalseOp()
@@ -158,6 +164,7 @@ struct RIGVM_API FRigVMFalseOp : public FRigVMBaseOp
 	FRigVMArgument Arg;
 };
 
+// set a given register to true
 struct RIGVM_API FRigVMTrueOp : public FRigVMBaseOp
 {
 	FRigVMTrueOp()
@@ -175,6 +182,7 @@ struct RIGVM_API FRigVMTrueOp : public FRigVMBaseOp
 	FRigVMArgument Arg;
 };
 
+// copy the content of one register to another
 struct RIGVM_API FRigVMCopyOp : public FRigVMBaseOp
 {
 	FRigVMCopyOp()
@@ -210,6 +218,7 @@ struct RIGVM_API FRigVMCopyOp : public FRigVMBaseOp
 	int32 NumBytes;
 };
 
+// increment a int32 register
 struct RIGVM_API FRigVMIncrementOp : public FRigVMBaseOp
 {
 	FRigVMIncrementOp()
@@ -227,6 +236,7 @@ struct RIGVM_API FRigVMIncrementOp : public FRigVMBaseOp
 	FRigVMArgument Arg;
 };
 
+// decrement a int32 register
 struct RIGVM_API FRigVMDecrementOp : public FRigVMBaseOp
 {
 	FRigVMDecrementOp()
@@ -244,6 +254,7 @@ struct RIGVM_API FRigVMDecrementOp : public FRigVMBaseOp
 	FRigVMArgument Arg;
 };
 
+// fill a bool register with the result of (A == B)
 struct RIGVM_API FRigVMEqualsOp : public FRigVMBaseOp
 {
 	FRigVMEqualsOp()
@@ -271,6 +282,7 @@ struct RIGVM_API FRigVMEqualsOp : public FRigVMBaseOp
 	FRigVMArgument Result;
 };
 
+// fill a bool register with the result of (A != B)
 struct RIGVM_API FRigVMNotEqualsOp : public FRigVMBaseOp
 {
 	FRigVMNotEqualsOp()
@@ -298,6 +310,9 @@ struct RIGVM_API FRigVMNotEqualsOp : public FRigVMBaseOp
 	FRigVMArgument Result;
 };
 
+// jump to a new instruction index.
+// the instruction can be absolute, relative forward or relative backward
+// based on the opcode 
 struct RIGVM_API FRigVMJumpOp : public FRigVMBaseOp
 {
 	FRigVMJumpOp()
@@ -317,6 +332,9 @@ struct RIGVM_API FRigVMJumpOp : public FRigVMBaseOp
 	int32 InstructionIndex;
 };
 
+// jump to a new instruction index based on a condition.
+// the instruction can be absolute, relative forward or relative backward
+// based on the opcode 
 struct RIGVM_API FRigVMJumpIfOp : public FRigVMBaseOp
 {
 	FRigVMJumpIfOp()
@@ -342,6 +360,7 @@ struct RIGVM_API FRigVMJumpIfOp : public FRigVMBaseOp
 	bool Condition;
 };
 
+// change the type of a register
 struct RIGVM_API FRigVMChangeTypeOp : public FRigVMBaseOp
 {
 	FRigVMChangeTypeOp()
@@ -371,6 +390,7 @@ struct RIGVM_API FRigVMChangeTypeOp : public FRigVMBaseOp
 	uint16 SliceCount;
 };
 
+// exit the execution loop
 struct RIGVM_API FRigVMExitOp : public FRigVMBaseOp
 {
 	FRigVMExitOp()
@@ -379,12 +399,16 @@ struct RIGVM_API FRigVMExitOp : public FRigVMBaseOp
 	}
 };
 
+/**
+ * The FRigVMInstruction represents
+ * a single instruction within the VM.
+ */
 USTRUCT()
-struct RIGVM_API FRigVMByteCodeTableEntry
+struct RIGVM_API FRigVMInstruction
 {
 	GENERATED_USTRUCT_BODY()
 
-	FRigVMByteCodeTableEntry(ERigVMOpCode InOpCode = ERigVMOpCode::Invalid, uint64 InByteCodeIndex = UINT64_MAX)
+	FRigVMInstruction(ERigVMOpCode InOpCode = ERigVMOpCode::Invalid, uint64 InByteCodeIndex = UINT64_MAX)
 		: OpCode(InOpCode)
 		, ByteCodeIndex(InByteCodeIndex)
 	{
@@ -397,31 +421,49 @@ struct RIGVM_API FRigVMByteCodeTableEntry
 	uint64 ByteCodeIndex;
 };
 
+/**
+ * The FRigVMInstructionArray represents all current instructions
+ * within a RigVM and can be used to iterate over all operators and retrieve
+ * each instruction's data.
+ */
 USTRUCT()
-struct RIGVM_API FRigVMByteCodeTable
+struct RIGVM_API FRigVMInstructionArray
 {
 	GENERATED_USTRUCT_BODY()
 
 public:
 
-	FRigVMByteCodeTable();
+	FRigVMInstructionArray();
 
+	// Resets the data structure and removes all storage.
 	void Reset();
-	FORCEINLINE bool IsValidIndex(int32 InIndex) const { return Entries.IsValidIndex(InIndex); }
-	FORCEINLINE int32 Num() const { return Entries.Num(); }
-	FORCEINLINE const FRigVMByteCodeTableEntry& operator[](int32 InIndex) const { return Entries[InIndex]; }
+
+	// Returns true if a given instruction index is valid.
+	FORCEINLINE bool IsValidIndex(int32 InIndex) const { return Instructions.IsValidIndex(InIndex); }
+
+	// Returns the number of instructions.
+	FORCEINLINE int32 Num() const { return Instructions.Num(); }
+
+	// const accessor for an instruction given its index
+	FORCEINLINE const FRigVMInstruction& operator[](int32 InIndex) const { return Instructions[InIndex]; }
 
 private:
 
 	// hide utility constructor
-	FRigVMByteCodeTable(const FRigVMByteCode& InByteCode);
+	FRigVMInstructionArray(const FRigVMByteCode& InByteCode);
 
 	UPROPERTY()
-	TArray<FRigVMByteCodeTableEntry> Entries;
+	TArray<FRigVMInstruction> Instructions;
 
 	friend struct FRigVMByteCode;
 };
 
+/**
+ * The FRigVMByteCode is a container to store a list of instructions with
+ * their corresponding data. The byte code is then used within a VM to 
+ * execute. To iterate over the instructions within the byte code you can 
+ * use GetInstructions() to retrieve a FRigVMInstructionArray.
+ */
 USTRUCT()
 struct RIGVM_API FRigVMByteCode
 {
@@ -431,36 +473,68 @@ public:
 
 	FRigVMByteCode();
 
+	// resets the container and removes all memory
 	void Reset();
+
+	// returns the number of instructions in this container
 	uint64 Num() const;
 
+	// adds an execute operator given its function index arguments
 	uint64 AddExecuteOp(uint16 InFunctionIndex, const TArrayView<FRigVMArgument>& InArguments);
+
+	// adds a zero operator to zero the memory of a given argument
 	uint64 AddZeroOp(const FRigVMArgument& InArg);
+
+	// adds a false operator to set a given argument to false
 	uint64 AddFalseOp(const FRigVMArgument& InArg);
+
+	// adds a true operator to set a given argument to true
 	uint64 AddTrueOp(const FRigVMArgument& InArg);
+
+	// adds a copy operator to copy the content of a source argument to a target argument
 	uint64 AddCopyOp(const FRigVMArgument& InSource, const FRigVMArgument& InTarget, int32 InSourceOffset = INDEX_NONE, int32 InTargetOffset = INDEX_NONE, int32 InNumBytes = INDEX_NONE);
+
+	// adds an increment operator to increment a int32 argument
 	uint64 AddIncrementOp(const FRigVMArgument& InArg);
+
+	// adds an decrement operator to decrement a int32 argument
 	uint64 AddDecrementOp(const FRigVMArgument& InArg);
+
+	// adds an equals operator to store the comparison result of A and B into a Result argument
 	uint64 AddEqualsOp(const FRigVMArgument& InA, const FRigVMArgument& InB, const FRigVMArgument& InResult);
+
+	// adds an not-equals operator to store the comparison result of A and B into a Result argument
 	uint64 AddNotEqualsOp(const FRigVMArgument& InA, const FRigVMArgument& InB, const FRigVMArgument& InResult);
+
+	// adds an absolute, forward or backward jump operator
 	uint64 AddJumpOp(ERigVMOpCode InOpCode, uint16 InInstructionIndex);
-	uint64 AddJumpIfOp(ERigVMOpCode InOpCode, uint16 InInstructionIndex, const FRigVMArgument& InConditionArg, bool bInCondition = false);
+
+	// adds an absolute, forward or backward jump operator based on a condition argument
+	uint64 AddJumpIfOp(ERigVMOpCode InOpCode, uint16 InInstructionIndex, const FRigVMArgument& InConditionArg, bool bJumpWhenConditionIs = false);
+
+	// adds a change-type operator to reuse a register for a smaller or same size type
 	uint64 AddChangeTypeOp(FRigVMArgument InArg, ERigVMRegisterType InType, uint16 InElementSize, uint16 InElementCount, uint16 InSliceCount = 1);
+
+	// adds an exit operator to exit the execution loop
 	uint64 AddExitOp();
 
-	FORCEINLINE FRigVMByteCodeTable GetTable() const
+	// returns an instruction array for iterating over all operators
+	FORCEINLINE FRigVMInstructionArray GetInstructions() const
 	{
-		return FRigVMByteCodeTable(*this);
+		return FRigVMInstructionArray(*this);
 	}
 
+	// returns the opcode at a given byte index
 	FORCEINLINE ERigVMOpCode GetOpCodeAt(uint64 InByteCodeIndex) const
 	{
 		ensure(InByteCodeIndex >= 0 && InByteCodeIndex < ByteCode.Num());
 		return (ERigVMOpCode)ByteCode[InByteCodeIndex];
 	}
 
+	// returns the size of the operator in bytes at a given byte index
 	uint64 GetOpNumBytesAt(uint64 InByteCodeIndex, bool bIncludeArguments = true) const;
 
+	// returns an operator at a given byte code index
 	template<class OpType>
 	FORCEINLINE const OpType& GetOpAt(uint64 InByteCodeIndex) const
 	{
@@ -468,30 +542,35 @@ public:
 		return *(const OpType*)(ByteCode.GetData() + InByteCodeIndex);
 	}
 
+	// returns an operator for a given instruction
 	template<class OpType>
-	FORCEINLINE const OpType& GetOpAt(const FRigVMByteCodeTableEntry& InEntry) const
+	FORCEINLINE const OpType& GetOpAt(const FRigVMInstruction& InInstruction) const
 	{
-		return GetOpAt<OpType>(InEntry.ByteCodeIndex);
+		return GetOpAt<OpType>(InInstruction.ByteCodeIndex);
 	}
 
+	// returns a list of arguments at a given byte code index
 	FORCEINLINE TArrayView<FRigVMArgument> GetArgumentsAt(uint64 InByteCodeIndex, uint16 InArgumentCount) const
 	{
 		ensure(InByteCodeIndex >= 0 && InByteCodeIndex <= ByteCode.Num() - sizeof(FRigVMArgument) * InArgumentCount);
 		return TArrayView<FRigVMArgument>((FRigVMArgument*)(ByteCode.GetData() + InByteCodeIndex), InArgumentCount);
 	}
 
+	// returns the arguments for an execute operator / instruction at a given byte code index
 	FORCEINLINE TArrayView<FRigVMArgument> GetArgumentsForExecuteOp(uint64 InByteCodeIndex) const
 	{
 		const FRigVMExecuteOp& ExecuteOp = GetOpAt<FRigVMExecuteOp>(InByteCodeIndex);
 		return GetArgumentsAt(InByteCodeIndex + sizeof(FRigVMExecuteOp), ExecuteOp.GetArgumentCount());
 	}
 
-	FORCEINLINE TArrayView<FRigVMArgument> GetArgumentsForExecuteOp(const FRigVMByteCodeTableEntry& InEntry) const
+	// returns the arguments for a given execute instruction
+	FORCEINLINE TArrayView<FRigVMArgument> GetArgumentsForExecuteOp(const FRigVMInstruction& InInstruction) const
 	{
-		const FRigVMExecuteOp& ExecuteOp = GetOpAt<FRigVMExecuteOp>(InEntry);
-		return GetArgumentsAt(InEntry.ByteCodeIndex + sizeof(FRigVMExecuteOp), ExecuteOp.GetArgumentCount());
+		const FRigVMExecuteOp& ExecuteOp = GetOpAt<FRigVMExecuteOp>(InInstruction);
+		return GetArgumentsAt(InInstruction.ByteCodeIndex + sizeof(FRigVMExecuteOp), ExecuteOp.GetArgumentCount());
 	}
 
+	// returns the raw data of the byte code
 	FORCEINLINE const TArrayView<uint8> GetByteCode() const
 	{
 		const uint8* Data = ByteCode.GetData();
diff --git a/Engine/Source/Runtime/RigVM/Public/RigVMCore/RigVMMemory.h b/Engine/Source/Runtime/RigVM/Public/RigVMCore/RigVMMemory.h
--- a/Engine/Source/Runtime/RigVM/Public/RigVMCore/RigVMMemory.h
+++ b/Engine/Source/Runtime/RigVM/Public/RigVMCore/RigVMMemory.h
@@ -173,7 +173,7 @@ struct RIGVM_API FRigVMRegister
 	// Returns true if the register stores plain / shallow memory
 	FORCEINLINE bool IsPlain() const { return ScriptStructIndex == INDEX_NONE; }
 
-	// Returns the number of allocated bytes (inclusing alignment + trailing bytes)
+	// Returns the number of allocated bytes (including alignment + trailing bytes)
 	FORCEINLINE uint16 GetAllocatedBytes() const { return ElementCount * ElementSize * SliceCount + (uint16)AlignmentBytes + TrailingBytes; }
 
 	// Returns the number of bytes for a complete slice
diff --git a/Engine/Source/Runtime/RigVM/Public/RigVMCore/RigVMRegistry.h b/Engine/Source/Runtime/RigVM/Public/RigVMCore/RigVMRegistry.h
--- a/Engine/Source/Runtime/RigVM/Public/RigVMCore/RigVMRegistry.h
+++ b/Engine/Source/Runtime/RigVM/Public/RigVMCore/RigVMRegistry.h
@@ -8,6 +8,10 @@
 typedef TArrayView<void*> FRigVMUserDataArray;
 typedef void (*FRigVMFunctionPtr)(const FRigVMArgumentArray&, FRigVMMemoryContainerPtrArray&, const FRigVMUserDataArray&);
 
+/**
+ * The FRigVMFunction is used to represent a function pointer generated by UHT 
+ * for a given name. The name might be something like "FMyStruct::MyVirtualMethod"
+ */
 struct RIGVM_API FRigVMFunction
 {
 	const TCHAR* Name;
@@ -26,14 +30,29 @@ struct RIGVM_API FRigVMFunction
 	}
 };
 
+/**
+ * The FRigVMRegistry is used to manage all known function pointers
+ * for use in the RigVM. The Register method is called automatically
+ * when the static struct is initially constructed for each USTRUCT
+ * hosting a RIGVM_METHOD enabled virtual function.
+ */
 struct RIGVM_API FRigVMRegistry
 {
 public:
+
+	// Returns the singleton registry
 	static FRigVMRegistry& Get();
 
+	// Registers a function given its name.
+	// The name will be the name of the struct and virtual method,
+	// for example "FMyStruct::MyVirtualMethod"
 	void Register(const TCHAR* InName, FRigVMFunctionPtr InFunctionPtr);
+
+	// Refreshes the list and finds the function pointers
+	// based on the names.
 	void Refresh();
 
+	// Returns a function pointer given its name
 	FRigVMFunctionPtr Find(const TCHAR* InName) const;
 
 private:
