RigVM: Fix typos in header

diff --git a/Engine/Source/Runtime/RigVM/Public/RigVMCore/RigVMMemory.h b/Engine/Source/Runtime/RigVM/Public/RigVMCore/RigVMMemory.h
--- a/Engine/Source/Runtime/RigVM/Public/RigVMCore/RigVMMemory.h
+++ b/Engine/Source/Runtime/RigVM/Public/RigVMCore/RigVMMemory.h
@@ -82,7 +82,7 @@ enum class ERigVMRegisterType : uint8
 // we can store arbitrary data, so it provides a series of properties to
 // describe the memory location.
 // Registers also support the notion of slices. A slice is a complete copy of
-// the memory - so for example if you register stores 4 Vectors, then a slice 
+// the memory - so for example if your register stores 4 Vectors, then a slice 
 // would contain 48 bytes (4 * 3 * 4). The register can however store multiple
 // slices / copies of that if needed. Slices can be used to provide 
 // per-invocation memory to functions within the same register.
@@ -217,7 +217,7 @@ typedef TArrayView<FRigVMRegister> FRigVMRegisterArray;
  * Elements can be accessed by index (index of the register), FRigVMArgument or by name.
  * Name access is optional and is specified upon construction of the container.
  * The memory container provides a series of templated functions to add and get data.
- * Plain types (shallow types without the need for construction) can be added an retrieved
+ * Plain types (shallow types without the need for construction) can be added and retrieved
  * using the methods suffixed with 'Plain'.
  *
  * For example:
@@ -229,8 +229,8 @@ typedef TArrayView<FRigVMRegister> FRigVMRegisterArray;
  * 		int32 Index = Container.AddPlainArray<float>(MyArray);
  *      TArrayView<float> ArrayView = Container.GetArray<float>(Index);
  *
- * More complex data structures such as USTRUCT types need to be store with the
- * 'Struct' suffixed methods - since the require a constructor / destructor to be called.
+ * More complex data structures such as USTRUCT types need to be stored with the
+ * 'Struct' suffixed methods - since they require a constructor / destructor to be called.
  *
  * For example:
  *
