RigVM: Make scriptstruct array transient

diff --git a/Engine/Source/Runtime/RigVM/Public/RigVMCore/RigVMMemory.h b/Engine/Source/Runtime/RigVM/Public/RigVMCore/RigVMMemory.h
--- a/Engine/Source/Runtime/RigVM/Public/RigVMCore/RigVMMemory.h
+++ b/Engine/Source/Runtime/RigVM/Public/RigVMCore/RigVMMemory.h
@@ -773,7 +773,7 @@ private:
 	UPROPERTY(transient)
 	TArray<uint8> Data;
 
-	UPROPERTY()
+	UPROPERTY(transient)
 	TArray<UScriptStruct*> ScriptStructs;
 
 	UPROPERTY(transient)
