COPY from //Dev-Anim to //Dev-Main

diff --git a/Engine/Source/Runtime/Engine/Public/Rendering/SkinWeightVertexBuffer.h b/Engine/Source/Runtime/Engine/Public/Rendering/SkinWeightVertexBuffer.h
--- a/Engine/Source/Runtime/Engine/Public/Rendering/SkinWeightVertexBuffer.h
+++ b/Engine/Source/Runtime/Engine/Public/Rendering/SkinWeightVertexBuffer.h
@@ -5,28 +5,45 @@
 #include "RenderResource.h"
 #include "StaticMeshVertexData.h"
 #include "GPUSkinPublicDefs.h"
+#include "SkeletalMeshTypes.h"
+#include "GPUSkinVertexFactory.h"
+#include "UObject/AnimObjectVersion.h"
 
 class FStaticMeshVertexDataInterface;
 struct FSoftSkinVertex;
 
+typedef uint8 FBoneIndex8;
+typedef uint16 FBoneIndex16;
+
 /** Struct for storing skin weight info in vertex buffer */
-template <bool bExtraBoneInfluences>
-struct TSkinWeightInfo
+template <bool bExtraBoneInfluences, typename BoneIndexType>
+struct TLegacySkinWeightInfo
 {
 	enum
 	{
-		NumInfluences = bExtraBoneInfluences ? MAX_TOTAL_INFLUENCES : MAX_INFLUENCES_PER_STREAM,
+		NumInfluences = bExtraBoneInfluences ? EXTRA_BONE_INFLUENCES : MAX_INFLUENCES_PER_STREAM,
 	};
-	uint8	InfluenceBones[NumInfluences];
-	uint8	InfluenceWeights[NumInfluences];
+	BoneIndexType	InfluenceBones[NumInfluences];
+	uint8			InfluenceWeights[NumInfluences];
 
-	friend FArchive& operator<<(FArchive& Ar, TSkinWeightInfo& I)
+	friend FArchive& operator<<(FArchive& Ar, TLegacySkinWeightInfo& I)
 	{
-		// serialize bone and weight uint8 arrays in order
+		Ar.UsingCustomVersion(FAnimObjectVersion::GUID);
+		// serialize bone index and weight arrays in order
 		// this is required when serializing as bulk data memory (see TArray::BulkSerialize notes)
 		for (uint32 InfluenceIndex = 0; InfluenceIndex < NumInfluences; InfluenceIndex++)
 		{
-			Ar << I.InfluenceBones[InfluenceIndex];
+			if (Ar.IsLoading() && Ar.CustomVer(FAnimObjectVersion::GUID) < FAnimObjectVersion::IncreaseBoneIndexLimitPerChunk)
+			{
+				// Older versions use uint8 bone indices
+				uint8 BoneIndex = 0;
+				Ar << BoneIndex;
+				I.InfluenceBones[InfluenceIndex] = BoneIndex;
+			}
+			else
+			{
+				Ar << I.InfluenceBones[InfluenceIndex];
+			}
 		}
 
 		for (uint32 InfluenceIndex = 0; InfluenceIndex < NumInfluences; InfluenceIndex++)
@@ -36,23 +53,139 @@ struct TSkinWeightInfo
 
 		return Ar;
 	}
+};
 
-	bool GetRigidWeightBone(uint8& OutBoneIndex) const
-	{
-		bool bIsRigid = false;
+/** An runtime structure for passing data to FSkinWeightVertexBuffer */
+struct FSkinWeightInfo
+{
+	FBoneIndexType	InfluenceBones[MAX_TOTAL_INFLUENCES];
+	uint8			InfluenceWeights[MAX_TOTAL_INFLUENCES];
+};
+
+struct FSkinWeightRHIInfo
+{
+	FVertexBufferRHIRef DataVertexBufferRHI;
+	FVertexBufferRHIRef LookupVertexBufferRHI;
+};
 
-		for (int32 WeightIdx = 0; WeightIdx < NumInfluences; WeightIdx++)
+/** A lookup vertex buffer storing skin weight stream offset / influence count. Only used for unlimited bone influences. */
+class FSkinWeightLookupVertexBuffer : public FVertexBuffer
+{
+public:
+	/** Default constructor. */
+	ENGINE_API FSkinWeightLookupVertexBuffer();
+
+	/** Constructor (copy) */
+	ENGINE_API FSkinWeightLookupVertexBuffer(const FSkinWeightLookupVertexBuffer& Other);
+
+	/** Destructor. */
+	ENGINE_API ~FSkinWeightLookupVertexBuffer();
+
+	/** Assignment. Assumes that vertex buffer will be rebuilt */
+	ENGINE_API FSkinWeightLookupVertexBuffer& operator=(const FSkinWeightLookupVertexBuffer& Other);
+
+	/** Delete existing resources */
+	ENGINE_API void CleanUp();
+
+	/** @return true is LookupData is valid */
+	bool IsLookupDataValid() const;
+
+	void Init(uint32 InNumVertices);
+
+	friend FArchive& operator<<(FArchive& Ar, FSkinWeightLookupVertexBuffer& VertexBuffer);
+
+	void SerializeMetaData(FArchive& Ar);
+	void CopyMetaData(const FSkinWeightLookupVertexBuffer& Other);
+
+	// FRenderResource interface.
+	virtual void InitRHI() override;
+	virtual void ReleaseRHI() override;
+
+	virtual FString GetFriendlyName() const override
+	{ return TEXT("SkeletalMesh Vertex Weights Lookup"); }
+
+	/** @return number of vertices in this vertex buffer */
+	FORCEINLINE uint32 GetNumVertices() const
+	{ return NumVertices; }
+
+	/** @return cached stride for vertex data type for this vertex buffer */
+	FORCEINLINE uint32 GetStride() const
+	{ return sizeof(uint32); }
+
+	/** @return total size of data in resource array */
+	FORCEINLINE uint32 GetVertexDataSize() const
+	{ return NumVertices * GetStride(); }
+
+	// @param guaranteed only to be valid if the vertex buffer is valid
+	FRHIShaderResourceView* GetSRV() const
+	{ return SRVValue; }
+
+	/** Set if the CPU needs access to this vertex buffer */
+	void SetNeedsCPUAccess(bool bInNeedsCPUAccess)
+	{ bNeedsCPUAccess = bInNeedsCPUAccess; }
+
+	bool GetNeedsCPUAccess() const
+	{ return bNeedsCPUAccess; }
+
+	/** Create an RHI vertex buffer with CPU data. CPU data may be discarded after creation (see TResourceArray::Discard) */
+	FVertexBufferRHIRef CreateRHIBuffer_RenderThread();
+	FVertexBufferRHIRef CreateRHIBuffer_Async();
+
+	/** Similar to Init/ReleaseRHI but only update existing SRV so references to the SRV stays valid */
+	template <uint32 MaxNumUpdates>
+	void InitRHIForStreaming(FRHIVertexBuffer* IntermediateBuffer, TRHIResourceUpdateBatcher<MaxNumUpdates>& Batcher)
+	{
+		if (VertexBufferRHI && IntermediateBuffer)
 		{
-			if (InfluenceWeights[WeightIdx] == 255)
-			{
-				bIsRigid = true;
-				OutBoneIndex = InfluenceBones[WeightIdx];
-				break;
-			}
+			check(SRVValue);
+			Batcher.QueueUpdateRequest(VertexBufferRHI, IntermediateBuffer);
+			Batcher.QueueUpdateRequest(SRVValue, VertexBufferRHI, PixelFormatStride, PixelFormat);
 		}
+	}
 
-		return bIsRigid;
+	template <uint32 MaxNumUpdates>
+	void ReleaseRHIForStreaming(TRHIResourceUpdateBatcher<MaxNumUpdates>& Batcher)
+	{
+		if (VertexBufferRHI)
+		{
+			Batcher.QueueUpdateRequest(VertexBufferRHI, nullptr);
+		}
+		if (SRVValue)
+		{
+			Batcher.QueueUpdateRequest(SRVValue, nullptr, 0, 0);
+		}
 	}
+
+	void GetWeightOffsetAndInfluenceCount(uint32 VertexIndex, uint32& OutWeightOffset, uint32& OutInfluenceCount) const;
+	void SetWeightOffsetAndInfluenceCount(uint32 VertexIndex, uint32 WeightOffset, uint32 InfluenceCount);
+
+protected:
+	// guaranteed only to be valid if the vertex buffer is valid
+	FShaderResourceViewRHIRef SRVValue;
+
+private:
+	void ResizeBuffer(uint32 InNumVertices);
+
+	/** Allocates the vertex data storage type. */
+	ENGINE_API void AllocateData();
+
+	template <bool bRenderThread>
+	FVertexBufferRHIRef CreateRHIBuffer_Internal();
+
+	/** true if this vertex buffer will be used with CPU skinning. Resource arrays are set to cpu accessible if this is true */
+	bool bNeedsCPUAccess;
+
+	/** The vertex data storage type */
+	FStaticMeshVertexDataInterface* LookupData;
+
+	/** The cached vertex data pointer. */
+	uint8* Data;
+
+	/** The cached number of vertices. */
+	uint32 NumVertices;
+
+	static const EPixelFormat PixelFormat = PF_R32_UINT;
+	static const uint32 PixelFormatStride = 4;
 };
 
 /** The implementation of the skin weight vertex data storage type. */
@@ -66,125 +199,96 @@ public:
 	}
 };
 
-/** A vertex buffer for skin weights. */
-class FSkinWeightVertexBuffer : public FVertexBuffer
+/** A vertex buffer storing bone index/weight data. */
+class FSkinWeightDataVertexBuffer : public FVertexBuffer
 {
 public:
 	/** Default constructor. */
-	ENGINE_API FSkinWeightVertexBuffer();
+	ENGINE_API FSkinWeightDataVertexBuffer();
 
 	/** Constructor (copy) */
-	ENGINE_API FSkinWeightVertexBuffer(const FSkinWeightVertexBuffer& Other);
+	ENGINE_API FSkinWeightDataVertexBuffer(const FSkinWeightDataVertexBuffer& Other);
 
 	/** Destructor. */
-	ENGINE_API ~FSkinWeightVertexBuffer();
+	ENGINE_API ~FSkinWeightDataVertexBuffer();
 
 	/** Assignment. Assumes that vertex buffer will be rebuilt */
-	ENGINE_API FSkinWeightVertexBuffer& operator=(const FSkinWeightVertexBuffer& Other);
+	ENGINE_API FSkinWeightDataVertexBuffer& operator=(const FSkinWeightDataVertexBuffer& Other);
 
 	/** Delete existing resources */
 	ENGINE_API void CleanUp();
 
-	/** @return true is WeightData is valid */
-	bool IsWeightDataValid() const;
+	void Init(uint32 InNumBones, uint32 InNumVertices);
 
-#if WITH_EDITOR
-	/** Init from another skin weight buffer */
-	ENGINE_API void Init(const TArray<FSoftSkinVertex>& InVertices);
-#endif // WITH_EDITOR
-
-	friend FArchive& operator<<(FArchive& Ar, FSkinWeightVertexBuffer& VertexBuffer);
+	friend FArchive& operator<<(FArchive& Ar, FSkinWeightDataVertexBuffer& VertexBuffer);
 
 	void SerializeMetaData(FArchive& Ar);
-
-	void CopyMetaData(const FSkinWeightVertexBuffer& Other);
+	void CopyMetaData(const FSkinWeightDataVertexBuffer& Other);
 
 	// FRenderResource interface.
 	virtual void InitRHI() override;
 	virtual void ReleaseRHI() override;
 
 	virtual FString GetFriendlyName() const override 
-	{ 
-		return TEXT("SkeletalMesh Vertex Weights"); 
-	}
-
-	/** Const access to entry in weight data */
-	template <bool bExtraBoneInfluencesT>
-	FORCEINLINE const TSkinWeightInfo<bExtraBoneInfluencesT>* GetSkinWeightPtr(uint32 VertexIndex) const
-	{
-		checkSlow(VertexIndex < GetNumVertices());
-		return (TSkinWeightInfo<bExtraBoneInfluencesT>*)(Data + VertexIndex * Stride);
-	}
-
-	/** Non-const access to entry in weight data */
-	template <bool bExtraBoneInfluencesT>
-	FORCEINLINE TSkinWeightInfo<bExtraBoneInfluencesT>* GetSkinWeightPtr(uint32 VertexIndex)
-	{
-		checkSlow(VertexIndex < GetNumVertices());
-		return (TSkinWeightInfo<bExtraBoneInfluencesT>*)(Data + VertexIndex * Stride);
-	}
+	{ return TEXT("SkeletalMesh Vertex Weights Data"); }
 
 	/** @return number of vertices in this vertex buffer */
 	FORCEINLINE uint32 GetNumVertices() const
-	{
-		return NumVertices;
-	}
+	{ return NumVertices; }
 
-	/** @return cached stride for vertex data type for this vertex buffer */
-	FORCEINLINE uint32 GetStride() const
-	{
-		return Stride;
-	}
+	/** @return number of bones in this vertex buffer */
+	FORCEINLINE uint32 GetNumBones() const
+	{ return NumBones; }
+
+	/** @return byte size of each bone index */
+	FORCEINLINE uint32 GetBoneIndexByteSize() const
+	{ return (Use16BitBoneIndex() ? sizeof(FBoneIndex16) : sizeof(FBoneIndex8)); }
+
+	/** @return vertex stride for when using constant number of bones per vertex buffer */
+	FORCEINLINE uint32 GetConstantInfluencesVertexStride() const
+	{ return (GetBoneIndexByteSize() + sizeof(uint8)) * MaxBoneInfluences; }
+
+	/** @return offset position for bone weights data for each vertex */
+	FORCEINLINE uint32 GetConstantInfluencesBoneWeightsOffset() const
+	{ return GetBoneIndexByteSize() * MaxBoneInfluences; }
 
 	/** @return total size of data in resource array */
 	FORCEINLINE uint32 GetVertexDataSize() const
-	{
-		return NumVertices * Stride;
-	}
+	{ return NumBones * (GetBoneIndexByteSize() + sizeof(uint8)); }
 
 	// @param guaranteed only to be valid if the vertex buffer is valid
 	FRHIShaderResourceView* GetSRV() const
-	{
-		return SRVValue;
-	}
+	{ return SRVValue; }
 
 	/** Set if the CPU needs access to this vertex buffer */
-	void SetNeedsCPUAccess(bool bInNeedsCPUAccess)
-	{
-		bNeedsCPUAccess = bInNeedsCPUAccess;
-	}
+	FORCEINLINE void SetNeedsCPUAccess(bool bInNeedsCPUAccess)
+	{ bNeedsCPUAccess = bInNeedsCPUAccess; }
 
-	bool GetNeedsCPUAccess() const
-	{
-		return bNeedsCPUAccess;
-	}
+	FORCEINLINE bool GetNeedsCPUAccess() const
+	{ return bNeedsCPUAccess; }
 
-	/** Set if this will have extra streams for bone indices & weights. */
-	void SetHasExtraBoneInfluences(bool bInHasExtraBoneInfluences)
-	{
-		bExtraBoneInfluences = bInHasExtraBoneInfluences;
-	}
+	FORCEINLINE bool GetVariableBonesPerVertex() const
+	{ return bVariableBonesPerVertex; }
 
-	bool HasExtraBoneInfluences() const
-	{
-		return bExtraBoneInfluences;
-	}
+	/** Set if this will have extra streams for bone indices & weights. */
+	void SetMaxBoneInfluences(uint32 InMaxBoneInfluences);
 
-	/** Assignment operator for assigning array of weights to this buffer */
-	template <bool bExtraBoneInfluencesT>
-	FSkinWeightVertexBuffer& operator=(const TArray< TSkinWeightInfo<bExtraBoneInfluencesT> >& InWeights)
-	{
-		check(bExtraBoneInfluences == bExtraBoneInfluencesT);
-		AllocateData();
+	FORCEINLINE uint32 GetMaxBoneInfluences() const
+	{ return MaxBoneInfluences; }
 
-		static_cast<FSkinWeightVertexData< TSkinWeightInfo<bExtraBoneInfluencesT> >*>(WeightData)->Assign(InWeights);
+	FORCEINLINE void SetUse16BitBoneIndex(bool bInUse16BitBoneIndex)
+	{ bUse16BitBoneIndex = bInUse16BitBoneIndex; }
 
-		Data = WeightData->GetDataPointer();
-		Stride = WeightData->GetStride();
-		NumVertices = InWeights.Num();
+	FORCEINLINE bool Use16BitBoneIndex() const
+	{ return bUse16BitBoneIndex; }
 
-		return *this;
-	}
+	GPUSkinBoneInfluenceType GetBoneInfluenceType() const;
+	bool GetRigidWeightBone(uint32 VertexWeightOffset, uint32 VertexInfluenceCount, int32& OutBoneIndex) const;
+	uint32 GetBoneIndex(uint32 VertexWeightOffset, uint32 VertexInfluenceCount, uint32 InfluenceIndex) const;
+	void SetBoneIndex(uint32 VertexWeightOffset, uint32 VertexInfluenceCount, uint32 InfluenceIndex, uint32 BoneIndex);
+	uint8 GetBoneWeight(uint32 VertexWeightOffset, uint32 VertexInfluenceCount, uint32 InfluenceIndex) const;
+	void SetBoneWeight(uint32 VertexWeightOffset, uint32 VertexInfluenceCount, uint32 InfluenceIndex, uint8 BoneWeight);
+	void ResetVertexBoneWeights(uint32 VertexWeightOffset, uint32 VertexInfluenceCount);
 
 	/** Create an RHI vertex buffer with CPU data. CPU data may be discarded after creation (see TResourceArray::Discard) */
 	FVertexBufferRHIRef CreateRHIBuffer_RenderThread();
@@ -198,7 +302,7 @@ public:
 		{
 			check(SRVValue);
 			Batcher.QueueUpdateRequest(VertexBufferRHI, IntermediateBuffer);
-			Batcher.QueueUpdateRequest(SRVValue, VertexBufferRHI, 4, PF_R32_UINT);
+			Batcher.QueueUpdateRequest(SRVValue, VertexBufferRHI, GetPixelFormatStride(), GetPixelFormat());
 		}
 	}
 
@@ -220,12 +324,24 @@ protected:
 	FShaderResourceViewRHIRef SRVValue;
 
 private:
+	void ResizeBuffer(uint32 InNumBones, uint32 InNumVertices);
+
+	EPixelFormat GetPixelFormat() const
+	{ return bVariableBonesPerVertex ? PF_R8_UINT : PF_R32_UINT; }
+
+	uint32 GetPixelFormatStride() const
+	{ return bVariableBonesPerVertex ? 1 : 4; }
 
 	/** true if this vertex buffer will be used with CPU skinning. Resource arrays are set to cpu accessible if this is true */
 	bool bNeedsCPUAccess;
 
+	bool bVariableBonesPerVertex;
+
 	/** Has extra bone influences per Vertex, which means using a different TGPUSkinVertexBase */
-	bool bExtraBoneInfluences;
+	uint32 MaxBoneInfluences;
+
+	/** Use 16 bit bone index instead of 8 bit */
+	bool bUse16BitBoneIndex;
 
 	/** The vertex data storage type */
 	FStaticMeshVertexDataInterface* WeightData;
@@ -233,21 +349,135 @@ private:
 	/** The cached vertex data pointer. */
 	uint8* Data;
 
-	/** The cached vertex stride. */
-	uint32 Stride;
-
 	/** The cached number of vertices. */
 	uint32 NumVertices;
 
+	/** The cached number of bones. */
+	uint32 NumBones;
+
 	/** Allocates the vertex data storage type. */
 	ENGINE_API void AllocateData();
 
-#if WITH_EDITOR
-	/** Helper for concrete types */
-	template <bool bUsesExtraBoneInfluences>
-	void SetWeightsForVertex(uint32 VertexIndex, const FSoftSkinVertex& SrcVertex);
-#endif //WITH_EDITOR
-
 	template <bool bRenderThread>
 	FVertexBufferRHIRef CreateRHIBuffer_Internal();
+};
+
+/** A container for skin weights data vertex buffer and lookup vertex buffer. */
+class FSkinWeightVertexBuffer
+{
+public:
+	/** Default constructor. */
+	ENGINE_API FSkinWeightVertexBuffer();
+
+	/** Constructor (copy) */
+	ENGINE_API FSkinWeightVertexBuffer(const FSkinWeightVertexBuffer& Other);
+
+	/** Destructor. */
+	ENGINE_API ~FSkinWeightVertexBuffer();
+
+	/** Assignment. Assumes that vertex buffer will be rebuilt */
+	ENGINE_API FSkinWeightVertexBuffer& operator=(const FSkinWeightVertexBuffer& Other);
+
+	/** Delete existing resources */
+	ENGINE_API void CleanUp();
+
+#if WITH_EDITOR
+	/** Init from another skin weight buffer */
+	ENGINE_API void Init(const TArray<FSoftSkinVertex>& InVertices);
+#endif // WITH_EDITOR
+
+	/** Assignment operator for assigning array of weights to this buffer */
+	FSkinWeightVertexBuffer& operator=(const TArray<FSkinWeightInfo>& InWeights);
+	void GetSkinWeights(TArray<FSkinWeightInfo>& OutVertices) const;
+	FSkinWeightInfo GetVertexSkinWeights(uint32 VertexIndex) const;
+
+	friend FArchive& operator<<(FArchive& Ar, FSkinWeightVertexBuffer& VertexBuffer);
+
+	void SerializeMetaData(FArchive& Ar);
+	void CopyMetaData(const FSkinWeightVertexBuffer& Other);
+
+	/** @return number of vertices in this vertex buffer */
+	FORCEINLINE uint32 GetNumVertices() const
+	{ return DataVertexBuffer.GetNumVertices(); }
+
+	/** @return total size of data in resource array */
+	FORCEINLINE uint32 GetVertexDataSize() const
+	{ return LookupVertexBuffer.GetVertexDataSize() + DataVertexBuffer.GetVertexDataSize(); }
+
+	void SetNeedsCPUAccess(bool bInNeedsCPUAccess)
+	{
+		DataVertexBuffer.SetNeedsCPUAccess(bInNeedsCPUAccess);
+		LookupVertexBuffer.SetNeedsCPUAccess(bInNeedsCPUAccess);
+	}
+
+	bool GetNeedsCPUAccess() const
+	{ return DataVertexBuffer.GetNeedsCPUAccess(); }
+
+	void SetMaxBoneInfluences(uint32 InMaxBoneInfluences)
+	{ DataVertexBuffer.SetMaxBoneInfluences(InMaxBoneInfluences); }
+
+	uint32 GetMaxBoneInfluences() const
+	{ return DataVertexBuffer.GetMaxBoneInfluences(); }
+
+	void SetUse16BitBoneIndex(bool bInUse16BitBoneIndex)
+	{ DataVertexBuffer.SetUse16BitBoneIndex(bInUse16BitBoneIndex); }
+
+	bool Use16BitBoneIndex() const
+	{ return DataVertexBuffer.Use16BitBoneIndex(); }
+
+	uint32 GetBoneIndexByteSize() const
+	{ return DataVertexBuffer.GetBoneIndexByteSize(); }
+
+	bool GetVariableBonesPerVertex() const
+	{ return DataVertexBuffer.GetVariableBonesPerVertex(); }
+
+	uint32 GetConstantInfluencesVertexStride() const
+	{ return DataVertexBuffer.GetConstantInfluencesVertexStride(); }
+
+	uint32 GetConstantInfluencesBoneWeightsOffset() const
+	{ return DataVertexBuffer.GetConstantInfluencesBoneWeightsOffset(); }
+
+	const FSkinWeightDataVertexBuffer* GetDataVertexBuffer() const
+	{ return &DataVertexBuffer; }
+
+	const FSkinWeightLookupVertexBuffer* GetLookupVertexBuffer() const
+	{ return &LookupVertexBuffer; }
+
+	FSkinWeightRHIInfo CreateRHIBuffer_RenderThread();
+	FSkinWeightRHIInfo CreateRHIBuffer_Async();
+
+	GPUSkinBoneInfluenceType GetBoneInfluenceType() const;
+	void GetVertexInfluenceOffsetCount(uint32 VertexIndex, uint32& VertexWeightOffset, uint32& VertexInfluenceCount) const;
+	ENGINE_API bool GetRigidWeightBone(uint32 VertexIndex, int32& OutBoneIndex) const;
+	uint32 GetBoneIndex(uint32 VertexIndex, uint32 InfluenceIndex) const;
+	void SetBoneIndex(uint32 VertexIndex, uint32 InfluenceIndex, uint32 BoneIndex);
+	uint8 GetBoneWeight(uint32 VertexIndex, uint32 InfluenceIndex) const;
+	void SetBoneWeight(uint32 VertexIndex, uint32 InfluenceIndex, uint8 BoneWeight);
+	void ResetVertexBoneWeights(uint32 VertexIndex);
+
+	void BeginInitResources();
+	void BeginReleaseResources();
+	void ReleaseResources();
+
+	/** Similar to Init/ReleaseRHI but only update existing SRV so references to the SRV stays valid */
+	template <uint32 MaxNumUpdates>
+	void InitRHIForStreaming(const FSkinWeightRHIInfo& RHIInfo, TRHIResourceUpdateBatcher<MaxNumUpdates>& Batcher)
+	{
+		DataVertexBuffer.InitRHIForStreaming(RHIInfo.DataVertexBufferRHI, Batcher);
+		LookupVertexBuffer.InitRHIForStreaming(RHIInfo.DataVertexBufferRHI, Batcher);
+	}
+
+	template <uint32 MaxNumUpdates>
+	void ReleaseRHIForStreaming(TRHIResourceUpdateBatcher<MaxNumUpdates>& Batcher)
+	{
+		DataVertexBuffer.ReleaseRHIForStreaming(Batcher);
+		LookupVertexBuffer.ReleaseRHIForStreaming(Batcher);
+	}
+
+private:
+	/** Skin weights for skinning */
+	FSkinWeightDataVertexBuffer		DataVertexBuffer;
+
+	/** Skin weights lookup buffer */
+	FSkinWeightLookupVertexBuffer	LookupVertexBuffer;
 };
\ No newline at end of file
diff --git a/Engine/Source/Runtime/Engine/Public/Rendering/StaticMeshVertexDataInterface.h b/Engine/Source/Runtime/Engine/Public/Rendering/StaticMeshVertexDataInterface.h
--- a/Engine/Source/Runtime/Engine/Public/Rendering/StaticMeshVertexDataInterface.h
+++ b/Engine/Source/Runtime/Engine/Public/Rendering/StaticMeshVertexDataInterface.h
@@ -46,7 +46,7 @@ public:
 	virtual FResourceArrayInterface* GetResourceArray() = 0;
 
 	/** Serializer. */
-	virtual void Serialize(FArchive& Ar) = 0;
+	virtual void Serialize(FArchive& Ar, bool bForcePerElementSerialization = false) = 0;
 
 	virtual SIZE_T GetResourceSize() const = 0;
 
diff --git a/Engine/Source/Runtime/Engine/Public/SkeletalMeshMerge.h b/Engine/Source/Runtime/Engine/Public/SkeletalMeshMerge.h
--- a/Engine/Source/Runtime/Engine/Public/SkeletalMeshMerge.h
+++ b/Engine/Source/Runtime/Engine/Public/SkeletalMeshMerge.h
@@ -231,8 +231,8 @@ private:
 	* Creates a new LOD model and adds the new merged sections to it. Modifies the MergedMesh.
 	* @param LODIdx - current LOD to process
 	*/
-	template<typename VertexDataType, typename SkinWeightType>
-	void GenerateLODModel( int32 LODIdx );
+	template<typename VertexDataType>
+	void GenerateLODModel( int32 LODIdx);
 
 	/**
 	* Generate the list of sections that need to be created along with info needed to merge sections
@@ -313,6 +313,6 @@ private:
 	void CopyVertexFromSource(VertexDataType& DestVert, const FSkeletalMeshLODRenderData& SrcLODData, int32 SourceVertIdx, const FMergeSectionInfo& MergeSectionInfo);
 
 	/** Copy skin weight info from source LOD model - templatized per SourceLODModel extra bone influence */
-	template<typename SkinWeightType, bool bHasExtraBoneInfluences>
+	template<typename SkinWeightType, bool bHasExtraBoneInfluences, typename BoneIndexType>
 	void CopyWeightFromSource(SkinWeightType& DestWeight, const FSkeletalMeshLODRenderData& SrcLODData, int32 SourceVertIdx, const FMergeSectionInfo& MergeSectionInfo);
 };
diff --git a/Engine/Source/Runtime/Engine/Public/StaticMeshVertexData.h b/Engine/Source/Runtime/Engine/Public/StaticMeshVertexData.h
--- a/Engine/Source/Runtime/Engine/Public/StaticMeshVertexData.h
+++ b/Engine/Source/Runtime/Engine/Public/StaticMeshVertexData.h
@@ -94,9 +94,9 @@ public:
 	* @param Ar - archive to serialize to
 	* @param B - data to serialize
 	*/
-	void Serialize(FArchive& Ar) override
+	void Serialize(FArchive& Ar, bool bForcePerElementSerialization = false) override
 	{
-		Data.TResourceArray<VertexDataType, VERTEXBUFFER_ALIGNMENT>::BulkSerialize(Ar);
+		Data.TResourceArray<VertexDataType, VERTEXBUFFER_ALIGNMENT>::BulkSerialize(Ar, bForcePerElementSerialization);
 	}
 	/**
 	* Assignment. This is currently the only method which allows for 
diff --git a/Engine/Source/Runtime/LiveLinkInterface/Private/Roles/LiveLinkAnimationBlueprintStructs.cpp b/Engine/Source/Runtime/LiveLinkInterface/Private/Roles/LiveLinkAnimationBlueprintStructs.cpp
--- a/Engine/Source/Runtime/LiveLinkInterface/Private/Roles/LiveLinkAnimationBlueprintStructs.cpp
+++ b/Engine/Source/Runtime/LiveLinkInterface/Private/Roles/LiveLinkAnimationBlueprintStructs.cpp
@@ -23,25 +23,44 @@ FCachedSubjectFrame::FCachedSubjectFrame(const FLiveLinkSkeletonStaticData* InSt
 	check(InStaticData->BoneNames.Num() == NumTransforms);
 	check(InStaticData->BoneParents.Num() == NumTransforms);
 	check(InStaticData->PropertyNames.Num() == InAnimData->PropertyValues.Num());
-	RootSpaceTransforms.SetNum(NumTransforms);
-	ChildTransformIndices.SetNum(NumTransforms);
+	CachedRootSpaceTransforms.SetNum(NumTransforms);
+	CachedChildTransformIndices.SetNum(NumTransforms);
 	for (int32 i = 0; i < NumTransforms; ++i)
 	{
-		RootSpaceTransforms[i].Key = false;
-		ChildTransformIndices[i].Key = false;
+		CachedRootSpaceTransforms[i].Key = false;
+		CachedChildTransformIndices[i].Key = false;
 	}
 }
 
-void FCachedSubjectFrame::SetCurvesFromCache(TMap<FName, float>& OutCurves)
+void FCachedSubjectFrame::SetCurvesFromCache(TMap<FName, float>& OutCurves) const
 {
 	if (!bHaveCachedCurves)
 	{
+		// This caching does not change the result of the function 
 		CacheCurves();
 	}
 	OutCurves = CachedCurves;
 };
 
-void FCachedSubjectFrame::GetSubjectMetadata(FSubjectMetadata& OutSubjectMetadata)
+bool FCachedSubjectFrame::GetCurveValueByName(FName InCurveName, float& OutCurveValue) const
+{
+	if (!bHaveCachedCurves)
+	{
+		// The caching does not change the result of the function 
+		CacheCurves();
+	}
+	if (const float* ValuePtr = CachedCurves.Find(InCurveName))
+	{
+		if (FMath::IsFinite(*ValuePtr))
+		{
+			OutCurveValue = *ValuePtr;
+			return true;
+		}
+	}
+	return false;
+}
+
+void FCachedSubjectFrame::GetSubjectMetadata(FSubjectMetadata& OutSubjectMetadata) const
 {
 PRAGMA_DISABLE_DEPRECATION_WARNINGS
 	OutSubjectMetadata.StringMetadata = SourceAnimationFrameData.MetaData.StringMetaData;
@@ -51,17 +70,17 @@ PRAGMA_DISABLE_DEPRECATION_WARNINGS
 PRAGMA_ENABLE_DEPRECATION_WARNINGS
 };
 
-int32 FCachedSubjectFrame::GetNumberOfTransforms()
+int32 FCachedSubjectFrame::GetNumberOfTransforms() const
 {
 	return SourceAnimationFrameData.Transforms.Num();
 };
 
-void FCachedSubjectFrame::GetTransformNames(TArray<FName>& OutTransformNames)
+void FCachedSubjectFrame::GetTransformNames(TArray<FName>& OutTransformNames) const
 {
 	OutTransformNames = SourceSkeletonData.GetBoneNames();
 };
 
-void FCachedSubjectFrame::GetTransformName(const int32 InTransformIndex, FName& OutName)
+void FCachedSubjectFrame::GetTransformName(const int32 InTransformIndex, FName& OutName) const
 {
 	if (IsValidTransformIndex(InTransformIndex))
 	{
@@ -73,12 +92,12 @@ void FCachedSubjectFrame::GetTransformName(const int32 InTransformIndex, FName&
 	}
 };
 
-int32 FCachedSubjectFrame::GetTransformIndexFromName(FName InTransformName)
+int32 FCachedSubjectFrame::GetTransformIndexFromName(FName InTransformName) const
 {
 	return SourceSkeletonData.GetBoneNames().IndexOfByKey(InTransformName);
 };
 
-int32 FCachedSubjectFrame::GetParentTransformIndex(const int32 InTransformIndex)
+int32 FCachedSubjectFrame::GetParentTransformIndex(const int32 InTransformIndex) const
 {
 	if (IsValidTransformIndex(InTransformIndex))
 	{
@@ -90,12 +109,12 @@ int32 FCachedSubjectFrame::GetParentTransformIndex(const int32 InTransformIndex)
 	}
 };
 
-void FCachedSubjectFrame::GetChildTransformIndices(const int32 InTransformIndex, TArray<int32>& OutChildIndices)
+void FCachedSubjectFrame::GetChildTransformIndices(const int32 InTransformIndex, TArray<int32>& OutChildIndices) const
 {
 	OutChildIndices.Reset();
 	if (IsValidTransformIndex(InTransformIndex))
 	{
-		TPair<bool, TArray<int32>>& ChildIndicesCache = ChildTransformIndices[InTransformIndex];
+		TPair<bool, TArray<int32>>& ChildIndicesCache = CachedChildTransformIndices[InTransformIndex];
 		bool bHasValidCache = ChildIndicesCache.Key;
 		TArray<int32>& CachedChildIndices = ChildIndicesCache.Value;
 		if (!bHasValidCache)
@@ -115,7 +134,7 @@ void FCachedSubjectFrame::GetChildTransformIndices(const int32 InTransformIndex,
 	}
 }
 
-void FCachedSubjectFrame::GetTransformParentSpace(const int32 InTransformIndex, FTransform& OutTransform)
+void FCachedSubjectFrame::GetTransformParentSpace(const int32 InTransformIndex, FTransform& OutTransform) const
 {
 	// Case: Root joint or invalid
 	OutTransform = FTransform::Identity;
@@ -126,13 +145,13 @@ void FCachedSubjectFrame::GetTransformParentSpace(const int32 InTransformIndex,
 	}
 };
 
-void FCachedSubjectFrame::GetTransformRootSpace(const int32 InTransformIndex, FTransform& OutTransform)
+void FCachedSubjectFrame::GetTransformRootSpace(const int32 InTransformIndex, FTransform& OutTransform) const
 {
 	// Case: Root joint or invalid
 	OutTransform = FTransform::Identity;
 	if (IsValidTransformIndex(InTransformIndex))
 	{
-		TPair<bool, FTransform>& RootSpaceCache = RootSpaceTransforms[InTransformIndex];
+		TPair<bool, FTransform>& RootSpaceCache = CachedRootSpaceTransforms[InTransformIndex];
 		bool bHasValidCache = RootSpaceCache.Key;
 		// Case: Have Cached Value
 		if (bHasValidCache)
@@ -145,7 +164,7 @@ void FCachedSubjectFrame::GetTransformRootSpace(const int32 InTransformIndex, FT
 			const TArray<int32>& BoneParents = SourceSkeletonData.GetBoneParents();
 			int32 ParentIndex = BoneParents[InTransformIndex];
 
-			FTransform& LocalSpaceTransform = SourceAnimationFrameData.Transforms[InTransformIndex];
+			const FTransform& LocalSpaceTransform = SourceAnimationFrameData.Transforms[InTransformIndex];
 
 			FTransform ParentRootSpaceTransform;
 			GetTransformRootSpace(ParentIndex, ParentRootSpaceTransform);
@@ -159,13 +178,13 @@ void FCachedSubjectFrame::GetTransformRootSpace(const int32 InTransformIndex, FT
 	}
 };
 
-int32 FCachedSubjectFrame::GetRootIndex()
+int32 FCachedSubjectFrame::GetRootIndex() const
 {
 	const TArray<int32>& BoneParents = SourceSkeletonData.GetBoneParents();
 	return BoneParents.IndexOfByPredicate([](const int32& ParentIndex) { return ParentIndex < 0; });
 };
 
-void FCachedSubjectFrame::CacheCurves()
+void FCachedSubjectFrame::CacheCurves() const
 {
 	bHaveCachedCurves = false;
 	CachedCurves.Reset();
@@ -183,7 +202,7 @@ void FCachedSubjectFrame::CacheCurves()
 	}
 };
 
-bool FCachedSubjectFrame::IsValidTransformIndex(int32 InTransformIndex)
+bool FCachedSubjectFrame::IsValidTransformIndex(int32 InTransformIndex) const
 {
 	return (InTransformIndex >= 0) && (InTransformIndex < SourceAnimationFrameData.Transforms.Num());
 };
@@ -195,9 +214,10 @@ FLiveLinkTransform::FLiveLinkTransform()
 // to evaluate as identity
 	: TransformIndex(-1)
 {
+
 };
 
-void FLiveLinkTransform::GetName(FName& OutName)
+void FLiveLinkTransform::GetName(FName& OutName) const
 {
 	if (CachedFrame.IsValid())
 	{
@@ -205,7 +225,7 @@ void FLiveLinkTransform::GetName(FName& OutName)
 	}
 };
 
-void FLiveLinkTransform::GetTransformParentSpace(FTransform& OutTransform)
+void FLiveLinkTransform::GetTransformParentSpace(FTransform& OutTransform) const
 {
 	if (CachedFrame.IsValid())
 	{
@@ -213,7 +233,7 @@ void FLiveLinkTransform::GetTransformParentSpace(FTransform& OutTransform)
 	}
 };
 
-void FLiveLinkTransform::GetTransformRootSpace(FTransform& OutTransform)
+void FLiveLinkTransform::GetTransformRootSpace(FTransform& OutTransform) const
 {
 	if (CachedFrame.IsValid())
 	{
@@ -221,7 +241,7 @@ void FLiveLinkTransform::GetTransformRootSpace(FTransform& OutTransform)
 	}
 };
 
-bool FLiveLinkTransform::HasParent()
+bool FLiveLinkTransform::HasParent() const
 {
 	if (CachedFrame.IsValid())
 	{
@@ -233,7 +253,7 @@ bool FLiveLinkTransform::HasParent()
 	}
 };
 
-void FLiveLinkTransform::GetParent(FLiveLinkTransform& OutParentTransform)
+void FLiveLinkTransform::GetParent(FLiveLinkTransform& OutParentTransform) const
 {
 	if (CachedFrame.IsValid())
 	{
@@ -243,7 +263,7 @@ void FLiveLinkTransform::GetParent(FLiveLinkTransform& OutParentTransform)
 	}
 };
 
-int32 FLiveLinkTransform::GetChildCount()
+int32 FLiveLinkTransform::GetChildCount() const
 {
 	if (CachedFrame.IsValid())
 	{
@@ -257,7 +277,7 @@ int32 FLiveLinkTransform::GetChildCount()
 	}
 };
 
-void FLiveLinkTransform::GetChildren(TArray<FLiveLinkTransform>& OutChildTransforms)
+void FLiveLinkTransform::GetChildren(TArray<FLiveLinkTransform>& OutChildTransforms) const
 {
 	OutChildTransforms.Reset();
 	if (CachedFrame.IsValid())
@@ -290,7 +310,7 @@ int32 FLiveLinkTransform::GetTransformIndex() const
 
 // FSubjectFrameHandle
 
-void FSubjectFrameHandle::GetCurves(TMap<FName, float>& OutCurves)
+void FSubjectFrameHandle::GetCurves(TMap<FName, float>& OutCurves) const
 {
 	if (CachedFrame.IsValid())
 	{
@@ -298,7 +318,16 @@ void FSubjectFrameHandle::GetCurves(TMap<FName, float>& OutCurves)
 	}
 };
 
-void FSubjectFrameHandle::GetSubjectMetadata(FSubjectMetadata& OutMetadata)
+bool FSubjectFrameHandle::GetCurveValueByName(FName CurveName, float& CurveValue) const
+{
+	if (CachedFrame.IsValid())
+	{
+		return CachedFrame->GetCurveValueByName(CurveName, CurveValue);
+	}
+	return false;
+}
+
+void FSubjectFrameHandle::GetSubjectMetadata(FSubjectMetadata& OutMetadata) const
 {
 	if (CachedFrame.IsValid())
 	{
@@ -306,7 +335,7 @@ void FSubjectFrameHandle::GetSubjectMetadata(FSubjectMetadata& OutMetadata)
 	}
 };
 
-int32 FSubjectFrameHandle::GetNumberOfTransforms()
+int32 FSubjectFrameHandle::GetNumberOfTransforms() const
 {
 	if (CachedFrame.IsValid())
 	{
@@ -318,7 +347,7 @@ int32 FSubjectFrameHandle::GetNumberOfTransforms()
 	}
 };
 
-void FSubjectFrameHandle::GetTransformNames(TArray<FName>& OutTransformNames)
+void FSubjectFrameHandle::GetTransformNames(TArray<FName>& OutTransformNames) const
 {
 	if (CachedFrame.IsValid())
 	{
@@ -326,7 +355,7 @@ void FSubjectFrameHandle::GetTransformNames(TArray<FName>& OutTransformNames)
 	}
 };
 
-void FSubjectFrameHandle::GetRootTransform(FLiveLinkTransform& OutLiveLinkTransform)
+void FSubjectFrameHandle::GetRootTransform(FLiveLinkTransform& OutLiveLinkTransform) const
 {
 	OutLiveLinkTransform.SetCachedFrame(CachedFrame);
 	if (CachedFrame.IsValid())
@@ -335,20 +364,28 @@ void FSubjectFrameHandle::GetRootTransform(FLiveLinkTransform& OutLiveLinkTransf
 	}
 }
 
-void FSubjectFrameHandle::GetTransformByIndex(int32 InTransformIndex, FLiveLinkTransform& OutLiveLinkTransform)
+void FSubjectFrameHandle::GetTransformByIndex(int32 InTransformIndex, FLiveLinkTransform& OutLiveLinkTransform) const
 {
 	OutLiveLinkTransform.SetCachedFrame(CachedFrame);
-	OutLiveLinkTransform.SetTransformIndex(InTransformIndex);
+	if (CachedFrame.IsValid())
+	{
+		OutLiveLinkTransform.SetCachedFrame(CachedFrame);
+		OutLiveLinkTransform.SetTransformIndex(InTransformIndex);
+	}
 };
 
-void FSubjectFrameHandle::GetTransformByName(FName InTransformName, FLiveLinkTransform& OutLiveLinkTransform)
+void FSubjectFrameHandle::GetTransformByName(FName InTransformName, FLiveLinkTransform& OutLiveLinkTransform) const
 {
-	int32 TransformIndex = CachedFrame->GetTransformIndexFromName(InTransformName);
 	OutLiveLinkTransform.SetCachedFrame(CachedFrame);
-	OutLiveLinkTransform.SetTransformIndex(TransformIndex);
+	if (CachedFrame.IsValid())
+	{
+		int32 TransformIndex = CachedFrame->GetTransformIndexFromName(InTransformName);
+		OutLiveLinkTransform.SetCachedFrame(CachedFrame);
+		OutLiveLinkTransform.SetTransformIndex(TransformIndex);
+	}
 };
 
-FLiveLinkSkeletonStaticData* FSubjectFrameHandle::GetSourceSkeletonStaticData()
+const FLiveLinkSkeletonStaticData* FSubjectFrameHandle::GetSourceSkeletonStaticData() const
 {
 	if (CachedFrame.IsValid())
 	{
@@ -357,7 +394,7 @@ FLiveLinkSkeletonStaticData* FSubjectFrameHandle::GetSourceSkeletonStaticData()
 	return nullptr;
 }
 
-FLiveLinkAnimationFrameData* FSubjectFrameHandle::GetSourceAnimationFrameData()
+const FLiveLinkAnimationFrameData* FSubjectFrameHandle::GetSourceAnimationFrameData() const
 {
 	if (CachedFrame.IsValid())
 	{
diff --git a/Engine/Source/Runtime/LiveLinkInterface/Public/Roles/LiveLinkAnimationBlueprintStructs.h b/Engine/Source/Runtime/LiveLinkInterface/Public/Roles/LiveLinkAnimationBlueprintStructs.h
--- a/Engine/Source/Runtime/LiveLinkInterface/Public/Roles/LiveLinkAnimationBlueprintStructs.h
+++ b/Engine/Source/Runtime/LiveLinkInterface/Public/Roles/LiveLinkAnimationBlueprintStructs.h
@@ -15,14 +15,14 @@ struct FSubjectMetadata
 {
 	GENERATED_USTRUCT_BODY()
 
-		UPROPERTY(BlueprintReadWrite, EditAnywhere, Category = "LiveLink")
-		TMap<FName, FString> StringMetadata;
+	UPROPERTY(BlueprintReadWrite, EditAnywhere, Category = "LiveLink")
+	TMap<FName, FString> StringMetadata;
 
 	UPROPERTY(BlueprintReadWrite, EditAnywhere, Category = "LiveLink")
-		FTimecode SceneTimecode;
+	FTimecode SceneTimecode;
 
 	UPROPERTY(BlueprintReadWrite, EditAnywhere, Category = "LiveLink")
-		FFrameRate SceneFramerate;
+	FFrameRate SceneFramerate;
 };
 
 USTRUCT()
@@ -30,52 +30,53 @@ struct FCachedSubjectFrame
 {
 	GENERATED_USTRUCT_BODY()
 
-		FCachedSubjectFrame();
+	FCachedSubjectFrame();
 
 	FCachedSubjectFrame(const FLiveLinkSkeletonStaticData* InStaticData, const FLiveLinkAnimationFrameData* InAnimData);
 
 	virtual ~FCachedSubjectFrame() = default;
 
-	void SetCurvesFromCache(TMap<FName, float>& OutCurves);
+	void SetCurvesFromCache(TMap<FName, float>& OutCurves) const;
+
+	bool GetCurveValueByName(FName InCurveName, float& OutCurveValue) const;
 
-	void GetSubjectMetadata(FSubjectMetadata& OutSubjectMetadata);
+	void GetSubjectMetadata(FSubjectMetadata& OutSubjectMetadata) const;
 
-	int32 GetNumberOfTransforms();
+	int32 GetNumberOfTransforms() const;
 
-	void GetTransformNames(TArray<FName>& OutTransformNames);
+	void GetTransformNames(TArray<FName>& OutTransformNames) const;
 
-	void GetTransformName(const int32 InTransformIndex, FName& OutName);
+	void GetTransformName(const int32 InTransformIndex, FName& OutName) const;
 
-	int32 GetTransformIndexFromName(FName InTransformName);
+	int32 GetTransformIndexFromName(FName InTransformName) const;
 
-	int32 GetParentTransformIndex(const int32 InTransformIndex);
+	int32 GetParentTransformIndex(const int32 InTransformIndex) const;
 
-	void GetChildTransformIndices(const int32 InTransformIndex, TArray<int32>& OutChildIndices);
+	void GetChildTransformIndices(const int32 InTransformIndex, TArray<int32>& OutChildIndices) const;
 
-	void GetTransformParentSpace(const int32 InTransformIndex, FTransform& OutTransform);
+	void GetTransformParentSpace(const int32 InTransformIndex, FTransform& OutTransform) const;
 
-	void GetTransformRootSpace(const int32 InTransformIndex, FTransform& OutTransform);
+	void GetTransformRootSpace(const int32 InTransformIndex, FTransform& OutTransform) const;
 
-	int32 GetRootIndex();
+	int32 GetRootIndex() const;
 
-	FLiveLinkSkeletonStaticData& GetSourceSkeletonData() { return SourceSkeletonData; }
 	const FLiveLinkSkeletonStaticData& GetSourceSkeletonData() const { return SourceSkeletonData; }
 
-	FLiveLinkAnimationFrameData& GetSourceAnimationFrameData() { return SourceAnimationFrameData; }
 	const FLiveLinkAnimationFrameData& GetSourceAnimationFrameData() const { return SourceAnimationFrameData; }
 
 private:
 	FLiveLinkSkeletonStaticData SourceSkeletonData;
 	FLiveLinkAnimationFrameData SourceAnimationFrameData;
-	TArray<TPair<bool, FTransform>> RootSpaceTransforms;
-	TArray<TPair<bool, TArray<int32>>> ChildTransformIndices;
-	TMap<FName, float> CachedCurves;
 
-	bool bHaveCachedCurves;
+	// Cache. Does not affect the underlying frame data
+	mutable TArray<TPair<bool, FTransform>> CachedRootSpaceTransforms;
+	mutable TArray<TPair<bool, TArray<int32>>> CachedChildTransformIndices;
+	mutable TMap<FName, float> CachedCurves;
+	mutable bool bHaveCachedCurves;
 
-	void CacheCurves();
+	void CacheCurves() const;
 
-	bool IsValidTransformIndex(int32 InTransformIndex);
+	bool IsValidTransformIndex(int32 InTransformIndex) const;
 };
 
 USTRUCT(BlueprintType)
@@ -83,23 +84,23 @@ struct LIVELINKINTERFACE_API FLiveLinkTransform
 {
 	GENERATED_USTRUCT_BODY()
 
-		FLiveLinkTransform();
+	FLiveLinkTransform();
 
 	virtual ~FLiveLinkTransform() = default;
 
-	void GetName(FName& Name);
+	void GetName(FName& Name) const;
 
-	void GetTransformParentSpace(FTransform& OutTransform);
+	void GetTransformParentSpace(FTransform& OutTransform) const;
 
-	void GetTransformRootSpace(FTransform& OutTransform);
+	void GetTransformRootSpace(FTransform& OutTransform) const;
 
-	bool HasParent();
+	bool HasParent() const;
 
-	void GetParent(FLiveLinkTransform& OutParentTransform);
+	void GetParent(FLiveLinkTransform& OutParentTransform) const;
 
-	int32 GetChildCount();
+	int32 GetChildCount() const;
 
-	void GetChildren(TArray<FLiveLinkTransform>& OutChildTransforms);
+	void GetChildren(TArray<FLiveLinkTransform>& OutChildTransforms) const;
 
 	void SetCachedFrame(TSharedPtr<FCachedSubjectFrame> InCachedFrame);
 
@@ -117,29 +118,31 @@ struct LIVELINKINTERFACE_API FSubjectFrameHandle : public FLiveLinkBaseBlueprint
 {
 	GENERATED_USTRUCT_BODY()
 
-		FSubjectFrameHandle() = default;
+	FSubjectFrameHandle() = default;
 
 	virtual ~FSubjectFrameHandle() = default;
 
-	void GetCurves(TMap<FName, float>& OutCurves);
+	void GetCurves(TMap<FName, float>& OutCurves) const;
+
+	bool GetCurveValueByName(FName CurveName, float& CurveValue) const;
 
-	void GetSubjectMetadata(FSubjectMetadata& OutMetadata);
+	void GetSubjectMetadata(FSubjectMetadata& OutMetadata) const;
 
-	int32 GetNumberOfTransforms();
+	int32 GetNumberOfTransforms() const;
 
-	void GetTransformNames(TArray<FName>& OutTransformNames);
+	void GetTransformNames(TArray<FName>& OutTransformNames) const;
 
-	void GetRootTransform(FLiveLinkTransform& OutLiveLinkTransform);
+	void GetRootTransform(FLiveLinkTransform& OutLiveLinkTransform) const;
 
-	void GetTransformByIndex(int32 InTransformIndex, FLiveLinkTransform& OutLiveLinkTransform);
+	void GetTransformByIndex(int32 InTransformIndex, FLiveLinkTransform& OutLiveLinkTransform) const;
 
-	void GetTransformByName(FName InTransformName, FLiveLinkTransform& OutLiveLinkTransform);
+	void GetTransformByName(FName InTransformName, FLiveLinkTransform& OutLiveLinkTransform) const;
 
 	void SetCachedFrame(TSharedPtr<FCachedSubjectFrame> InCachedFrame);
 
-	FLiveLinkSkeletonStaticData* GetSourceSkeletonStaticData();
+	const FLiveLinkSkeletonStaticData* GetSourceSkeletonStaticData() const;
 
-	FLiveLinkAnimationFrameData* GetSourceAnimationFrameData();
+	const FLiveLinkAnimationFrameData* GetSourceAnimationFrameData() const;
 
 private:
 	TSharedPtr<FCachedSubjectFrame> CachedFrame;
diff --git a/Engine/Source/Runtime/MovieSceneTracks/Private/Evaluation/MovieSceneParameterTemplate.cpp b/Engine/Source/Runtime/MovieSceneTracks/Private/Evaluation/MovieSceneParameterTemplate.cpp
--- a/Engine/Source/Runtime/MovieSceneTracks/Private/Evaluation/MovieSceneParameterTemplate.cpp
+++ b/Engine/Source/Runtime/MovieSceneTracks/Private/Evaluation/MovieSceneParameterTemplate.cpp
@@ -8,6 +8,8 @@
 
 FMovieSceneParameterSectionTemplate::FMovieSceneParameterSectionTemplate(const UMovieSceneParameterSection& Section)
 	: Scalars(Section.GetScalarParameterNamesAndCurves())
+	, Bools(Section.GetBoolParameterNamesAndCurves())
+	, Vector2Ds(Section.GetVector2DParameterNamesAndCurves())
 	, Vectors(Section.GetVectorParameterNamesAndCurves())
 	, Colors(Section.GetColorParameterNamesAndCurves())
 	, Transforms(Section.GetTransformParameterNamesAndCurves())
@@ -27,6 +29,30 @@ void FMovieSceneParameterSectionTemplate::EvaluateCurves(const FMovieSceneContex
 		}
 	}
 
+	for (const FBoolParameterNameAndCurve& Bool : Bools)
+	{
+		bool Value = false;
+		if (Bool.ParameterCurve.Evaluate(Time, Value))
+		{
+			Values.BoolValues.Emplace(Bool.ParameterName, Value);
+		}
+	}
+
+	for (const FVector2DParameterNameAndCurves& Vector : Vector2Ds)
+	{
+		FVector2D Value(ForceInitToZero);
+
+		bool bAnyEvaluated = false;
+		bAnyEvaluated |= Vector.XCurve.Evaluate(Time, Value.X);
+		bAnyEvaluated |= Vector.YCurve.Evaluate(Time, Value.Y);
+
+		if (bAnyEvaluated)
+		{
+			Values.Vector2DValues.Emplace(Vector.ParameterName, Value);
+		}
+	}
+
+
 	for ( const FVectorParameterNameAndCurves& Vector : Vectors )
 	{
 		FVector Value(ForceInitToZero);
diff --git a/Engine/Source/Runtime/MovieSceneTracks/Private/Evaluation/MovieSceneSkeletalAnimationTemplate.cpp b/Engine/Source/Runtime/MovieSceneTracks/Private/Evaluation/MovieSceneSkeletalAnimationTemplate.cpp
--- a/Engine/Source/Runtime/MovieSceneTracks/Private/Evaluation/MovieSceneSkeletalAnimationTemplate.cpp
+++ b/Engine/Source/Runtime/MovieSceneTracks/Private/Evaluation/MovieSceneSkeletalAnimationTemplate.cpp
@@ -8,6 +8,7 @@
 #include "Components/SkeletalMeshComponent.h"
 #include "Animation/AnimSingleNodeInstance.h"
 #include "Runtime/AnimGraphRuntime/Public/AnimSequencerInstance.h"
+#include "AnimCustomInstanceHelper.h"
 #include "Evaluation/MovieSceneEvaluation.h"
 #include "IMovieScenePlayer.h"
 #include "UObject/ObjectKey.h"
@@ -31,7 +32,30 @@ bool CanPlayAnimation(USkeletalMeshComponent* SkeletalMeshComponent, UAnimSequen
 
 void ResetAnimSequencerInstance(UObject& ObjectToRestore, IMovieScenePlayer& Player)
 {
-	CastChecked<UAnimSequencerInstance>(&ObjectToRestore)->ResetNodes();
+	CastChecked<ISequencerAnimationSupport>(&ObjectToRestore)->ResetNodes();
+}
+
+UAnimSequencerInstance* GetAnimSequencerInstance(USkeletalMeshComponent* SkeletalMeshComponent) 
+{
+	ISequencerAnimationSupport* SeqInterface = Cast<ISequencerAnimationSupport>(SkeletalMeshComponent->GetAnimInstance());
+	if (SeqInterface)
+	{
+		return Cast<UAnimSequencerInstance>(SeqInterface->GetSourceAnimInstance());
+	}
+
+	return nullptr;
+}
+
+UAnimInstance* GetSourceAnimInstance(USkeletalMeshComponent* SkeletalMeshComponent) 
+{
+	UAnimInstance* SkelAnimInstance = SkeletalMeshComponent->GetAnimInstance();
+	ISequencerAnimationSupport* SeqInterface = Cast<ISequencerAnimationSupport>(SkelAnimInstance);
+	if (SeqInterface)
+	{
+		return SeqInterface->GetSourceAnimInstance();
+	}
+
+	return SkelAnimInstance;
 }
 
 struct FStopPlayingMontageTokenProducer : IMovieScenePreAnimatedTokenProducer
@@ -94,14 +118,14 @@ struct FPreAnimatedAnimationTokenProducer : IMovieScenePreAnimatedTokenProducer
 			{
 				USkeletalMeshComponent* Component = CastChecked<USkeletalMeshComponent>(&ObjectToRestore);
 
-				UAnimSequencerInstance* SequencerInst = Cast<UAnimSequencerInstance>(Component->GetAnimInstance());
+				ISequencerAnimationSupport* SequencerInst = Cast<ISequencerAnimationSupport>(GetAnimSequencerInstance(Component));
 				if (SequencerInst)
 				{
 					SequencerInst->ResetPose();
 					SequencerInst->ResetNodes();
 				}
 
-				UAnimSequencerInstance::UnbindFromSkeletalMeshComponent(Component);
+				FAnimCustomInstanceHelper::UnbindFromSkeletalMeshComponent<UAnimSequencerInstance>(Component);
 
 				// Reset the mesh component update flag and animation mode to what they were before we animated the object
 				Component->VisibilityBasedAnimTickOption = VisibilityBasedAnimTickOption;
@@ -243,9 +267,9 @@ namespace MovieScene
 			}
 			OriginalStack.SavePreAnimatedState(Player, *SkeletalMeshComponent, GetAnimControlTypeID(), FPreAnimatedAnimationTokenProducer());
 
-			UAnimInstance* ExistingAnimInstance = SkeletalMeshComponent->GetAnimInstance();
+			UAnimInstance* ExistingAnimInstance = GetSourceAnimInstance(SkeletalMeshComponent);
 			bool bWasCreated = false;
-			UAnimSequencerInstance* SequencerInstance = UAnimCustomInstance::BindToSkeletalMeshComponent<UAnimSequencerInstance>(SkeletalMeshComponent,bWasCreated);
+			ISequencerAnimationSupport* SequencerInstance = FAnimCustomInstanceHelper::BindToSkeletalMeshComponent<UAnimSequencerInstance>(SkeletalMeshComponent,bWasCreated);
 			if (SequencerInstance && bWasCreated)
 			{
 				SequencerInstance->SavePose();
@@ -372,7 +396,7 @@ namespace MovieScene
 				}
 			}
 		}
-
+		
 		void SetAnimPosition(FPersistentEvaluationData& PersistentData, IMovieScenePlayer& Player, USkeletalMeshComponent* SkeletalMeshComponent, FName SlotName, FObjectKey Section, UAnimSequenceBase* InAnimSequence, float InFromPosition, float InToPosition, float Weight, bool bPlaying, bool bFireNotifies, bool bForceCustomMode)
 		{
 			static const bool bLooping = false;
@@ -385,7 +409,7 @@ namespace MovieScene
 			{
 				SkeletalMeshComponent->SetAnimationMode(EAnimationMode::AnimationCustomMode);
 			}
-			UAnimSequencerInstance* SequencerInst = Cast<UAnimSequencerInstance>(SkeletalMeshComponent->GetAnimInstance());
+			UAnimSequencerInstance* SequencerInst = GetAnimSequencerInstance(SkeletalMeshComponent);
 			if (SequencerInst)
 			{
 				FMovieSceneAnimTypeID AnimTypeID = SectionToAnimationIDs.GetAnimTypeID(Section);
@@ -395,14 +419,14 @@ namespace MovieScene
 				// Set position and weight
 				SequencerInst->UpdateAnimTrack(InAnimSequence, GetTypeHash(AnimTypeID), InFromPosition, InToPosition, Weight, bFireNotifies);
 			}
-			else if (UAnimInstance* AnimInst = SkeletalMeshComponent->GetAnimInstance())
+			else if (UAnimInstance* AnimInst = GetSourceAnimInstance(SkeletalMeshComponent))
 			{
 				FMontagePlayerPerSectionData* SectionData = MontageData.Find(Section);
 
 				int32 InstanceId = (SectionData) ? SectionData->MontageInstanceId : INDEX_NONE;
 
 				const float AssetPlayRate = FMath::IsNearlyZero(InAnimSequence->RateScale) ? 1.0f : InAnimSequence->RateScale;
-				TWeakObjectPtr<UAnimMontage> Montage = FAnimMontageInstance::SetSequencerMontagePosition(SlotName, SkeletalMeshComponent, InstanceId, InAnimSequence, InFromPosition / AssetPlayRate, InToPosition / AssetPlayRate, Weight, bLooping, bPlaying);
+				TWeakObjectPtr<UAnimMontage> Montage = FAnimMontageInstance::SetSequencerMontagePosition(SlotName, AnimInst, InstanceId, InAnimSequence, InFromPosition / AssetPlayRate, InToPosition / AssetPlayRate, Weight, bLooping, bPlaying);
 
 				if (Montage.IsValid())
 				{
@@ -432,7 +456,7 @@ namespace MovieScene
 			{
 				SkeletalMeshComponent->SetAnimationMode(EAnimationMode::AnimationCustomMode);
 			}
-			UAnimSequencerInstance* SequencerInst = Cast<UAnimSequencerInstance>(SkeletalMeshComponent->GetAnimInstance());
+			UAnimSequencerInstance* SequencerInst = GetAnimSequencerInstance(SkeletalMeshComponent);
 			if (SequencerInst)
 			{
 				// Unique anim type ID per slot
@@ -442,14 +466,14 @@ namespace MovieScene
 				// Set position and weight
 				SequencerInst->UpdateAnimTrack(InAnimSequence, GetTypeHash(AnimTypeID), InFromPosition, InToPosition, Weight, bFireNotifies);
 			}
-			else if (UAnimInstance* AnimInst = SkeletalMeshComponent->GetAnimInstance())
+			else if (UAnimInstance* AnimInst = GetSourceAnimInstance(SkeletalMeshComponent))
 			{
 				FMontagePlayerPerSectionData* SectionData = MontageData.Find(Section);
 
 				int32 InstanceId = (SectionData)? SectionData->MontageInstanceId : INDEX_NONE;
 
 				const float AssetPlayRate = FMath::IsNearlyZero(InAnimSequence->RateScale) ? 1.0f : InAnimSequence->RateScale;
-				TWeakObjectPtr<UAnimMontage> Montage = FAnimMontageInstance::PreviewSequencerMontagePosition(SlotName, SkeletalMeshComponent, InstanceId, InAnimSequence, InFromPosition / AssetPlayRate, InToPosition / AssetPlayRate, Weight, bLooping, bFireNotifies, bPlaying);
+				TWeakObjectPtr<UAnimMontage> Montage = FAnimMontageInstance::PreviewSequencerMontagePosition(SlotName, SkeletalMeshComponent, AnimInst, InstanceId, InAnimSequence, InFromPosition / AssetPlayRate, InToPosition / AssetPlayRate, Weight, bLooping, bFireNotifies, bPlaying);
 
 				if (Montage.IsValid())
 				{
diff --git a/Engine/Source/Runtime/MovieSceneTracks/Private/Sections/MovieSceneParameterSection.cpp b/Engine/Source/Runtime/MovieSceneTracks/Private/Sections/MovieSceneParameterSection.cpp
--- a/Engine/Source/Runtime/MovieSceneTracks/Private/Sections/MovieSceneParameterSection.cpp
+++ b/Engine/Source/Runtime/MovieSceneTracks/Private/Sections/MovieSceneParameterSection.cpp
@@ -9,7 +9,17 @@ FScalarParameterNameAndCurve::FScalarParameterNameAndCurve( FName InParameterNam
 	ParameterName = InParameterName;
 }
 
-FVectorParameterNameAndCurves::FVectorParameterNameAndCurves( FName InParameterName )
+FBoolParameterNameAndCurve::FBoolParameterNameAndCurve(FName InParameterName)
+{
+	ParameterName = InParameterName;
+}
+
+FVector2DParameterNameAndCurves::FVector2DParameterNameAndCurves( FName InParameterName )
+{
+	ParameterName = InParameterName;
+}
+
+FVectorParameterNameAndCurves::FVectorParameterNameAndCurves(FName InParameterName)
 {
 	ParameterName = InParameterName;
 }
@@ -42,7 +52,8 @@ void UMovieSceneParameterSection::Serialize(FArchive& Ar)
 
 	if (Ar.IsLoading())
 	{
-		ReconstructChannelProxy();
+		//Don't force if transacting, since it may not be a channel creation/deletion change
+		ReconstructChannelProxy(!Ar.IsTransacting());
 	}
 }
 
@@ -50,12 +61,13 @@ void UMovieSceneParameterSection::PostEditImport()
 {
 	Super::PostEditImport();
 
-	ReconstructChannelProxy();
+	ReconstructChannelProxy(true);
 }
 
 
-void UMovieSceneParameterSection::ReconstructChannelProxy()
+void UMovieSceneParameterSection::ReconstructChannelProxy(bool bForce)
 {
+
 	FMovieSceneChannelProxyData Channels;
 
 #if WITH_EDITOR
@@ -67,6 +79,22 @@ void UMovieSceneParameterSection::ReconstructChannelProxy()
 		MetaData.bCanCollapseToTrack = false;
 		Channels.Add(Scalar.ParameterCurve, MetaData, TMovieSceneExternalValue<float>());
 	}
+
+	for (FBoolParameterNameAndCurve& Bool : GetBoolParameterNamesAndCurves())
+	{
+		FMovieSceneChannelMetaData MetaData(Bool.ParameterName, FText::FromName(Bool.ParameterName));
+		// Prevent single channels from collapsing to the track node
+		MetaData.bCanCollapseToTrack = false;
+		Channels.Add(Bool.ParameterCurve, MetaData, TMovieSceneExternalValue<bool>());
+	}
+	for (FVector2DParameterNameAndCurves& Vector2D : GetVector2DParameterNamesAndCurves())
+	{
+		FString ParameterString = Vector2D.ParameterName.ToString();
+		FText Group = FText::FromString(ParameterString);
+
+		Channels.Add(Vector2D.XCurve, FMovieSceneChannelMetaData(*(ParameterString + TEXT(".X")), FCommonChannelData::ChannelX, Group), TMovieSceneExternalValue<float>());
+		Channels.Add(Vector2D.YCurve, FMovieSceneChannelMetaData(*(ParameterString + TEXT(".Y")), FCommonChannelData::ChannelY, Group), TMovieSceneExternalValue<float>());
+	}
 	for (FVectorParameterNameAndCurves& Vector : GetVectorParameterNamesAndCurves())
 	{
 		FString ParameterString = Vector.ParameterName.ToString();
@@ -126,11 +154,14 @@ void UMovieSceneParameterSection::ReconstructChannelProxy()
 	{
 		Channels.Add(Scalar.ParameterCurve);
 	}
-	for (FVectorParameterNameAndCurves& Vector : GetVectorParameterNamesAndCurves())
+	for (FBoolParameterNameAndCurve& Bool : GetBoolParameterNamesAndCurves())
+	{
+		Channels.Add(Bool.ParameterCurve);
+	}
+	for (FVector2DParameterNameAndCurves& Vector2D : GetVector2DParameterNamesAndCurves())
 	{
-		Channels.Add(Vector.XCurve);
-		Channels.Add(Vector.YCurve);
-		Channels.Add(Vector.ZCurve);
+		Channels.Add(Vector2D.XCurve);
+		Channels.Add(Vector2D.YCurve);
 	}
 	for (FColorParameterNameAndCurves& Color : GetColorParameterNamesAndCurves())
 	{
@@ -159,6 +190,7 @@ void UMovieSceneParameterSection::ReconstructChannelProxy()
 #endif
 
 	ChannelProxy = MakeShared<FMovieSceneChannelProxy>(MoveTemp(Channels));
+	
 }
 
 void UMovieSceneParameterSection::AddScalarParameterKey( FName InParameterName, FFrameNumber InTime, float InValue )
@@ -177,7 +209,7 @@ void UMovieSceneParameterSection::AddScalarParameterKey( FName InParameterName,
 		const int32 NewIndex = ScalarParameterNamesAndCurves.Add( FScalarParameterNameAndCurve( InParameterName ) );
 		ExistingChannel = &ScalarParameterNamesAndCurves[NewIndex].ParameterCurve;
 
-		ReconstructChannelProxy();
+		ReconstructChannelProxy(true);
 	}
 
 	ExistingChannel->AddCubicKey(InTime, InValue);
@@ -188,6 +220,62 @@ void UMovieSceneParameterSection::AddScalarParameterKey( FName InParameterName,
 	}
 }
 
+void UMovieSceneParameterSection::AddBoolParameterKey(FName InParameterName, FFrameNumber InTime, bool InValue)
+{
+	FMovieSceneBoolChannel* ExistingChannel = nullptr;
+	for (FBoolParameterNameAndCurve& BoolParameterNameAndCurve : BoolParameterNamesAndCurves)
+	{
+		if (BoolParameterNameAndCurve.ParameterName == InParameterName)
+		{
+			ExistingChannel = &BoolParameterNameAndCurve.ParameterCurve;
+			break;
+		}
+	}
+	if (ExistingChannel == nullptr)
+	{
+		const int32 NewIndex = BoolParameterNamesAndCurves.Add(FBoolParameterNameAndCurve(InParameterName));
+		ExistingChannel = &BoolParameterNamesAndCurves[NewIndex].ParameterCurve;
+
+		ReconstructChannelProxy(true);
+	}
+
+	ExistingChannel->GetData().UpdateOrAddKey(InTime, InValue);
+
+	if (TryModify())
+	{
+		SetRange(TRange<FFrameNumber>::Hull(TRange<FFrameNumber>(InTime), GetRange()));
+	}
+}
+
+
+void UMovieSceneParameterSection::AddVector2DParameterKey(FName InParameterName, FFrameNumber InTime, FVector2D InValue)
+{
+	FVector2DParameterNameAndCurves* ExistingCurves = nullptr;
+	for (FVector2DParameterNameAndCurves& VectorParameterNameAndCurve : Vector2DParameterNamesAndCurves)
+	{
+		if (VectorParameterNameAndCurve.ParameterName == InParameterName)
+		{
+			ExistingCurves = &VectorParameterNameAndCurve;
+			break;
+		}
+	}
+	if (ExistingCurves == nullptr)
+	{
+		int32 NewIndex = Vector2DParameterNamesAndCurves.Add(FVector2DParameterNameAndCurves(InParameterName));
+		ExistingCurves = &Vector2DParameterNamesAndCurves[NewIndex];
+
+		ReconstructChannelProxy(true);
+	}
+
+	ExistingCurves->XCurve.AddCubicKey(InTime, InValue.X);
+	ExistingCurves->YCurve.AddCubicKey(InTime, InValue.Y);
+
+	if (TryModify())
+	{
+		SetRange(TRange<FFrameNumber>::Hull(TRange<FFrameNumber>(InTime), GetRange()));
+	}
+}
+
 void UMovieSceneParameterSection::AddVectorParameterKey( FName InParameterName, FFrameNumber InTime, FVector InValue )
 {
 	FVectorParameterNameAndCurves* ExistingCurves = nullptr;
@@ -204,7 +292,7 @@ void UMovieSceneParameterSection::AddVectorParameterKey( FName InParameterName,
 		int32 NewIndex = VectorParameterNamesAndCurves.Add( FVectorParameterNameAndCurves( InParameterName ) );
 		ExistingCurves = &VectorParameterNamesAndCurves[NewIndex];
 
-		ReconstructChannelProxy();
+		ReconstructChannelProxy(true);
 	}
 
 	ExistingCurves->XCurve.AddCubicKey(InTime, InValue.X);
@@ -233,7 +321,7 @@ void UMovieSceneParameterSection::AddColorParameterKey( FName InParameterName, F
 		int32 NewIndex = ColorParameterNamesAndCurves.Add( FColorParameterNameAndCurves( InParameterName ) );
 		ExistingCurves = &ColorParameterNamesAndCurves[NewIndex];
 
-		ReconstructChannelProxy();
+		ReconstructChannelProxy(true);
 	}
 
 	ExistingCurves->RedCurve.AddCubicKey(   InTime, InValue.R );
@@ -263,7 +351,7 @@ void UMovieSceneParameterSection::AddTransformParameterKey(FName InParameterName
 		int32 NewIndex = TransformParameterNamesAndCurves.Add(FTransformParameterNameAndCurves(InParameterName));
 		ExistingCurves = &TransformParameterNamesAndCurves[NewIndex];
 
-		ReconstructChannelProxy();
+		ReconstructChannelProxy(true);
 	}
 	FVector Translation = InValue.GetTranslation();
 	FRotator Rotator = InValue.GetRotation().Rotator();
@@ -293,7 +381,35 @@ bool UMovieSceneParameterSection::RemoveScalarParameter( FName InParameterName )
 		if ( ScalarParameterNamesAndCurves[i].ParameterName == InParameterName )
 		{
 			ScalarParameterNamesAndCurves.RemoveAt(i);
-			ReconstructChannelProxy();
+			ReconstructChannelProxy(true);
+			return true;
+		}
+	}
+	return false;
+}
+
+bool UMovieSceneParameterSection::RemoveBoolParameter(FName InParameterName)
+{
+	for (int32 i = 0; i < BoolParameterNamesAndCurves.Num(); i++)
+	{
+		if (BoolParameterNamesAndCurves[i].ParameterName == InParameterName)
+		{
+			BoolParameterNamesAndCurves.RemoveAt(i);
+			ReconstructChannelProxy(true);
+			return true;
+		}
+	}
+	return false;
+}
+
+bool UMovieSceneParameterSection::RemoveVector2DParameter(FName InParameterName)
+{
+	for (int32 i = 0; i < Vector2DParameterNamesAndCurves.Num(); i++)
+	{
+		if (Vector2DParameterNamesAndCurves[i].ParameterName == InParameterName)
+		{
+			Vector2DParameterNamesAndCurves.RemoveAt(i);
+			ReconstructChannelProxy(true);
 			return true;
 		}
 	}
@@ -307,7 +423,7 @@ bool UMovieSceneParameterSection::RemoveVectorParameter( FName InParameterName )
 		if ( VectorParameterNamesAndCurves[i].ParameterName == InParameterName )
 		{
 			VectorParameterNamesAndCurves.RemoveAt( i );
-			ReconstructChannelProxy();
+			ReconstructChannelProxy(true);
 			return true;
 		}
 	}
@@ -321,7 +437,7 @@ bool UMovieSceneParameterSection::RemoveColorParameter( FName InParameterName )
 		if ( ColorParameterNamesAndCurves[i].ParameterName == InParameterName )
 		{
 			ColorParameterNamesAndCurves.RemoveAt( i );
-			ReconstructChannelProxy();
+			ReconstructChannelProxy(true);
 			return true;
 		}
 	}
@@ -335,7 +451,7 @@ bool UMovieSceneParameterSection::RemoveTransformParameter(FName InParameterName
 		if (TransformParameterNamesAndCurves[i].ParameterName == InParameterName)
 		{
 			TransformParameterNamesAndCurves.RemoveAt(i);
-			ReconstructChannelProxy();
+			ReconstructChannelProxy(true);
 			return true;
 		}
 	}
@@ -352,7 +468,27 @@ const TArray<FScalarParameterNameAndCurve>& UMovieSceneParameterSection::GetScal
 	return ScalarParameterNamesAndCurves;
 }
 
-TArray<FVectorParameterNameAndCurves>& UMovieSceneParameterSection::GetVectorParameterNamesAndCurves()
+TArray<FBoolParameterNameAndCurve>& UMovieSceneParameterSection::GetBoolParameterNamesAndCurves() 
+{
+	return BoolParameterNamesAndCurves;
+}
+
+const TArray<FBoolParameterNameAndCurve>& UMovieSceneParameterSection::GetBoolParameterNamesAndCurves() const
+{
+	return BoolParameterNamesAndCurves;
+}
+
+TArray<FVector2DParameterNameAndCurves>& UMovieSceneParameterSection::GetVector2DParameterNamesAndCurves()
+{
+	return Vector2DParameterNamesAndCurves;
+}
+
+const TArray<FVector2DParameterNameAndCurves>& UMovieSceneParameterSection::GetVector2DParameterNamesAndCurves() const
+{
+	return Vector2DParameterNamesAndCurves;
+}
+
+TArray<FVectorParameterNameAndCurves>& UMovieSceneParameterSection::GetVectorParameterNamesAndCurves() 
 {
 	return VectorParameterNamesAndCurves;
 }
diff --git a/Engine/Source/Runtime/MovieSceneTracks/Public/Evaluation/MovieSceneParameterTemplate.h b/Engine/Source/Runtime/MovieSceneTracks/Public/Evaluation/MovieSceneParameterTemplate.h
--- a/Engine/Source/Runtime/MovieSceneTracks/Public/Evaluation/MovieSceneParameterTemplate.h
+++ b/Engine/Source/Runtime/MovieSceneTracks/Public/Evaluation/MovieSceneParameterTemplate.h
@@ -32,6 +32,10 @@ struct FEvaluatedParameterSectionValues
 
 	/** Array of evaluated scalar values */
 	TArray<FScalarParameterNameAndValue, TInlineAllocator<2>> ScalarValues;
+	/** Array of evaluated bool values */
+	TArray<FBoolParameterNameAndValue, TInlineAllocator<2>> BoolValues;
+	/** Array of evaluated vector2D values */
+	TArray<FVector2DParameterNameAndValue, TInlineAllocator<2>> Vector2DValues;
 	/** Array of evaluated vector values */
 	TArray<FVectorParameterNameAndValue, TInlineAllocator<2>> VectorValues;
 	/** Array of evaluated color values */
@@ -62,6 +66,14 @@ protected:
 	UPROPERTY()
 	TArray<FScalarParameterNameAndCurve> Scalars;
 
+	/** The bool parameter names and their associated curves. */
+	UPROPERTY()
+	TArray<FBoolParameterNameAndCurve> Bools;
+
+	/** The vector parameter names and their associated curves. */
+	UPROPERTY()
+	TArray<FVector2DParameterNameAndCurves> Vector2Ds;
+
 	/** The vector parameter names and their associated curves. */
 	UPROPERTY()
 	TArray<FVectorParameterNameAndCurves> Vectors;
diff --git a/Engine/Source/Runtime/MovieSceneTracks/Public/Sections/MovieSceneParameterSection.h b/Engine/Source/Runtime/MovieSceneTracks/Public/Sections/MovieSceneParameterSection.h
--- a/Engine/Source/Runtime/MovieSceneTracks/Public/Sections/MovieSceneParameterSection.h
+++ b/Engine/Source/Runtime/MovieSceneTracks/Public/Sections/MovieSceneParameterSection.h
@@ -8,6 +8,7 @@
 #include "Curves/RichCurve.h"
 #include "MovieSceneSection.h"
 #include "Channels/MovieSceneFloatChannel.h"
+#include "Channels/MovieSceneBoolChannel.h"
 #include "Sections/MovieScene3DTransformSection.h"
 #include "MovieSceneParameterSection.generated.h"
 
@@ -33,6 +34,44 @@ struct FScalarParameterNameAndValue
 	float Value;
 };
 
+/**
+ * Structure representing the value of a bool parameter.
+ */
+struct FBoolParameterNameAndValue
+{
+	/** Creates a new FScalarParameterAndValue with a parameter name and a value. */
+	FBoolParameterNameAndValue(FName InParameterName, bool InValue)
+	{
+		ParameterName = InParameterName;
+		Value = InValue;
+	}
+
+	/** The name of the bool parameter. */
+	FName ParameterName;
+
+	/** The value of the bool parameter. */
+	bool Value;
+};
+
+/**
+ * Structure representing the animated value of a vector2D parameter.
+ */
+struct FVector2DParameterNameAndValue
+{
+	/** Creates a new FVectorParameterAndValue with a parameter name and a value. */
+	FVector2DParameterNameAndValue(FName InParameterName, FVector2D InValue)
+	{
+		ParameterName = InParameterName;
+		Value = InValue;
+	}
+
+	/** The name of the vector parameter. */
+	FName ParameterName;
+
+	/** The animated value of the vector parameter. */
+	FVector2D Value;
+};
+
 /**
  * Structure representing the animated value of a vector parameter.
  */
@@ -52,7 +91,6 @@ struct FVectorParameterNameAndValue
 	FVector Value;
 };
 
-
 /**
  * Structure representing the animated value of a color parameter.
  */
@@ -92,6 +130,29 @@ struct FTransformParameterNameAndValue
 	}
 };
 
+/**
+ * Structure representing an bool  parameter and it's associated animation curve.
+ */
+USTRUCT()
+struct MOVIESCENETRACKS_API FBoolParameterNameAndCurve
+{
+	GENERATED_USTRUCT_BODY()
+
+	FBoolParameterNameAndCurve()
+	{}
+
+	/** Creates a new FScalarParameterNameAndCurve for a specific scalar parameter. */
+	FBoolParameterNameAndCurve(FName InParameterName);
+
+	/** The name of the scalar parameter which is being animated. */
+	UPROPERTY()
+	FName ParameterName;
+
+	/** The curve which contains the animation data for the scalar parameter. */
+	UPROPERTY()
+	FMovieSceneBoolChannel ParameterCurve;
+};
+
 /**
  * Structure representing an animated scalar parameter and it's associated animation curve.
  */
@@ -115,6 +176,34 @@ struct MOVIESCENETRACKS_API FScalarParameterNameAndCurve
 	FMovieSceneFloatChannel ParameterCurve;
 };
 
+/**
+ * Structure representing an animated vector2D parameter and it's associated animation curve.
+ */
+USTRUCT()
+struct MOVIESCENETRACKS_API FVector2DParameterNameAndCurves
+{
+	GENERATED_USTRUCT_BODY()
+
+	FVector2DParameterNameAndCurves()
+	{}
+
+	/** Creates a new FVectorParameterNameAndCurve for a specific vector parameter. */
+	FVector2DParameterNameAndCurves(FName InParameterName);
+
+	/** The name of the vector parameter which is being animated. */
+	UPROPERTY()
+	FName ParameterName;
+
+	/** The curve which contains the animation data for the x component of the vector parameter. */
+	UPROPERTY()
+	FMovieSceneFloatChannel XCurve;
+
+	/** The curve which contains the animation data for the y component of the vector parameter. */
+	UPROPERTY()
+	FMovieSceneFloatChannel YCurve;
+
+};
+
 
 /**
  * Structure representing an animated vector parameter and it's associated animation curve.
@@ -224,7 +313,13 @@ class MOVIESCENETRACKS_API UMovieSceneParameterSection
 
 public:
 	/** Adds a a key for a specific scalar parameter at the specified time with the specified value. */
-	 void AddScalarParameterKey(FName InParameterName, FFrameNumber InTime, float InValue);
+	void AddScalarParameterKey(FName InParameterName, FFrameNumber InTime, float InValue);
+
+	/** Adds a a key for a specific bool parameter at the specified time with the specified value. */
+	void AddBoolParameterKey(FName InParameterName, FFrameNumber InTime, bool InValue);
+
+	 /** Adds a a key for a specific vector2D parameter at the specified time with the specified value. */
+	 void AddVector2DParameterKey(FName InParameterName, FFrameNumber InTime, FVector2D InValue);
 
 	/** Adds a a key for a specific vector parameter at the specified time with the specified value. */
 	 void AddVectorParameterKey(FName InParameterName, FFrameNumber InTime, FVector InValue);
@@ -243,12 +338,28 @@ public:
 	 */
 	 bool RemoveScalarParameter(FName InParameterName);
 
+	 /**
+	  * Removes a bool parameter from this section.
+	  *
+	  * @param InParameterName The name of the bool parameter to remove.
+	  * @returns True if a parameter with that name was found and removed, otherwise false.
+	  */
+	 bool RemoveBoolParameter(FName InParameterName);
+
 	/**
-	 * Removes a vector parameter from this section.
+	 * Removes a vector2D parameter from this section.
 	 *
-	 * @param InParameterName The name of the vector parameter to remove.
+	 * @param InParameterName The name of the vector2D parameter to remove.
 	 * @returns True if a parameter with that name was found and removed, otherwise false.
 	 */
+	 bool RemoveVector2DParameter(FName InParameterName);
+
+	 /**
+  * Removes a vector parameter from this section.
+  *
+  * @param InParameterName The name of the vector parameter to remove.
+  * @returns True if a parameter with that name was found and removed, otherwise false.
+  */
 	 bool RemoveVectorParameter(FName InParameterName);
 
 	/**
@@ -271,6 +382,14 @@ public:
 	 TArray<FScalarParameterNameAndCurve>& GetScalarParameterNamesAndCurves();
 	 const TArray<FScalarParameterNameAndCurve>& GetScalarParameterNamesAndCurves() const;
 
+	 /** Gets the animated bool parameters and their associated curves. */
+	 TArray<FBoolParameterNameAndCurve>& GetBoolParameterNamesAndCurves();
+	 const TArray<FBoolParameterNameAndCurve>& GetBoolParameterNamesAndCurves() const;
+
+	 /** Gets the animated vector2D parameters and their associated curves. */
+	 TArray<FVector2DParameterNameAndCurves>& GetVector2DParameterNamesAndCurves();
+	 const TArray<FVector2DParameterNameAndCurves>& GetVector2DParameterNamesAndCurves() const;
+
 	/** Gets the animated vector parameters and their associated curves. */
 	 TArray<FVectorParameterNameAndCurves>& GetVectorParameterNamesAndCurves();
 	 const TArray<FVectorParameterNameAndCurves>& GetVectorParameterNamesAndCurves() const;
@@ -291,23 +410,30 @@ protected:
 	//~ UMovieSceneSection interface
 	virtual void Serialize(FArchive& Ar) override;
 	virtual void PostEditImport() override;
-	virtual void ReconstructChannelProxy();
+	virtual void ReconstructChannelProxy(bool bForce);
 
 protected:
+	/** The bool parameter names and their associated curves. */
+	UPROPERTY()
+	TArray<FBoolParameterNameAndCurve> BoolParameterNamesAndCurves;
 
 	/** The scalar parameter names and their associated curves. */
 	UPROPERTY()
 	TArray<FScalarParameterNameAndCurve> ScalarParameterNamesAndCurves;
 
+	/** The vector3D parameter names and their associated curves. */
+	UPROPERTY()
+	TArray<FVector2DParameterNameAndCurves> Vector2DParameterNamesAndCurves;
+
 	/** The vector parameter names and their associated curves. */
 	UPROPERTY()
 	TArray<FVectorParameterNameAndCurves> VectorParameterNamesAndCurves;
 
-	/** The vector parameter names and their associated curves. */
+	/** The color parameter names and their associated curves. */
 	UPROPERTY()
 	TArray<FColorParameterNameAndCurves> ColorParameterNamesAndCurves;
 
-	/** The vector parameter names and their associated curves. */
+	/** The transform  parameter names and their associated curves. */
 	UPROPERTY()
 	TArray<FTransformParameterNameAndCurves> TransformParameterNamesAndCurves;
 };
\ No newline at end of file
diff --git a/Engine/Source/Runtime/RHI/Public/RHIDefinitions.h b/Engine/Source/Runtime/RHI/Public/RHIDefinitions.h
--- a/Engine/Source/Runtime/RHI/Public/RHIDefinitions.h
+++ b/Engine/Source/Runtime/RHI/Public/RHIDefinitions.h
@@ -1343,7 +1343,7 @@ inline int32 GetFeatureLevelMaxNumberOfBones(ERHIFeatureLevel::Type FeatureLevel
 	case ERHIFeatureLevel::ES3_1:
 		return 75;	
 	case ERHIFeatureLevel::SM5:
-		return 256;
+		return 65536; // supports uint16
 	default:
 		checkf(0, TEXT("Unknown FeatureLevel %d"), (int32)FeatureLevel);
 	}
diff --git a/Engine/Source/Runtime/Renderer/Private/GpuDebugRendering.cpp b/Engine/Source/Runtime/Renderer/Private/GpuDebugRendering.cpp
--- a/Engine/Source/Runtime/Renderer/Private/GpuDebugRendering.cpp
+++ b/Engine/Source/Runtime/Renderer/Private/GpuDebugRendering.cpp
@@ -31,7 +31,8 @@ namespace ShaderDrawDebug
 
 	bool IsShaderDrawDebugEnabled()
 	{
-#if WITH_EDITOR
+		// This debug mode causes a GPU restart on Mac. Forcing off until UE-87288 is fixed.
+#if WITH_EDITOR && !PLATFORM_MAC
 		return CVarShaderDrawEnable.GetValueOnAnyThread() > 0;
 #else
 		return false;
diff --git a/Engine/Source/Runtime/RigVM/Private/RigVMCore/RigVM.cpp b/Engine/Source/Runtime/RigVM/Private/RigVMCore/RigVM.cpp
--- a/Engine/Source/Runtime/RigVM/Private/RigVMCore/RigVM.cpp
+++ b/Engine/Source/Runtime/RigVM/Private/RigVMCore/RigVM.cpp
@@ -2,6 +2,30 @@
 
 #include "RigVMCore/RigVM.h"
 #include "UObject/Package.h"
+#include "UObject/AnimObjectVersion.h"
+
+bool FRigVMParameter::Serialize(FArchive& Ar)
+{
+	Ar.UsingCustomVersion(FAnimObjectVersion::GUID);
+
+	if (Ar.CustomVer(FAnimObjectVersion::GUID) < FAnimObjectVersion::StoreMarkerNamesOnSkeleton)
+	{
+		return false;
+	}
+
+	Ar << Type;
+	Ar << Name;
+	Ar << RegisterIndex;
+	Ar << CPPType;
+	Ar << ScriptStructPath;
+
+	if (Ar.IsLoading())
+	{
+		ScriptStruct = nullptr;
+	}
+
+	return true;
+}
 
 UScriptStruct* FRigVMParameter::GetScriptStruct() const
 {
@@ -27,6 +51,41 @@ URigVM::~URigVM()
 	Reset();
 }
 
+void URigVM::Serialize(FArchive& Ar)
+{
+	Ar.UsingCustomVersion(FAnimObjectVersion::GUID);
+
+	if (Ar.CustomVer(FAnimObjectVersion::GUID) < FAnimObjectVersion::StoreMarkerNamesOnSkeleton)
+	{
+		return;
+	}
+
+	if (Ar.IsLoading())
+	{
+		Reset();
+	}
+
+	Ar << WorkMemory;
+	Ar << LiteralMemory;
+	Ar << FunctionNames;
+	Ar << ByteCode;
+	Ar << Parameters;
+
+	if (Ar.IsLoading())
+	{
+		Instructions.Reset();
+		Functions.Reset();
+		ParametersNameMap.Reset();
+
+		for (int32 Index = 0; Index < Parameters.Num(); Index++)
+		{
+			ParametersNameMap.Add(Parameters[Index].Name, Index);
+		}
+
+		InvalidateCachedMemory();
+	}
+}
+
 void URigVM::Reset()
 {
 	WorkMemory.Reset();
@@ -37,13 +96,49 @@ void URigVM::Reset()
 	Instructions.Reset();
 	Parameters.Reset();
 	ParametersNameMap.Reset();
+
+	InvalidateCachedMemory();
+}
+
+void URigVM::Empty()
+{
+	WorkMemory.Empty();
+	LiteralMemory.Empty();
+	FunctionNames.Empty();
+	Functions.Empty();
+	ByteCode.Empty();
+	Instructions.Empty();
+	Parameters.Empty();
+	ParametersNameMap.Empty();
+
+	InvalidateCachedMemory();
+
+	CachedMemory.Empty();
+	FirstPointerForInstruction.Empty();
+	CachedMemoryPointers.Empty();
+}
+
+void URigVM::CopyFrom(URigVM* InVM)
+{
+	check(InVM);
+	
+	Reset();
+
+	WorkMemory = InVM->WorkMemory;
+	LiteralMemory = InVM->LiteralMemory;
+	FunctionNames = InVM->FunctionNames;
+	Functions = InVM->Functions;
+	ByteCode = InVM->ByteCode;
+	Instructions = InVM->Instructions;
+	Parameters = InVM->Parameters;
+	ParametersNameMap = InVM->ParametersNameMap;
 }
 
 int32 URigVM::AddRigVMFunction(UScriptStruct* InRigVMStruct, const FName& InMethodName)
 {
 	check(InRigVMStruct);
 	FString FunctionKey = FString::Printf(TEXT("F%s::%s"), *InRigVMStruct->GetName(), *InMethodName.ToString());
-	int32 FunctionIndex = FunctionNames.Find(FunctionKey);
+	int32 FunctionIndex = FunctionNames.Find(*FunctionKey);
 	if (FunctionIndex != INDEX_NONE)
 	{
 		return FunctionIndex;
@@ -55,21 +150,51 @@ int32 URigVM::AddRigVMFunction(UScriptStruct* InRigVMStruct, const FName& InMeth
 		return INDEX_NONE;
 	}
 
-	FunctionNames.Add(FunctionKey);
+	FunctionNames.Add(*FunctionKey);
 	return Functions.Add(Function);
 }
 
+FString URigVM::GetRigVMFunctionName(int32 InFunctionIndex) const
+{
+	return FunctionNames[InFunctionIndex].ToString();
+}
+
 const FRigVMInstructionArray& URigVM::GetInstructions()
 {
 	RefreshInstructionsIfRequired();
 	return Instructions;
 }
 
-const TArray<FRigVMParameter>& URigVM::GetParameters()
+const TArray<FRigVMParameter>& URigVM::GetParameters() const
 {
 	return Parameters;
 }
 
+FRigVMParameter URigVM::GetParameterByName(const FName& InParameterName)
+{
+	if (ParametersNameMap.Num() == Parameters.Num())
+	{
+		const int32* ParameterIndex = ParametersNameMap.Find(InParameterName);
+		if (ParameterIndex)
+		{
+			Parameters[*ParameterIndex].GetScriptStruct();
+			return Parameters[*ParameterIndex];
+		}
+		return FRigVMParameter();
+	}
+
+	for (FRigVMParameter& Parameter : Parameters)
+	{
+		if (Parameter.GetName() == InParameterName)
+		{
+			Parameter.GetScriptStruct();
+			return Parameter;
+		}
+	}
+
+	return FRigVMParameter();
+}
+
 void URigVM::ResolveFunctionsIfRequired()
 {
 	if (Functions.Num() != FunctionNames.Num())
@@ -79,7 +204,7 @@ void URigVM::ResolveFunctionsIfRequired()
 
 		for (int32 FunctionIndex = 0; FunctionIndex < FunctionNames.Num(); FunctionIndex++)
 		{
-			Functions[FunctionIndex] = FRigVMRegistry::Get().Find(*FunctionNames[FunctionIndex]);
+			Functions[FunctionIndex] = FRigVMRegistry::Get().Find(*FunctionNames[FunctionIndex].ToString());
 		}
 	}
 }
@@ -92,6 +217,244 @@ void URigVM::RefreshInstructionsIfRequired()
 	}
 }
 
+void URigVM::InvalidateCachedMemory()
+{
+	CachedMemory.Reset();
+	FirstPointerForInstruction.Reset();
+	CachedMemoryPointers.Reset();
+}
+
+void URigVM::CacheMemoryPointersIfRequired(FRigVMMemoryContainerPtrArray InMemory)
+{
+	RefreshInstructionsIfRequired();
+
+	if (Instructions.Num() == 0 || InMemory.Num() == 0)
+	{
+		InvalidateCachedMemory();
+		return;
+	}
+
+	if (Instructions.Num() != FirstPointerForInstruction.Num())
+	{
+		InvalidateCachedMemory();
+	}
+	else if (InMemory.Num() != CachedMemory.Num())
+	{
+		InvalidateCachedMemory();
+	}
+	else
+	{
+		for (int32 Index = 0; Index < InMemory.Num(); Index++)
+		{
+			if (InMemory[Index] != CachedMemory[Index])
+			{
+				InvalidateCachedMemory();
+				break;
+			}
+		}
+	}
+
+	if (Instructions.Num() == FirstPointerForInstruction.Num())
+	{
+		return;
+	}
+
+	for (int32 Index = 0; Index < InMemory.Num(); Index++)
+	{
+		CachedMemory.Add(InMemory[Index]);
+	}
+
+	uint16 InstructionIndex = 0;
+	while (Instructions[InstructionIndex].OpCode != ERigVMOpCode::Exit)
+	{
+		FirstPointerForInstruction.Add(CachedMemoryPointers.Num());
+
+		switch (Instructions[InstructionIndex].OpCode)
+		{
+			case ERigVMOpCode::Execute_0_Operands:
+			case ERigVMOpCode::Execute_1_Operands:
+			case ERigVMOpCode::Execute_2_Operands:
+			case ERigVMOpCode::Execute_3_Operands:
+			case ERigVMOpCode::Execute_4_Operands:
+			case ERigVMOpCode::Execute_5_Operands:
+			case ERigVMOpCode::Execute_6_Operands:
+			case ERigVMOpCode::Execute_7_Operands:
+			case ERigVMOpCode::Execute_8_Operands:
+			case ERigVMOpCode::Execute_9_Operands:
+			case ERigVMOpCode::Execute_10_Operands:
+			case ERigVMOpCode::Execute_11_Operands:
+			case ERigVMOpCode::Execute_12_Operands:
+			case ERigVMOpCode::Execute_13_Operands:
+			case ERigVMOpCode::Execute_14_Operands:
+			case ERigVMOpCode::Execute_15_Operands:
+			case ERigVMOpCode::Execute_16_Operands:
+			case ERigVMOpCode::Execute_17_Operands:
+			case ERigVMOpCode::Execute_18_Operands:
+			case ERigVMOpCode::Execute_19_Operands:
+			case ERigVMOpCode::Execute_20_Operands:
+			case ERigVMOpCode::Execute_21_Operands:
+			case ERigVMOpCode::Execute_22_Operands:
+			case ERigVMOpCode::Execute_23_Operands:
+			case ERigVMOpCode::Execute_24_Operands:
+			case ERigVMOpCode::Execute_25_Operands:
+			case ERigVMOpCode::Execute_26_Operands:
+			case ERigVMOpCode::Execute_27_Operands:
+			case ERigVMOpCode::Execute_28_Operands:
+			case ERigVMOpCode::Execute_29_Operands:
+			case ERigVMOpCode::Execute_30_Operands:
+			case ERigVMOpCode::Execute_31_Operands:
+			case ERigVMOpCode::Execute_32_Operands:
+			case ERigVMOpCode::Execute_33_Operands:
+			case ERigVMOpCode::Execute_34_Operands:
+			case ERigVMOpCode::Execute_35_Operands:
+			case ERigVMOpCode::Execute_36_Operands:
+			case ERigVMOpCode::Execute_37_Operands:
+			case ERigVMOpCode::Execute_38_Operands:
+			case ERigVMOpCode::Execute_39_Operands:
+			case ERigVMOpCode::Execute_40_Operands:
+			case ERigVMOpCode::Execute_41_Operands:
+			case ERigVMOpCode::Execute_42_Operands:
+			case ERigVMOpCode::Execute_43_Operands:
+			case ERigVMOpCode::Execute_44_Operands:
+			case ERigVMOpCode::Execute_45_Operands:
+			case ERigVMOpCode::Execute_46_Operands:
+			case ERigVMOpCode::Execute_47_Operands:
+			case ERigVMOpCode::Execute_48_Operands:
+			case ERigVMOpCode::Execute_49_Operands:
+			case ERigVMOpCode::Execute_50_Operands:
+			case ERigVMOpCode::Execute_51_Operands:
+			case ERigVMOpCode::Execute_52_Operands:
+			case ERigVMOpCode::Execute_53_Operands:
+			case ERigVMOpCode::Execute_54_Operands:
+			case ERigVMOpCode::Execute_55_Operands:
+			case ERigVMOpCode::Execute_56_Operands:
+			case ERigVMOpCode::Execute_57_Operands:
+			case ERigVMOpCode::Execute_58_Operands:
+			case ERigVMOpCode::Execute_59_Operands:
+			case ERigVMOpCode::Execute_60_Operands:
+			case ERigVMOpCode::Execute_61_Operands:
+			case ERigVMOpCode::Execute_62_Operands:
+			case ERigVMOpCode::Execute_63_Operands:
+			case ERigVMOpCode::Execute_64_Operands:
+			{
+				const FRigVMExecuteOp& Op = ByteCode.GetOpAt<FRigVMExecuteOp>(Instructions[InstructionIndex]);
+				FRigVMOperandArray Operands = ByteCode.GetOperandsForExecuteOp(Instructions[InstructionIndex]);
+
+				for (const FRigVMOperand& Arg : Operands)
+				{
+					void* Ptr = CachedMemory[Arg.GetContainerIndex()]->GetData(Arg.GetRegisterIndex(), Arg.GetRegisterOffset());
+					CachedMemoryPointers.Add(Ptr);
+
+					FRigVMRegister Register = (*CachedMemory[Arg.GetContainerIndex()])[Arg];
+					if (Register.IsArray())
+					{
+						void* NumBytesAsPtr = reinterpret_cast<void*>(Register.ElementCount);
+						CachedMemoryPointers.Add(NumBytesAsPtr);
+					}
+				}
+
+				InstructionIndex++;
+				break;
+			}
+			case ERigVMOpCode::Zero:
+			case ERigVMOpCode::BoolFalse:
+			case ERigVMOpCode::BoolTrue:
+			case ERigVMOpCode::Increment:
+			case ERigVMOpCode::Decrement:
+			{
+				const FRigVMUnaryOp& Op = ByteCode.GetOpAt<FRigVMUnaryOp>(Instructions[InstructionIndex]);
+				const FRigVMOperand& Arg = Op.Arg;
+				CachedMemoryPointers.Add(CachedMemory[Arg.GetContainerIndex()]->GetData(Arg.GetRegisterIndex(), Arg.GetRegisterOffset()));
+				InstructionIndex++;
+				break;
+			}
+			case ERigVMOpCode::Copy:
+			{
+				const FRigVMCopyOp& Op = ByteCode.GetOpAt<FRigVMCopyOp>(Instructions[InstructionIndex]);
+				FRigVMOperand Arg = Op.Source;
+				CachedMemoryPointers.Add(CachedMemory[Arg.GetContainerIndex()]->GetData(Arg.GetRegisterIndex(), Arg.GetRegisterOffset()));
+				Arg = Op.Target;
+				CachedMemoryPointers.Add(CachedMemory[Arg.GetContainerIndex()]->GetData(Arg.GetRegisterIndex(), Arg.GetRegisterOffset()));
+
+				const FRigVMRegister& TargetRegister = CachedMemory[Arg.GetContainerIndex()]->Registers[Arg.GetRegisterIndex()];
+				uint16 NumBytes = TargetRegister.GetNumBytesPerSlice();
+
+				ERigVMRegisterType TargetType = TargetRegister.Type;
+				if (Arg.GetRegisterOffset() != INDEX_NONE)
+				{
+					TargetType = CachedMemory[Arg.GetContainerIndex()]->RegisterOffsets[Arg.GetRegisterOffset()].GetType();
+					NumBytes = CachedMemory[Arg.GetContainerIndex()]->RegisterOffsets[Arg.GetRegisterOffset()].GetElementSize();
+				}
+
+				CachedMemoryPointers.Add(reinterpret_cast<void*>(NumBytes));
+				CachedMemoryPointers.Add(reinterpret_cast<void*>((uint16)TargetType));
+
+				if (TargetType == ERigVMRegisterType::Struct)
+				{
+					CachedMemoryPointers.Add(CachedMemory[Arg.GetContainerIndex()]->GetScriptStruct(Arg.GetRegisterIndex(), Arg.GetRegisterOffset()));
+				}
+
+				InstructionIndex++;
+				break;
+			}
+			case ERigVMOpCode::Equals:
+			case ERigVMOpCode::NotEquals:
+			{
+				const FRigVMComparisonOp& Op = ByteCode.GetOpAt<FRigVMComparisonOp>(Instructions[InstructionIndex]);
+				FRigVMOperand Arg = Op.A;
+				CachedMemoryPointers.Add(CachedMemory[Arg.GetContainerIndex()]->GetData(Arg.GetRegisterIndex(), Arg.GetRegisterOffset()));
+				Arg = Op.B;
+				CachedMemoryPointers.Add(CachedMemory[Arg.GetContainerIndex()]->GetData(Arg.GetRegisterIndex(), Arg.GetRegisterOffset()));
+				Arg = Op.Result;
+				CachedMemoryPointers.Add(CachedMemory[Arg.GetContainerIndex()]->GetData(Arg.GetRegisterIndex(), Arg.GetRegisterOffset()));
+				InstructionIndex++;
+				break;
+			}
+			case ERigVMOpCode::JumpAbsolute:
+			case ERigVMOpCode::JumpForward:
+			case ERigVMOpCode::JumpBackward:
+			{
+				InstructionIndex++;
+				break;
+			}
+			case ERigVMOpCode::JumpAbsoluteIf:
+			case ERigVMOpCode::JumpForwardIf:
+			case ERigVMOpCode::JumpBackwardIf:
+			{
+				const FRigVMJumpIfOp& Op = ByteCode.GetOpAt<FRigVMJumpIfOp>(Instructions[InstructionIndex]);
+				const FRigVMOperand& Arg = Op.Arg;
+				CachedMemoryPointers.Add(CachedMemory[Arg.GetContainerIndex()]->GetData(Arg.GetRegisterIndex(), Arg.GetRegisterOffset()));
+				InstructionIndex++;
+				break;
+			}
+			case ERigVMOpCode::ChangeType:
+			{
+				const FRigVMChangeTypeOp& Op = ByteCode.GetOpAt<FRigVMChangeTypeOp>(Instructions[InstructionIndex]);
+				const FRigVMOperand& Arg = Op.Arg;
+				CachedMemoryPointers.Add(CachedMemory[Arg.GetContainerIndex()]->GetData(Arg.GetRegisterIndex(), Arg.GetRegisterOffset()));
+				InstructionIndex++;
+				break;
+			}
+			case ERigVMOpCode::Exit:
+			case ERigVMOpCode::Invalid:
+			{
+				ensure(false);
+				break;
+			}
+		}
+
+		if (!Instructions.IsValidIndex(InstructionIndex))
+		{
+			break;
+		}
+	}
+
+	if (FirstPointerForInstruction.Num() < Instructions.Num())
+	{
+		FirstPointerForInstruction.Add(CachedMemoryPointers.Num());
+	}
+}
+
 bool URigVM::Execute(FRigVMMemoryContainerPtrArray Memory, TArrayView<void*> AdditionalArguments)
 {
 	ResolveFunctionsIfRequired();
@@ -108,6 +471,8 @@ bool URigVM::Execute(FRigVMMemoryContainerPtrArray Memory, TArrayView<void*> Add
 		Memory = FRigVMMemoryContainerPtrArray(LocalMemory, 2);
 	}
 
+	CacheMemoryPointersIfRequired(Memory);
+
 	uint16 InstructionIndex = 0;
 	while (Instructions[InstructionIndex].OpCode != ERigVMOpCode::Exit)
 	{
@@ -180,8 +545,13 @@ bool URigVM::Execute(FRigVMMemoryContainerPtrArray Memory, TArrayView<void*> Add
 			case ERigVMOpCode::Execute_64_Operands:
 			{
 				const FRigVMExecuteOp& Op = ByteCode.GetOpAt<FRigVMExecuteOp>(Instructions[InstructionIndex]);
-				FRigVMOperandArray Operands = ByteCode.GetOperandsForExecuteOp(Instructions[InstructionIndex]);
-				(*Functions[Op.FunctionIndex])(Operands, Memory, AdditionalArguments);
+				int32 OperandCount = FirstPointerForInstruction[InstructionIndex + 1] - FirstPointerForInstruction[InstructionIndex];
+				FRigVMOperandMemory OpMemory(&CachedMemoryPointers[FirstPointerForInstruction[InstructionIndex]], OperandCount);
+#if WITH_EDITOR
+				(*Functions[Op.FunctionIndex])(FunctionNames[Op.FunctionIndex], InstructionIndex, OpMemory, AdditionalArguments);
+#else
+				(*Functions[Op.FunctionIndex])(NAME_None, INDEX_NONE, OpMemory, AdditionalArguments);
+#endif
 				InstructionIndex++;
 				break;
 			}
@@ -194,36 +564,81 @@ bool URigVM::Execute(FRigVMMemoryContainerPtrArray Memory, TArrayView<void*> Add
 			}
 			case ERigVMOpCode::BoolFalse:
 			{
-				const FRigVMUnaryOp& Op = ByteCode.GetOpAt<FRigVMUnaryOp>(Instructions[InstructionIndex]);
-				Memory[Op.Arg.GetContainerIndex()]->GetRef<bool>(Op.Arg.GetRegisterIndex()) = false;
+				*((bool*)CachedMemoryPointers[FirstPointerForInstruction[InstructionIndex]]) = false;
 				InstructionIndex++;
 				break;
 			}
 			case ERigVMOpCode::BoolTrue:
 			{
-				const FRigVMUnaryOp& Op = ByteCode.GetOpAt<FRigVMUnaryOp>(Instructions[InstructionIndex]);
-				Memory[Op.Arg.GetContainerIndex()]->GetRef<bool>(Op.Arg.GetRegisterIndex()) = false;
+				*((bool*)CachedMemoryPointers[FirstPointerForInstruction[InstructionIndex]]) = true;
 				InstructionIndex++;
 				break;
 			}
 			case ERigVMOpCode::Copy:
 			{
 				const FRigVMCopyOp& Op = ByteCode.GetOpAt<FRigVMCopyOp>(Instructions[InstructionIndex]);
-				Memory[Op.Target.GetContainerIndex()]->Copy(Op.Source, Op.Target, Memory[Op.Source.GetContainerIndex()]);
+
+				void * SourcePtr = CachedMemoryPointers[FirstPointerForInstruction[InstructionIndex]];
+				void * TargetPtr = CachedMemoryPointers[FirstPointerForInstruction[InstructionIndex] + 1];
+				uint64 NumBytes = reinterpret_cast<uint64>(CachedMemoryPointers[FirstPointerForInstruction[InstructionIndex] + 2]);
+				ERigVMRegisterType MemoryType = (ERigVMRegisterType)reinterpret_cast<uint64>(CachedMemoryPointers[FirstPointerForInstruction[InstructionIndex] + 3]);
+
+				switch (MemoryType)
+				{
+					case ERigVMRegisterType::Plain:
+					{
+						FMemory::Memcpy(TargetPtr, SourcePtr, NumBytes);
+						break;
+					}
+					case ERigVMRegisterType::Name:
+					{
+						int32 NumNames = NumBytes / sizeof(FName);
+						TArrayView<FName> TargetNames((FName*)TargetPtr, NumNames);
+						TArrayView<FName> SourceNames((FName*)SourcePtr, NumNames);
+						for (int32 Index = 0; Index < NumNames; Index++)
+						{
+							TargetNames[Index] = SourceNames[Index];
+						}
+						break;
+					}
+					case ERigVMRegisterType::String:
+					{
+						int32 NumStrings = NumBytes / sizeof(FString);
+						TArrayView<FString> TargetStrings((FString*)TargetPtr, NumStrings);
+						TArrayView<FString> SourceStrings((FString*)SourcePtr, NumStrings);
+						for (int32 Index = 0; Index < NumStrings; Index++)
+						{
+							TargetStrings[Index] = SourceStrings[Index];
+						}
+						break;
+					}
+					case ERigVMRegisterType::Struct:
+					{
+						UScriptStruct* ScriptStruct = (UScriptStruct*)CachedMemoryPointers[FirstPointerForInstruction[InstructionIndex] + 4];
+						int32 NumStructs = NumBytes / ScriptStruct->GetStructureSize();
+						ScriptStruct->CopyScriptStruct(TargetPtr, SourcePtr, NumStructs);
+						break;
+					}
+					default:
+					{
+						// the default pass for any complex memory
+						Memory[Op.Target.GetContainerIndex()]->Copy(Op.Source, Op.Target, Memory[Op.Source.GetContainerIndex()]);
+						break;
+					}
+				}
+
 				InstructionIndex++;
 				break;
 			}
 			case ERigVMOpCode::Increment:
 			{
-				const FRigVMUnaryOp& Op = ByteCode.GetOpAt<FRigVMUnaryOp>(Instructions[InstructionIndex]);
-				Memory[Op.Arg.GetContainerIndex()]->GetRef<int32>(Op.Arg.GetRegisterIndex())++;
+				(*(int32*)CachedMemoryPointers[FirstPointerForInstruction[InstructionIndex]])++;
 				InstructionIndex++;
 				break;
 			}
 			case ERigVMOpCode::Decrement:
 			{
-				const FRigVMUnaryOp& Op = ByteCode.GetOpAt<FRigVMUnaryOp>(Instructions[InstructionIndex]);
-				Memory[Op.Arg.GetContainerIndex()]->GetRef<int32>(Op.Arg.GetRegisterIndex())--;
+				(*(int32*)CachedMemoryPointers[FirstPointerForInstruction[InstructionIndex]])--;
 				InstructionIndex++;
 				break;
 			}
@@ -244,8 +659,8 @@ bool URigVM::Execute(FRigVMMemoryContainerPtrArray Memory, TArrayView<void*> Add
 						case ERigVMRegisterType::Plain:
 						case ERigVMRegisterType::Name:
 						{
-							void * DataA = Memory[Op.A.GetContainerIndex()]->GetData(Op.A.GetRegisterIndex());
-							void * DataB = Memory[Op.B.GetContainerIndex()]->GetData(Op.B.GetRegisterIndex());
+							void * DataA = CachedMemoryPointers[FirstPointerForInstruction[InstructionIndex]];
+							void * DataB = CachedMemoryPointers[FirstPointerForInstruction[InstructionIndex]+1];
 							Result = FMemory::Memcmp(DataA, DataB, BytesA) == 0;
 							break;
 						}
@@ -269,8 +684,8 @@ bool URigVM::Execute(FRigVMMemoryContainerPtrArray Memory, TArrayView<void*> Add
 						{
 							UScriptStruct* ScriptStruct = Memory[Op.A.GetContainerIndex()]->GetScriptStruct(RegisterA.ScriptStructIndex);
 
-							uint8* DataA = (uint8*)Memory[Op.A.GetContainerIndex()]->GetData(Op.A.GetRegisterIndex());
-							uint8* DataB = (uint8*)Memory[Op.B.GetContainerIndex()]->GetData(Op.B.GetRegisterIndex());
+							uint8* DataA = (uint8*)CachedMemoryPointers[FirstPointerForInstruction[InstructionIndex]];
+							uint8* DataB = (uint8*)CachedMemoryPointers[FirstPointerForInstruction[InstructionIndex]+1];
 							
 							Result = true;
 							for (int32 ElementIndex = 0; ElementIndex < RegisterA.ElementCount; ElementIndex++)
@@ -296,7 +711,8 @@ bool URigVM::Execute(FRigVMMemoryContainerPtrArray Memory, TArrayView<void*> Add
 				{
 					Result = !Result;
 				}
-				Memory[Op.Result.GetContainerIndex()]->GetRef<bool>(Op.Result.GetRegisterIndex()) = Result;
+
+				(*(bool*)CachedMemoryPointers[FirstPointerForInstruction[InstructionIndex]+2]) = Result;
 				InstructionIndex++;
 				break;
 			}
@@ -321,7 +737,7 @@ bool URigVM::Execute(FRigVMMemoryContainerPtrArray Memory, TArrayView<void*> Add
 			case ERigVMOpCode::JumpAbsoluteIf:
 			{
 				const FRigVMJumpIfOp& Op = ByteCode.GetOpAt<FRigVMJumpIfOp>(Instructions[InstructionIndex]);
-				const bool Condition = Memory[0]->GetRef<bool>(Op.Arg.GetRegisterIndex());
+				const bool Condition = *(bool*)CachedMemoryPointers[FirstPointerForInstruction[InstructionIndex]];
 				if (Condition == Op.Condition)
 				{
 					InstructionIndex = Op.InstructionIndex;
@@ -335,7 +751,7 @@ bool URigVM::Execute(FRigVMMemoryContainerPtrArray Memory, TArrayView<void*> Add
 			case ERigVMOpCode::JumpForwardIf:
 			{
 				const FRigVMJumpIfOp& Op = ByteCode.GetOpAt<FRigVMJumpIfOp>(Instructions[InstructionIndex]);
-				const bool Condition = Memory[0]->GetRef<bool>(Op.Arg.GetRegisterIndex());
+				const bool Condition = *(bool*)CachedMemoryPointers[FirstPointerForInstruction[InstructionIndex]];
 				if (Condition == Op.Condition)
 				{
 					InstructionIndex += Op.InstructionIndex;
@@ -349,7 +765,7 @@ bool URigVM::Execute(FRigVMMemoryContainerPtrArray Memory, TArrayView<void*> Add
 			case ERigVMOpCode::JumpBackwardIf:
 			{
 				const FRigVMJumpIfOp& Op = ByteCode.GetOpAt<FRigVMJumpIfOp>(Instructions[InstructionIndex]);
-				const bool Condition = Memory[0]->GetRef<bool>(Op.Arg.GetRegisterIndex());
+				const bool Condition = *(bool*)CachedMemoryPointers[FirstPointerForInstruction[InstructionIndex]];
 				if (Condition == Op.Condition)
 				{
 					InstructionIndex -= Op.InstructionIndex;
@@ -387,4 +803,16 @@ bool URigVM::Execute(FRigVMMemoryContainerPtrArray Memory, TArrayView<void*> Add
 bool URigVM::Execute()
 {
 	return Execute(FRigVMMemoryContainerPtrArray(), TArrayView<void*>());
-}
\ No newline at end of file
+}
+
+void URigVM::SetRegisterValueFromString(const FRigVMOperand& InOperand, const FString& InCPPType, const UObject* InCPPTypeObject, const TArray<FString>& InDefaultValues)
+{
+	if (InOperand.GetMemoryType() == ERigVMMemoryType::Literal)
+	{
+		LiteralMemory.SetRegisterValueFromString(InOperand, InCPPType, InCPPTypeObject, InDefaultValues);
+	}
+	else if (InOperand.GetMemoryType() == ERigVMMemoryType::Work)
+	{
+		WorkMemory.SetRegisterValueFromString(InOperand, InCPPType, InCPPTypeObject, InDefaultValues);
+	}
+}
diff --git a/Engine/Source/Runtime/RigVM/Private/RigVMCore/RigVMByteCode.cpp b/Engine/Source/Runtime/RigVM/Private/RigVMCore/RigVMByteCode.cpp
--- a/Engine/Source/Runtime/RigVM/Private/RigVMCore/RigVMByteCode.cpp
+++ b/Engine/Source/Runtime/RigVM/Private/RigVMCore/RigVMByteCode.cpp
@@ -1,6 +1,7 @@
 // Copyright Epic Games, Inc. All Rights Reserved.
 
 #include "RigVMCore/RigVMByteCode.h"
+#include "UObject/AnimObjectVersion.h"
 
 FRigVMInstructionArray::FRigVMInstructionArray()
 {
@@ -22,15 +23,38 @@ void FRigVMInstructionArray::Reset()
 	Instructions.Reset();
 }
 
+void FRigVMInstructionArray::Empty()
+{
+	Instructions.Empty();
+}
+
 FRigVMByteCode::FRigVMByteCode()
 {
 }
 
+bool FRigVMByteCode::Serialize(FArchive& Ar)
+{
+	Ar.UsingCustomVersion(FAnimObjectVersion::GUID);
+
+	if (Ar.CustomVer(FAnimObjectVersion::GUID) < FAnimObjectVersion::StoreMarkerNamesOnSkeleton)
+	{
+		return false;
+	}
+
+	Ar << ByteCode;
+	return true;
+}
+
 void FRigVMByteCode::Reset()
 {
 	ByteCode.Reset();
 }
 
+void FRigVMByteCode::Empty()
+{
+	ByteCode.Empty();
+}
+
 uint64 FRigVMByteCode::Num() const
 {
 	return (uint64)ByteCode.Num();
@@ -181,18 +205,21 @@ uint64 FRigVMByteCode::AddTrueOp(const FRigVMOperand& InArg)
 
 uint64 FRigVMByteCode::AddCopyOp(const FRigVMOperand& InSource, const FRigVMOperand& InTarget)
 {
+	ensure(InTarget.GetMemoryType() != ERigVMMemoryType::Literal);
 	FRigVMCopyOp Op(InSource, InTarget);
 	return AddOp(Op);
 }
 
 uint64 FRigVMByteCode::AddIncrementOp(const FRigVMOperand& InArg)
 {
+	ensure(InArg.GetMemoryType() != ERigVMMemoryType::Literal);
 	FRigVMUnaryOp Op(ERigVMOpCode::Increment, InArg);
 	return AddOp(Op);
 }
 
 uint64 FRigVMByteCode::AddDecrementOp(const FRigVMOperand& InArg)
 {
+	ensure(InArg.GetMemoryType() != ERigVMMemoryType::Literal);
 	FRigVMUnaryOp Op(ERigVMOpCode::Decrement, InArg);
 	return AddOp(Op);
 }
