COPY from //Dev-Anim to //Dev-Main

diff --git a/Engine/Source/Runtime/RigVM/Private/RigVMCore/RigVMMemory.cpp b/Engine/Source/Runtime/RigVM/Private/RigVMCore/RigVMMemory.cpp
--- a/Engine/Source/Runtime/RigVM/Private/RigVMCore/RigVMMemory.cpp
+++ b/Engine/Source/Runtime/RigVM/Private/RigVMCore/RigVMMemory.cpp
@@ -25,6 +25,22 @@ bool FRigVMRegister::Serialize(FArchive& Ar)
 	Ar << Name;
 	Ar << ScriptStructIndex;
 
+	if (Ar.IsLoading())
+	{
+		if (Ar.CustomVer(FAnimObjectVersion::GUID) >= FAnimObjectVersion::SerializeRigVMRegisterArrayState)
+		{
+			Ar << bIsArray;
+		}
+		else
+		{
+			bIsArray = false;
+		}
+	}
+	else
+	{
+		Ar << bIsArray;
+	}
+
 	return true;
 }
 
@@ -295,6 +311,7 @@ FRigVMMemoryContainer& FRigVMMemoryContainer::operator= (const FRigVMMemoryConta
 	bUseNameMap = InOther.bUseNameMap;
 	Data.Append(InOther.Data);
 	Registers.Append(InOther.Registers);
+	RegisterOffsets.Append(InOther.RegisterOffsets);
 	ScriptStructs.Append(InOther.ScriptStructs);
 
 	UpdateRegisters();
@@ -364,6 +381,11 @@ bool FRigVMMemoryContainer::Serialize(FArchive& Ar)
 
 		for (const FRigVMRegister& Register : Registers)
 		{
+			if (Register.ElementCount == 0)
+			{
+				continue;
+			}
+
 			switch (Register.Type)
 			{
 				case ERigVMRegisterType::Plain:
@@ -430,6 +452,11 @@ bool FRigVMMemoryContainer::Serialize(FArchive& Ar)
 
 		for (FRigVMRegister& Register : Registers)
 		{
+			if (Register.ElementCount == 0)
+			{
+				continue;
+			}
+
 			Register.MoveToFirstSlice();
 
 			switch(Register.Type)
@@ -481,17 +508,30 @@ bool FRigVMMemoryContainer::Serialize(FArchive& Ar)
 
 void FRigVMMemoryContainer::Reset()
 {
-	for (int32 Index = 0; Index < Registers.Num(); Index++)
+	if (Data.Num() > 0)
 	{
-		Destroy(Index);
+		for (int32 Index = 0; Index < Registers.Num(); Index++)
+		{
+			Destroy(Index);
+		}
 	}
 
 	Data.Reset();
 	Registers.Reset();
+	RegisterOffsets.Reset();
 	ScriptStructs.Reset();
 	NameMap.Reset();
 }
 
+void FRigVMMemoryContainer::Empty()
+{
+	Data.Empty();
+	Registers.Empty();
+	RegisterOffsets.Empty();
+	ScriptStructs.Empty();
+	NameMap.Empty();
+}
+
 bool FRigVMMemoryContainer::Copy(
 	int32 InSourceRegisterIndex,
 	int32 InTargetRegisterIndex,
@@ -515,6 +555,11 @@ bool FRigVMMemoryContainer::Copy(
 	const FRigVMRegister& Source = InSourceMemory->Registers[InSourceRegisterIndex];
 	const FRigVMRegister& Target = Registers[InTargetRegisterIndex];
 
+	if (Target.GetNumBytesPerSlice() == 0)
+	{
+		return true;
+	}
+
 	const uint8* SourcePtr = InSourceMemory->GetData(InSourceRegisterIndex, InSourceRegisterOffset);
 	uint8* TargetPtr = GetData(InTargetRegisterIndex, InTargetRegisterOffset);
 	uint16 NumBytes = Target.GetNumBytesPerSlice();
@@ -612,7 +657,7 @@ int32 FRigVMMemoryContainer::Allocate(const FName& InNewName, int32 InElementSiz
 		} while (!IsNameAvailable(Name));
 	}
 
-	ensure(InElementSize > 0 && InElementCount > 0 && InSliceCount > 0);
+	ensure(InElementSize > 0 && InElementCount >= 0 && InSliceCount > 0);
 
 	if (bUseNameMap)
 	{
@@ -666,6 +711,12 @@ bool FRigVMMemoryContainer::Construct(int32 InRegisterIndex, int32 InElementInde
 	ensure(Registers.IsValidIndex(InRegisterIndex));
 
 	const FRigVMRegister& Register = Registers[InRegisterIndex];
+
+	if (Register.ElementCount == 0)
+	{
+		return true;
+	}
+
 	switch (Register.Type)
 	{
 		case ERigVMRegisterType::Struct:
@@ -716,6 +767,11 @@ bool FRigVMMemoryContainer::Destroy(int32 InRegisterIndex, int32 InElementIndex)
 
 	FRigVMRegister& Register = Registers[InRegisterIndex];
 
+	if (Register.ElementCount == 0)
+	{
+		return true;
+	}
+
 	if (InElementIndex == INDEX_NONE)
 	{
 		Register.MoveToFirstSlice();
@@ -965,6 +1021,131 @@ int32 FRigVMMemoryContainer::GetOrAddRegisterOffset(int32 InRegisterIndex, UScri
 	return ExistingIndex;
 }
 
+void FRigVMMemoryContainer::SetRegisterValueFromString(const FRigVMOperand& InOperand, const FString& InCPPType, const UObject* InCPPTypeObject, const TArray<FString>& InDefaultValues)
+{
+	if (InOperand.GetRegisterIndex() < 0 || InOperand.GetRegisterIndex() >= Registers.Num())
+	{
+		return;
+	}
+
+	FRigVMRegister& Register = Registers[InOperand.GetRegisterIndex()];
+
+	if (Register.ElementCount != InDefaultValues.Num())
+	{
+		return;
+	}
+
+	for (int32 Index = 0; Index < InDefaultValues.Num(); Index++)
+	{
+		FString DefaultValue = InDefaultValues[Index];
+
+		if (Register.ScriptStructIndex != INDEX_NONE)
+		{
+			UScriptStruct* ScriptStruct = ScriptStructs[Register.ScriptStructIndex];
+			if (ScriptStruct == InCPPTypeObject)
+			{
+				uint8* DataPtr = (uint8*)GetData(Register);
+				DataPtr += Index * ScriptStruct->GetStructureSize();
+				ScriptStruct->ImportText(*DefaultValue, DataPtr, nullptr, PPF_None, nullptr, ScriptStruct->GetName());
+			}
+		}
+		else if (const UEnum* Enum = Cast<const UEnum>(InCPPTypeObject))
+		{
+			if (FCString::IsNumeric(*DefaultValue))
+			{
+				GetArray<uint8>(Register)[Index] = (uint8)FCString::Atoi(*DefaultValue);
+			}
+			else
+			{
+				GetArray<uint8>(Register)[Index] = (uint8)Enum->GetValueByNameString(DefaultValue);
+			}
+		}
+		else if (InCPPType == TEXT("bool") && Register.GetNumBytesPerSlice() == sizeof(bool))
+		{
+			GetArray<bool>(Register)[Index] = (DefaultValue == TEXT("True")) || (DefaultValue == TEXT("true")) || (DefaultValue == TEXT("1"));
+		}
+		else if (InCPPType == TEXT("int32") && Register.GetNumBytesPerSlice() == sizeof(int32))
+		{
+			GetArray<int32>(Register)[Index] = FCString::Atoi(*DefaultValue);
+		}
+		else if (InCPPType == TEXT("float") && Register.GetNumBytesPerSlice() == sizeof(float))
+		{
+			GetArray<float>(Register)[Index] = FCString::Atof(*DefaultValue);
+		}
+		else if (InCPPType == TEXT("FName") && Register.GetNumBytesPerSlice() == sizeof(FName))
+		{
+			GetArray<FName>(Register)[Index] = *DefaultValue;
+		}
+		else if (InCPPType == TEXT("FString") && Register.GetNumBytesPerSlice() == sizeof(FString))
+		{
+			GetArray<FString>(Register)[Index] = *DefaultValue;
+		}
+		else
+		{
+			ensure(false);
+		}
+	}
+}
+
+TArray<FString> FRigVMMemoryContainer::GetRegisterValueAsString(const FRigVMOperand& InOperand, const FString& InCPPType, const UObject* InCPPTypeObject)
+{
+	TArray<FString> DefaultValues;
+
+	if (InOperand.GetRegisterIndex() < 0 || InOperand.GetRegisterIndex() >= Registers.Num())
+	{
+		return DefaultValues;
+	}
+
+	FRigVMRegister& Register = Registers[InOperand.GetRegisterIndex()];
+	for (int32 Index = 0; Index < Register.ElementCount; Index++)
+	{
+		FString DefaultValue;
+
+		if (Register.ScriptStructIndex != INDEX_NONE)
+		{
+			UScriptStruct* ScriptStruct = ScriptStructs[Register.ScriptStructIndex];
+			if (ScriptStruct == InCPPTypeObject)
+			{
+				uint8* DataPtr = (uint8*)GetData(Register);
+				DataPtr += Index * ScriptStruct->GetStructureSize();
+				ScriptStruct->ExportText(DefaultValue, DataPtr, nullptr, nullptr, PPF_None, nullptr);
+			}
+		}
+		else if (const UEnum* Enum = Cast<const UEnum>(InCPPTypeObject))
+		{
+			DefaultValue = Enum->GetNameStringByValue((int64)GetArray<int32>(Register)[Index]);
+		}
+		else if (InCPPType == TEXT("bool") && Register.GetNumBytesPerSlice() == sizeof(bool))
+		{
+			DefaultValue = GetArray<bool>(Register)[Index] ? TEXT("True") : TEXT("False");
+		}
+		else if (InCPPType == TEXT("int32") && Register.GetNumBytesPerSlice() == sizeof(int32))
+		{
+			DefaultValue = FString::FromInt(GetArray<int32>(Register)[Index]);
+		}
+		else if (InCPPType == TEXT("float") && Register.GetNumBytesPerSlice() == sizeof(float))
+		{
+			DefaultValue = FString::Printf(TEXT("%f"), GetArray<float>(Register)[Index]);
+		}
+		else if (InCPPType == TEXT("FName") && Register.GetNumBytesPerSlice() == sizeof(FName))
+		{
+			DefaultValue = GetArray<FName>(Register)[Index].ToString();
+		}
+		else if (InCPPType == TEXT("FString") && Register.GetNumBytesPerSlice() == sizeof(FString))
+		{
+			DefaultValue  = GetArray<FString>(Register)[Index];
+		}
+		else
+		{
+			ensure(false);
+		}
+
+		DefaultValues.Add(DefaultValue);
+	}
+
+	return DefaultValues;
+}
+
 void FRigVMMemoryContainer::UpdateRegisters()
 {
 	int32 AlignmentShift = 0;
diff --git a/Engine/Source/Runtime/RigVM/Private/RigVMCore/RigVMRegistry.cpp b/Engine/Source/Runtime/RigVM/Private/RigVMCore/RigVMRegistry.cpp
--- a/Engine/Source/Runtime/RigVM/Private/RigVMCore/RigVMRegistry.cpp
+++ b/Engine/Source/Runtime/RigVM/Private/RigVMCore/RigVMRegistry.cpp
@@ -14,13 +14,13 @@ void FRigVMRegistry::Refresh()
 {
 }
 
-void FRigVMRegistry::Register(const TCHAR* InName, FRigVMFunctionPtr InFunctionPtr)
+void FRigVMRegistry::Register(const TCHAR* InName, FRigVMFunctionPtr InFunctionPtr, UScriptStruct* InStruct)
 {
 	if (Find(InName) != nullptr)
 	{
 		return;
 	}
-	Functions.Add(FRigVMFunction(InName, InFunctionPtr));
+	Functions.Add(FRigVMFunction(InName, InFunctionPtr, InStruct));
 }
 
 FRigVMFunctionPtr FRigVMRegistry::Find(const TCHAR* InName) const
@@ -35,3 +35,9 @@ FRigVMFunctionPtr FRigVMRegistry::Find(const TCHAR* InName) const
 
 	return nullptr;
 }
+
+const TArray<FRigVMFunction>& FRigVMRegistry::GetFunctions() const
+{
+	return Functions;
+}
+
diff --git a/Engine/Source/Runtime/RigVM/Public/RigVMCore/RigVM.h b/Engine/Source/Runtime/RigVM/Public/RigVMCore/RigVM.h
--- a/Engine/Source/Runtime/RigVM/Public/RigVMCore/RigVM.h
+++ b/Engine/Source/Runtime/RigVM/Public/RigVMCore/RigVM.h
@@ -6,6 +6,7 @@
 #include "RigVMMemory.h"
 #include "RigVMRegistry.h"
 #include "RigVMByteCode.h"
+#include "RigVMStatistics.h"
 #include "RigVM.generated.h"
 
 // The type of parameter for a VM
@@ -23,7 +24,7 @@ enum class ERigVMParameterType : uint8
  * used to set input parameters as well as retrieve output parameters.
  */
 USTRUCT(BlueprintType)
-struct FRigVMParameter
+struct RIGVM_API FRigVMParameter
 {
 	GENERATED_BODY()
 
@@ -39,6 +40,13 @@ public:
 	{
 	}
 
+	bool Serialize(FArchive& Ar);
+	FORCEINLINE friend FArchive& operator<<(FArchive& Ar, FRigVMParameter& P)
+	{
+		P.Serialize(Ar);
+		return Ar;
+	}
+
 	// returns true if the parameter is valid
 	bool IsValid() const { return Type != ERigVMParameterType::Invalid; }
 
@@ -110,9 +118,18 @@ public:
 	URigVM();
 	virtual ~URigVM();
 
-	// resets the container and removes all memory
+	// UObject interface
+	void Serialize(FArchive& Ar);
+
+	// resets the container and maintains all memory
 	void Reset();
 
+	// resets the container and removes all memory
+	void Empty();
+
+	// resets the container and clones the input VM
+	void CopyFrom(URigVM* InVM);
+
 	// Executes the VM.
 	// You can optionally provide external memory to the execution
 	// and provide optional additional operands.
@@ -129,6 +146,10 @@ public:
 	UFUNCTION()
 	int32 AddRigVMFunction(UScriptStruct* InRigVMStruct, const FName& InMethodName);
 
+	// Returns the name of a function given its index
+	UFUNCTION()
+	FString GetRigVMFunctionName(int32 InFunctionIndex) const;
+
 	// The default mutable work memory
 	UPROPERTY()
 	FRigVMMemoryContainer WorkMemory;
@@ -145,7 +166,10 @@ public:
 	const FRigVMInstructionArray& GetInstructions();
 
 	// Returns the parameters of the VM
-	const TArray<FRigVMParameter>& GetParameters();
+	const TArray<FRigVMParameter>& GetParameters() const;
+
+	// Returns a parameter given it's name
+	FRigVMParameter GetParameterByName(const FName& InParameterName);
 
 	// Adds a new input / output to the VM
 	template<class T>
@@ -202,31 +226,38 @@ public:
 
 	// Retrieve the value of a parameter
 	template<class T>
-	const T& GetParameterValue(const FRigVMParameter& InParameter, int32 InArrayIndex = 0)
+	T GetParameterValue(const FRigVMParameter& InParameter, int32 InArrayIndex = 0, T DefaultValue = T{})
 	{
-		return WorkMemory.GetArray<T>(InParameter.GetRegisterIndex())[InArrayIndex];
+		if (InParameter.GetRegisterIndex() != INDEX_NONE)
+		{
+			return WorkMemory.GetArray<T>(InParameter.GetRegisterIndex())[InArrayIndex];
+		}
+		return DefaultValue;
 	}
 
 	// Retrieve the value of a parameter given its index
 	template<class T>
-	const T& GetParameterValue(int32 InParameterIndex, int32 InArrayIndex = 0)
+	T GetParameterValue(int32 InParameterIndex, int32 InArrayIndex = 0, T DefaultValue = T{})
 	{
-		return GetParameterValue<T>(Parameters[InParameterIndex], InArrayIndex);
+		return GetParameterValue<T>(Parameters[InParameterIndex], InArrayIndex, DefaultValue);
 	}
 
 	// Retrieve the value of a parameter given its name
 	template<class T>
-	const T& GetParameterValue(const FName& InParameterName, int32 InArrayIndex = 0)
+	T GetParameterValue(const FName& InParameterName, int32 InArrayIndex = 0, T DefaultValue = T{})
 	{
 		int32 ParameterIndex = ParametersNameMap.FindChecked(InParameterName);
-		return GetParameterValue<T>(ParameterIndex, InArrayIndex);
+		return GetParameterValue<T>(ParameterIndex, InArrayIndex, DefaultValue);
 	}
 
 	// Set the value of a parameter
 	template<class T>
 	void SetParameterValue(const FRigVMParameter& InParameter, const T& InNewValue, int32 InArrayIndex = 0)
 	{
-		WorkMemory.GetArray<T>(InParameter.GetRegisterIndex())[InArrayIndex] = InNewValue;
+		if (InParameter.GetRegisterIndex() != INDEX_NONE)
+		{
+			WorkMemory.GetArray<T>(InParameter.GetRegisterIndex())[InArrayIndex] = InNewValue;
+		}
 	}
 
 	// Set the value of a parameter given its index
@@ -277,25 +308,25 @@ public:
 	UFUNCTION(BlueprintCallable, Category = RigVM)
 	FVector2D GetParameterValueVector2D(const FName& InParameterName, int32 InArrayIndex = 0)
 	{
-		return GetParameterValue<FVector2D>(InParameterName, InArrayIndex);
+		return GetParameterValue<FVector2D>(InParameterName, InArrayIndex, FVector2D::ZeroVector);
 	}
 
 	UFUNCTION(BlueprintCallable, Category = RigVM)
 	FVector GetParameterValueVector(const FName& InParameterName, int32 InArrayIndex = 0)
 	{
-		return GetParameterValue<FVector>(InParameterName, InArrayIndex);
+		return GetParameterValue<FVector>(InParameterName, InArrayIndex, FVector::ZeroVector);
 	}
 
 	UFUNCTION(BlueprintCallable, Category = RigVM)
 	FQuat GetParameterValueQuat(const FName& InParameterName, int32 InArrayIndex = 0)
 	{
-		return GetParameterValue<FQuat>(InParameterName, InArrayIndex);
+		return GetParameterValue<FQuat>(InParameterName, InArrayIndex, FQuat::Identity);
 	}
 
 	UFUNCTION(BlueprintCallable, Category = RigVM)
 	FTransform GetParameterValueTransform(const FName& InParameterName, int32 InArrayIndex = 0)
 	{
-		return GetParameterValue<FTransform>(InParameterName, InArrayIndex);
+		return GetParameterValue<FTransform>(InParameterName, InArrayIndex, FTransform::Identity);
 	}
 
 	UFUNCTION(BlueprintCallable, Category = RigVM)
@@ -352,6 +383,26 @@ public:
 		SetParameterValue<FTransform>(InParameterName, InValue, InArrayIndex);
 	}
 
+	void SetRegisterValueFromString(const FRigVMOperand& InOperand, const FString& InCPPType, const UObject* InCPPTypeObject, const TArray<FString>& InDefaultValues);
+
+	// returns the statistics information
+	FRigVMStatistics GetStatistics() const
+	{
+		FRigVMStatistics Statistics;
+		Statistics.LiteralMemory = LiteralMemory.GetStatistics();
+		Statistics.WorkMemory = WorkMemory.GetStatistics();
+		Statistics.ByteCode = ByteCode.GetStatistics();
+		Statistics.BytesForCaching = FirstPointerForInstruction.GetAllocatedSize() + CachedMemoryPointers.GetAllocatedSize();
+		Statistics.BytesForCDO =
+			Statistics.LiteralMemory.TotalBytes +
+			Statistics.WorkMemory.TotalBytes +
+			Statistics.ByteCode.DataBytes +
+			Statistics.BytesForCaching;
+		// Statistics.BytesPerInstance = Statistics.WorkMemory.DataBytes;
+		Statistics.BytesPerInstance = Statistics.BytesForCDO;
+		return Statistics;
+	}
+
 private:
 
 	FORCEINLINE FRigVMParameter AddPlainParameter(ERigVMParameterType InParameterType, const FName& InName, const FString& InCPPType, const uint8* InPtr, uint16 InSize, int32 InArraySize)
@@ -386,12 +437,14 @@ private:
 
 	void ResolveFunctionsIfRequired();
 	void RefreshInstructionsIfRequired();
+	void InvalidateCachedMemory();
+	void CacheMemoryPointersIfRequired(FRigVMMemoryContainerPtrArray InMemory);
 
 	UPROPERTY(transient)
 	FRigVMInstructionArray Instructions;
 
 	UPROPERTY()
-	TArray<FString> FunctionNames;
+	TArray<FName> FunctionNames;
 
 	TArray<FRigVMFunctionPtr> Functions;
 
@@ -401,5 +454,9 @@ private:
 	UPROPERTY()
 	TMap<FName, int32> ParametersNameMap;
 
+	TArray<uint32> FirstPointerForInstruction;
+	TArray<void*> CachedMemoryPointers;
+	TArray<FRigVMMemoryContainer*> CachedMemory;
+
 	friend class URigVMCompiler;
 };
diff --git a/Engine/Source/Runtime/RigVM/Public/RigVMCore/RigVMByteCode.h b/Engine/Source/Runtime/RigVM/Public/RigVMCore/RigVMByteCode.h
--- a/Engine/Source/Runtime/RigVM/Public/RigVMCore/RigVMByteCode.h
+++ b/Engine/Source/Runtime/RigVM/Public/RigVMCore/RigVMByteCode.h
@@ -4,6 +4,7 @@
 
 #include "CoreMinimal.h"
 #include "RigVMRegistry.h"
+#include "RigVMStatistics.h"
 #include "RigVMByteCode.generated.h"
 
 struct FRigVMByteCode;
@@ -324,9 +325,12 @@ public:
 
 	FRigVMInstructionArray();
 
-	// Resets the data structure and removes all storage.
+	// Resets the data structure and maintains all storage.
 	void Reset();
 
+	// Resets the data structure and removes all storage.
+	void Empty();
+
 	// Returns true if a given instruction index is valid.
 	FORCEINLINE bool IsValidIndex(int32 InIndex) const { return Instructions.IsValidIndex(InIndex); }
 
@@ -362,9 +366,19 @@ public:
 
 	FRigVMByteCode();
 
-	// resets the container and removes all memory
+	bool Serialize(FArchive& Ar);
+	FORCEINLINE friend FArchive& operator<<(FArchive& Ar, FRigVMByteCode& P)
+	{
+		P.Serialize(Ar);
+		return Ar;
+	}
+
+	// resets the container and maintains all memory
 	void Reset();
 
+	// resets the container and removes all memory
+	void Empty();
+
 	// returns the number of instructions in this container
 	uint64 Num() const;
 
@@ -466,6 +480,15 @@ public:
 		return TArrayView<uint8>((uint8*)Data, ByteCode.Num());
 	}
 
+	// returns the statistics information
+	FRigVMByteCodeStatistics GetStatistics() const
+	{
+		FRigVMByteCodeStatistics Statistics;
+		Statistics.InstructionCount = GetInstructions().Num();
+		Statistics.DataBytes = ByteCode.GetAllocatedSize();
+		return Statistics;
+	}
+
 private:
 
 	template<class OpType>
diff --git a/Engine/Source/Runtime/RigVM/Public/RigVMCore/RigVMMemory.h b/Engine/Source/Runtime/RigVM/Public/RigVMCore/RigVMMemory.h
--- a/Engine/Source/Runtime/RigVM/Public/RigVMCore/RigVMMemory.h
+++ b/Engine/Source/Runtime/RigVM/Public/RigVMCore/RigVMMemory.h
@@ -4,6 +4,7 @@
 
 #include "CoreMinimal.h"
 #include "UObject/ObjectMacros.h"
+#include "RigVMStatistics.h"
 #include "RigVMMemory.generated.h"
 
 /**
@@ -26,8 +27,11 @@ enum class ERigVMMemoryType: uint8
  * The FRigVMOperand is a light weight address for a register in
  * a FRigVMMemoryContainer.
  */
+USTRUCT()
 struct RIGVM_API FRigVMOperand
 {
+	GENERATED_BODY()
+
 public:
 
 	FRigVMOperand()
@@ -44,6 +48,21 @@ public:
 	{
 	}
 
+	FORCEINLINE bool operator == (const FRigVMOperand& InOther) const
+	{
+		return MemoryType == InOther.MemoryType &&
+			RegisterIndex == InOther.RegisterIndex &&
+			RegisterOffset == InOther.RegisterOffset;
+	}
+
+	FORCEINLINE bool operator != (const FRigVMOperand& InOther) const
+	{
+		return !(*this == InOther);
+	}
+
+	// returns the memory type of this argument
+	FORCEINLINE bool IsValid() const { return RegisterIndex != UINT16_MAX; }
+
 	// returns the memory type of this argument
 	FORCEINLINE ERigVMMemoryType GetMemoryType() const { return MemoryType; }
 
@@ -58,8 +77,13 @@ public:
 
 private:
 
+	UPROPERTY()
 	ERigVMMemoryType MemoryType;
+
+	UPROPERTY()
 	uint16 RegisterIndex;
+	
+	UPROPERTY()
 	uint16 RegisterOffset;
 };
 
@@ -104,6 +128,7 @@ struct RIGVM_API FRigVMRegister
 		, TrailingBytes(0)
 		, Name(NAME_None)
 		, ScriptStructIndex(INDEX_NONE)
+		, bIsArray(false)
 	{
 	}
 
@@ -149,6 +174,10 @@ struct RIGVM_API FRigVMRegister
 	UPROPERTY()
 	int32 ScriptStructIndex;
 
+	// If true defines this register as an array
+	UPROPERTY()
+	bool bIsArray;
+
 	bool Serialize(FArchive& Ar);
 	FORCEINLINE friend FArchive& operator<<(FArchive& Ar, FRigVMRegister& P)
 	{
@@ -168,7 +197,7 @@ struct RIGVM_API FRigVMRegister
 	FORCEINLINE uint8 GetAlignmentBytes() const { return AlignmentBytes; }
 
 	// Returns true if the register stores more than one element
-	FORCEINLINE bool IsArray() const { return ElementCount > 1; }
+	FORCEINLINE bool IsArray() const { return bIsArray || (ElementCount > 1); }
 
 	// Returns true if the register stores plain / shallow memory
 	FORCEINLINE bool IsPlain() const { return ScriptStructIndex == INDEX_NONE; }
@@ -330,9 +359,12 @@ public:
 	// returns the number of registers in this container
 	FORCEINLINE int32 Num() const { return Registers.Num(); }
 
-	// resets the container and removes all storage.
+	// resets the container but maintains storage.
 	void Reset();
 
+	// resets the container and removes all storage.
+	void Empty();
+
 	// const accessor for a register based on index
 	FORCEINLINE const FRigVMRegister& operator[](int32 InIndex) const { return Registers[InIndex]; }
 
@@ -397,7 +429,6 @@ public:
 	// Note: This refers to the active slice - and can change over time.
 	FORCEINLINE const uint8* GetData(const FRigVMRegister& Register, int32 InRegisterOffset = INDEX_NONE) const
 	{
-		ensure(Register.ElementCount > 0);
 		uint8* Ptr = (uint8*)&Data[Register.GetWorkByteIndex()];
 		if (InRegisterOffset != INDEX_NONE)
 		{
@@ -419,7 +450,11 @@ public:
 	// Note: This refers to the active slice - and can change over time.
 	FORCEINLINE uint8* GetData(FRigVMRegister& Register, int32 InRegisterOffset = INDEX_NONE, bool bMoveToNextSlice = false)
 	{
-		ensure(Register.ElementCount > 0);
+		if (Register.ElementCount <= 0)
+		{
+			return nullptr;
+		}
+
 		if(bMoveToNextSlice)
 		{
 			Register.MoveToNextSlice();
@@ -446,7 +481,6 @@ public:
 	template<class T>
 	FORCEINLINE const T* Get(const FRigVMRegister& InRegister, int32 InRegisterOffset = INDEX_NONE) const
 	{
-		ensure(InRegister.ElementCount > 0);
 		return (const T*)GetData(InRegister, InRegisterOffset);
 	}
 
@@ -497,7 +531,6 @@ public:
 	template<class T>
 	FORCEINLINE T* Get(FRigVMRegister& InRegister, int32 InRegisterOffset = INDEX_NONE, bool bMoveToNextSlice = false)
 	{
-		ensure(InRegister.ElementCount > 0);
 		if(bMoveToNextSlice)
 		{
 			InRegister.MoveToNextSlice();
@@ -551,7 +584,6 @@ public:
 	template<class T>
 	FORCEINLINE TArrayView<T> GetArray(FRigVMRegister& InRegister, int32 InRegisterOffset = INDEX_NONE, bool bMoveToNextSlice = false)
 	{
-		ensure(InRegister.ElementCount > 0);
 		if (InRegisterOffset == INDEX_NONE)
 		{
 			return TArrayView<T>((T*)GetData(InRegister, INDEX_NONE, bMoveToNextSlice), InRegister.ElementCount);
@@ -770,6 +802,19 @@ public:
 	// Adds a register path and returns its index
 	int32 GetOrAddRegisterOffset(int32 InRegisterIndex, UScriptStruct* InScriptStruct, const FString& InSegmentPath, int32 InInitialOffset = 0, int32 InElementSize = 0);
 
+	void SetRegisterValueFromString(const FRigVMOperand& InOperand, const FString& InCPPType, const UObject* InCPPTypeObject, const TArray<FString>& InDefaultValues);
+	TArray<FString> GetRegisterValueAsString(const FRigVMOperand& InOperand, const FString& InCPPType, const UObject* InCPPTypeObject);
+
+	// returns the statistics information
+	FRigVMMemoryStatistics GetStatistics() const
+	{
+		FRigVMMemoryStatistics Statistics;
+		Statistics.RegisterCount = Registers.Num();
+		Statistics.DataBytes = Data.GetAllocatedSize();
+		Statistics.TotalBytes = Data.GetAllocatedSize() + Registers.GetAllocatedSize() + RegisterOffsets.GetAllocatedSize();
+		return Statistics;
+	}
+
 private:
 
 	// Adds a new named register for a plain array from a data pointer (used by compiler)
diff --git a/Engine/Source/Runtime/RigVM/Public/RigVMCore/RigVMRegistry.h b/Engine/Source/Runtime/RigVM/Public/RigVMCore/RigVMRegistry.h
--- a/Engine/Source/Runtime/RigVM/Public/RigVMCore/RigVMRegistry.h
+++ b/Engine/Source/Runtime/RigVM/Public/RigVMCore/RigVMRegistry.h
@@ -5,8 +5,9 @@
 #include "CoreMinimal.h"
 #include "RigVMMemory.h"
 
+typedef TArrayView<void*> FRigVMOperandMemory;
 typedef TArrayView<void*> FRigVMUserDataArray;
-typedef void (*FRigVMFunctionPtr)(const FRigVMOperandArray&, FRigVMMemoryContainerPtrArray&, const FRigVMUserDataArray&);
+typedef void (*FRigVMFunctionPtr)(const FName& RigVMOperatorName, int32 RigVMOperatorIndex, FRigVMOperandMemory RigVMOperandMemory, const FRigVMUserDataArray& RigVMUserData);
 
 /**
  * The FRigVMFunction is used to represent a function pointer generated by UHT 
@@ -15,16 +16,19 @@ typedef void (*FRigVMFunctionPtr)(const FRigVMOperandArray&, FRigVMMemoryContain
 struct RIGVM_API FRigVMFunction
 {
 	const TCHAR* Name;
+	UScriptStruct* Struct;
 	FRigVMFunctionPtr FunctionPtr;
 
 	FRigVMFunction()
 		: Name(nullptr)
+		, Struct(nullptr)
 		, FunctionPtr(nullptr)
 	{
 	}
 
-	FRigVMFunction(const TCHAR* InName, FRigVMFunctionPtr InFunctionPtr)
+	FRigVMFunction(const TCHAR* InName, FRigVMFunctionPtr InFunctionPtr, UScriptStruct* InStruct = nullptr)
 		: Name(InName)
+		, Struct(InStruct)
 		, FunctionPtr(InFunctionPtr)
 	{
 	}
@@ -46,7 +50,7 @@ public:
 	// Registers a function given its name.
 	// The name will be the name of the struct and virtual method,
 	// for example "FMyStruct::MyVirtualMethod"
-	void Register(const TCHAR* InName, FRigVMFunctionPtr InFunctionPtr);
+	void Register(const TCHAR* InName, FRigVMFunctionPtr InFunctionPtr, UScriptStruct* InStruct = nullptr);
 
 	// Refreshes the list and finds the function pointers
 	// based on the names.
@@ -55,6 +59,9 @@ public:
 	// Returns a function pointer given its name
 	FRigVMFunctionPtr Find(const TCHAR* InName) const;
 
+	// Returns all current rigvm functions
+	const TArray<FRigVMFunction>& GetFunctions() const;
+
 private:
 
 	// disable default constructor
diff --git a/Engine/Source/Runtime/RigVM/Public/RigVMCore/RigVMStatistics.h b/Engine/Source/Runtime/RigVM/Public/RigVMCore/RigVMStatistics.h
new file mode 100644
--- /dev/null
+++ b/Engine/Source/Runtime/RigVM/Public/RigVMCore/RigVMStatistics.h
@@ -0,0 +1,63 @@
+// Copyright Epic Games, Inc. All Rights Reserved.
+
+#pragma once
+
+#include "CoreMinimal.h"
+#include "RigVMStatistics.generated.h"
+
+USTRUCT()
+struct FRigVMMemoryStatistics
+{
+	GENERATED_BODY()
+
+public:
+
+	UPROPERTY(VisibleAnywhere, Category = "Statistics")
+	uint32 RegisterCount;
+
+	UPROPERTY(VisibleAnywhere, Category = "Statistics")
+	uint32 DataBytes;
+
+	UPROPERTY(VisibleAnywhere, Category = "Statistics")
+	uint32 TotalBytes;
+};
+
+USTRUCT()
+struct FRigVMByteCodeStatistics
+{
+	GENERATED_BODY()
+
+public:
+
+	UPROPERTY(VisibleAnywhere, Category = "Statistics")
+	uint32 InstructionCount;
+
+	UPROPERTY(VisibleAnywhere, Category = "Statistics")
+	uint32 DataBytes;
+};
+
+USTRUCT()
+struct FRigVMStatistics
+{
+	GENERATED_BODY()
+
+public:
+
+	UPROPERTY(VisibleAnywhere, Category = "Statistics")
+	uint32 BytesForCDO;
+
+	UPROPERTY(VisibleAnywhere, Category = "Statistics")
+	uint32 BytesPerInstance;
+
+	UPROPERTY(VisibleAnywhere, Category = "Statistics")
+	FRigVMMemoryStatistics LiteralMemory;
+
+	UPROPERTY(VisibleAnywhere, Category = "Statistics")
+	FRigVMMemoryStatistics WorkMemory;
+
+	UPROPERTY(VisibleAnywhere, Category = "Statistics")
+	uint32 BytesForCaching;
+
+	UPROPERTY(VisibleAnywhere, Category = "Statistics")
+	FRigVMByteCodeStatistics ByteCode;
+};
\ No newline at end of file
diff --git a/Engine/Source/Runtime/RigVM/Public/RigVMCore/RigVMStruct.h b/Engine/Source/Runtime/RigVM/Public/RigVMCore/RigVMStruct.h
new file mode 100644
--- /dev/null
+++ b/Engine/Source/Runtime/RigVM/Public/RigVMCore/RigVMStruct.h
@@ -0,0 +1,24 @@
+// Copyright Epic Games, Inc. All Rights Reserved.
+
+#pragma once
+
+#include "CoreMinimal.h"
+#include "RigVMCore/RigVMRegistry.h"
+#include "RigVMStruct.generated.h"
+
+/**
+ * The base class for all RigVM enabled structs.
+ */
+USTRUCT()
+struct RIGVM_API FRigVMStruct
+{
+	GENERATED_BODY()
+
+	virtual ~FRigVMStruct() {}
+	virtual FString ProcessPinLabelForInjection(const FString& InLabel) const { return InLabel; }
+	virtual FName DetermineSpaceForPin(const FString& InPinPath, void* InUserContext) const { return NAME_None; }
+
+public:
+
+	FORCEINLINE virtual int32 GetMaxArraySize(const FName& InParameterName, const FRigVMUserDataArray& RigVMUserData) { return INDEX_NONE; }
+};
