Fix for warning after upgrading to PVS-Studio 7.0.4:

V1028: Possible overflow. Consider casting operands, not the result.

#rnx

diff --git a/Engine/Source/Runtime/RigVM/Public/RigVMCore/RigVMMemory.h b/Engine/Source/Runtime/RigVM/Public/RigVMCore/RigVMMemory.h
--- a/Engine/Source/Runtime/RigVM/Public/RigVMCore/RigVMMemory.h
+++ b/Engine/Source/Runtime/RigVM/Public/RigVMCore/RigVMMemory.h
@@ -191,7 +191,7 @@ struct RIGVM_API FRigVMRegister
 	FORCEINLINE uint64 GetWorkByteIndex() const { return ByteIndex; }
 
 	// returns the first allocated byte in the data byte array
-	FORCEINLINE uint64 GetFirstAllocatedByte() const { return ByteIndex - (uint64)AlignmentBytes - (uint64)(SliceIndex * GetNumBytesPerSlice()); }
+	FORCEINLINE uint64 GetFirstAllocatedByte() const { return ByteIndex - (uint64)AlignmentBytes - ((uint64)SliceIndex * GetNumBytesPerSlice()); }
 
 	// Returns the leading alignment bytes
 	FORCEINLINE uint8 GetAlignmentBytes() const { return AlignmentBytes; }
