RigVM: Add GetRegister accessors to avoid array out of range

diff --git a/Engine/Source/Runtime/RigVM/Private/RigVMCore/RigVM.cpp b/Engine/Source/Runtime/RigVM/Private/RigVMCore/RigVM.cpp
--- a/Engine/Source/Runtime/RigVM/Private/RigVMCore/RigVM.cpp
+++ b/Engine/Source/Runtime/RigVM/Private/RigVMCore/RigVM.cpp
@@ -353,7 +353,7 @@ void URigVM::CacheMemoryPointersIfRequired(FRigVMMemoryContainerPtrArray InMemor
 					void* Ptr = CachedMemory[Arg.GetContainerIndex()]->GetData(Arg.GetRegisterIndex(), Arg.GetRegisterOffset());
 					CachedMemoryPointers.Add(Ptr);
 
-					FRigVMRegister Register = (*CachedMemory[Arg.GetContainerIndex()])[Arg];
+					const FRigVMRegister& Register = CachedMemory[Arg.GetContainerIndex()]->GetRegister(Arg);
 					if (Register.IsArray())
 					{
 						void* NumBytesAsPtr = reinterpret_cast<void*>(Register.ElementCount);
diff --git a/Engine/Source/Runtime/RigVM/Public/RigVMCore/RigVMMemory.h b/Engine/Source/Runtime/RigVM/Public/RigVMCore/RigVMMemory.h
--- a/Engine/Source/Runtime/RigVM/Public/RigVMCore/RigVMMemory.h
+++ b/Engine/Source/Runtime/RigVM/Public/RigVMCore/RigVMMemory.h
@@ -366,28 +366,46 @@ public:
 	void Empty();
 
 	// const accessor for a register based on index
-	FORCEINLINE const FRigVMRegister& operator[](int32 InIndex) const { return Registers[InIndex]; }
+	FORCEINLINE const FRigVMRegister& operator[](int32 InIndex) const { return GetRegister(InIndex); }
 
 	// accessor for a register based on index
-	FORCEINLINE FRigVMRegister& operator[](int32 InIndex) { return Registers[InIndex]; }
+	FORCEINLINE FRigVMRegister& operator[](int32 InIndex) { return GetRegister(InIndex); }
 
 	// const accessor for a register based on an argument
-	FORCEINLINE const FRigVMRegister& operator[](const FRigVMOperand& InArg) const { return Registers[InArg.GetRegisterIndex()]; }
+	FORCEINLINE const FRigVMRegister& operator[](const FRigVMOperand& InArg) const { return GetRegister(InArg); }
 
 	// accessor for a register based on an argument
-	FORCEINLINE FRigVMRegister& operator[](const FRigVMOperand& InArg) { return Registers[InArg.GetRegisterIndex()]; }
+	FORCEINLINE FRigVMRegister& operator[](const FRigVMOperand& InArg) { return GetRegister(InArg); }
 
 	// const accessor for a register based on a a name. note: only works if SupportsNames() == true
-	FORCEINLINE const FRigVMRegister& operator[](const FName& InName) const { return Registers[GetIndex(InName)]; }
+	FORCEINLINE const FRigVMRegister& operator[](const FName& InName) const { return GetRegister(InName); }
 
 	// accessor for a register based on a a name. note: only works if SupportsNames() == true
-	FORCEINLINE FRigVMRegister& operator[](const FName& InName) { return Registers[GetIndex(InName)]; }
+	FORCEINLINE FRigVMRegister& operator[](const FName& InName) { return GetRegister(InName); }
 
 	FORCEINLINE TArray<FRigVMRegister>::RangedForIteratorType      begin() { return Registers.begin(); }
 	FORCEINLINE TArray<FRigVMRegister>::RangedForConstIteratorType begin() const { return Registers.begin(); }
 	FORCEINLINE TArray<FRigVMRegister>::RangedForIteratorType      end() { return Registers.end(); }
 	FORCEINLINE TArray<FRigVMRegister>::RangedForConstIteratorType end() const { return Registers.end(); }
 
+	// const accessor for a register based on index
+	FORCEINLINE const FRigVMRegister& GetRegister(int32 InIndex) const { return Registers[InIndex]; }
+
+	// accessor for a register based on index
+	FORCEINLINE FRigVMRegister& GetRegister(int32 InIndex) { return Registers[InIndex]; }
+
+	// const accessor for a register based on an argument
+	FORCEINLINE const FRigVMRegister& GetRegister(const FRigVMOperand& InArg) const { return Registers[InArg.GetRegisterIndex()]; }
+
+	// accessor for a register based on an argument
+	FORCEINLINE FRigVMRegister& GetRegister(const FRigVMOperand& InArg) { return Registers[InArg.GetRegisterIndex()]; }
+
+	// const accessor for a register based on a a name. note: only works if SupportsNames() == true
+	FORCEINLINE const FRigVMRegister& GetRegister(const FName& InName) const { return Registers[GetIndex(InName)]; }
+
+	// accessor for a register based on a a name. note: only works if SupportsNames() == true
+	FORCEINLINE FRigVMRegister& GetRegister(const FName& InName) { return Registers[GetIndex(InName)]; }
+
 	bool Serialize(FArchive& Ar);
 	FORCEINLINE friend FArchive& operator<<(FArchive& Ar, FRigVMMemoryContainer& P)
 	{
