RigVM: Avoid synchronization issues by guarding against multiple threads accessing the same VM.

diff --git a/Engine/Plugins/Experimental/ControlRig/Source/ControlRig/Private/ControlRig.cpp b/Engine/Plugins/Experimental/ControlRig/Source/ControlRig/Private/ControlRig.cpp
--- a/Engine/Plugins/Experimental/ControlRig/Source/ControlRig/Private/ControlRig.cpp
+++ b/Engine/Plugins/Experimental/ControlRig/Source/ControlRig/Private/ControlRig.cpp
@@ -197,6 +197,14 @@ void UControlRig::Execute(const EControlRigState InState)
 {
 	DECLARE_SCOPE_HIERARCHICAL_COUNTER_FUNC()
 
+	if (VM)
+	{
+		if (VM->GetOuter() != this)
+		{
+			InstantiateVMFromCDO();
+		}
+	}
+
 	if (bRequiresInitExecution)
 	{
 		bRequiresInitExecution = false;
diff --git a/Engine/Source/Runtime/RigVM/Private/RigVMCore/RigVM.cpp b/Engine/Source/Runtime/RigVM/Private/RigVMCore/RigVM.cpp
--- a/Engine/Source/Runtime/RigVM/Private/RigVMCore/RigVM.cpp
+++ b/Engine/Source/Runtime/RigVM/Private/RigVMCore/RigVM.cpp
@@ -3,6 +3,7 @@
 #include "RigVMCore/RigVM.h"
 #include "UObject/Package.h"
 #include "UObject/AnimObjectVersion.h"
+#include "HAL/PlatformTLS.h"
 
 bool FRigVMParameter::Serialize(FArchive& Ar)
 {
@@ -41,6 +42,7 @@ UScriptStruct* FRigVMParameter::GetScriptStruct() const
 }
 
 URigVM::URigVM()
+	: ThreadId(INDEX_NONE)
 {
 	WorkMemory.SetMemoryType(ERigVMMemoryType::Work);
 	LiteralMemory.SetMemoryType(ERigVMMemoryType::Literal);
@@ -234,6 +236,12 @@ void URigVM::InvalidateCachedMemory()
 
 void URigVM::CacheMemoryPointersIfRequired(FRigVMMemoryContainerPtrArray InMemory)
 {
+	if (ThreadId != INDEX_NONE)
+	{
+		ensureMsgf(ThreadId == FPlatformTLS::GetCurrentThreadId(), TEXT("RigVM from multiple threads (%d and %d)"), ThreadId, (int32)FPlatformTLS::GetCurrentThreadId());
+	}
+	TGuardValue<int32> GuardThreadId(ThreadId, FPlatformTLS::GetCurrentThreadId());
+
 	RefreshInstructionsIfRequired();
 
 	if (Instructions.Num() == 0 || InMemory.Num() == 0)
diff --git a/Engine/Source/Runtime/RigVM/Public/RigVMCore/RigVM.h b/Engine/Source/Runtime/RigVM/Public/RigVMCore/RigVM.h
--- a/Engine/Source/Runtime/RigVM/Public/RigVMCore/RigVM.h
+++ b/Engine/Source/Runtime/RigVM/Public/RigVMCore/RigVM.h
@@ -458,5 +458,7 @@ private:
 	TArray<void*> CachedMemoryPointers;
 	TArray<FRigVMMemoryContainer*> CachedMemory;
 
+	int32 ThreadId;
+
 	friend class URigVMCompiler;
 };
