Editgrated 14478375 from Release-4.26 (Control Rig / RigVM: Add alignment for non-intel platforms)

diff --git a/Engine/Source/Developer/RigVMDeveloper/Private/RigVMCompiler/RigVMCompiler.cpp b/Engine/Source/Developer/RigVMDeveloper/Private/RigVMCompiler/RigVMCompiler.cpp
--- a/Engine/Source/Developer/RigVMDeveloper/Private/RigVMCompiler/RigVMCompiler.cpp
+++ b/Engine/Source/Developer/RigVMDeveloper/Private/RigVMCompiler/RigVMCompiler.cpp
@@ -215,6 +215,8 @@ bool URigVMCompiler::Compile(URigVMGraph* InGraph, URigVMController* InControlle
 		WorkData.VM->GetByteCode().AddExitOp();
 	}
 
+	WorkData.VM->GetByteCode().AlignByteCode();
+
 	// loop over all nodes once more and setup the instruction index for reroute nodes
 	if (Settings.SetupNodeInstructionIndex)
 	{
diff --git a/Engine/Source/Runtime/RigVM/Private/RigVMCore/RigVMByteCode.cpp b/Engine/Source/Runtime/RigVM/Private/RigVMCore/RigVMByteCode.cpp
--- a/Engine/Source/Runtime/RigVM/Private/RigVMCore/RigVMByteCode.cpp
+++ b/Engine/Source/Runtime/RigVM/Private/RigVMCore/RigVMByteCode.cpp
@@ -69,7 +69,7 @@ FRigVMInstructionArray::FRigVMInstructionArray()
 {
 }
 
-FRigVMInstructionArray::FRigVMInstructionArray(const FRigVMByteCode& InByteCode)
+FRigVMInstructionArray::FRigVMInstructionArray(const FRigVMByteCode& InByteCode, bool bByteCodeIsAligned)
 {
 	uint64 ByteIndex = 0;
 	while (ByteIndex < InByteCode.Num())
@@ -81,8 +81,37 @@ FRigVMInstructionArray::FRigVMInstructionArray(const FRigVMByteCode& InByteCode)
 			Instructions.Reset();
 			break;
 		}
-		Instructions.Add(FRigVMInstruction(OpCode, ByteIndex));
-		ByteIndex += InByteCode.GetOpNumBytesAt(ByteIndex);
+
+		uint8 OperandAlignment = 0;
+
+		if (bByteCodeIsAligned)
+		{
+			uint64 Alignment = InByteCode.GetOpAlignment(OpCode);
+			if (Alignment > 0)
+			{
+				while (!IsAligned(&InByteCode[ByteIndex], Alignment))
+				{
+					ByteIndex++;
+				}
+			}
+
+			if (OpCode >= ERigVMOpCode::Execute_0_Operands && OpCode < ERigVMOpCode::Execute_64_Operands)
+			{
+				uint64 OperandByteIndex = ByteIndex + (uint64)sizeof(FRigVMExecuteOp);
+
+				Alignment = InByteCode.GetOperandAlignment();
+				if (Alignment > 0)
+				{
+					while (!IsAligned(&InByteCode[OperandByteIndex + OperandAlignment], Alignment))
+					{
+						OperandAlignment++;
+					}
+				}
+			}
+		}
+
+		Instructions.Add(FRigVMInstruction(OpCode, ByteIndex, OperandAlignment));
+		ByteIndex += InByteCode.GetOpNumBytesAt(ByteIndex, true);
 	}
 }
 
@@ -98,6 +127,7 @@ void FRigVMInstructionArray::Empty()
 
 FRigVMByteCode::FRigVMByteCode()
 	: NumInstructions(0)
+	, bByteCodeIsAligned(false)
 {
 }
 
@@ -128,6 +158,7 @@ bool FRigVMByteCode::Serialize(FArchive& Ar)
 	else
 	{
 		ByteCode.Reset();
+		bByteCodeIsAligned = false;
 	}
 
 	Ar << InstructionCount;
@@ -218,7 +249,7 @@ bool FRigVMByteCode::Serialize(FArchive& Ar)
 					FRigVMExecuteOp Op = GetOpAt<FRigVMExecuteOp>(Instruction.ByteCodeIndex);
 					Ar << Op;
 
-					FRigVMOperandArray Operands = GetOperandsForExecuteOp(Instruction.ByteCodeIndex);
+					FRigVMOperandArray Operands = GetOperandsForExecuteOp(Instruction);
 					int32 OperandCount = (int32)Op.GetOperandCount();
 					ensure(OperandCount == Operands.Num());
 
@@ -380,6 +411,8 @@ bool FRigVMByteCode::Serialize(FArchive& Ar)
 
 	if (Ar.IsLoading())
 	{
+		AlignByteCode();
+
 		Entries.Reset();
 		if (Ar.CustomVer(FAnimObjectVersion::GUID) >= FAnimObjectVersion::SerializeRigVMEntries)
 		{
@@ -420,6 +453,7 @@ bool FRigVMByteCode::Serialize(FArchive& Ar)
 void FRigVMByteCode::Reset()
 {
 	ByteCode.Reset();
+	bByteCodeIsAligned = false;
 	NumInstructions = 0;
 	Entries.Reset();
 }
@@ -427,6 +461,7 @@ void FRigVMByteCode::Reset()
 void FRigVMByteCode::Empty()
 {
 	ByteCode.Empty();
+	bByteCodeIsAligned = false;
 	NumInstructions = 0;
 	Entries.Empty();
 }
@@ -532,7 +567,20 @@ uint64 FRigVMByteCode::GetOpNumBytesAt(uint64 InByteCodeIndex, bool bIncludeOper
 			uint64 NumBytes = (uint64)sizeof(FRigVMExecuteOp);
 			if(bIncludeOperands)
 			{
-				const FRigVMExecuteOp& ExecuteOp = GetOpAt<FRigVMExecuteOp>(InByteCodeIndex);
+				FRigVMExecuteOp ExecuteOp;
+				FMemory::Memcpy(&ExecuteOp, &ByteCode[InByteCodeIndex], sizeof(FRigVMExecuteOp));
+
+				if (bByteCodeIsAligned)
+				{
+					static const uint64 OperandAlignment = GetOperandAlignment();
+					if (OperandAlignment > 0)
+					{
+						while (!IsAligned(&ByteCode[InByteCodeIndex + NumBytes], OperandAlignment))
+						{
+							NumBytes++;
+						}
+					}
+				}
 				NumBytes += (uint64)ExecuteOp.GetOperandCount() * (uint64)sizeof(FRigVMOperand);
 			}
 			return NumBytes;
@@ -664,6 +712,7 @@ uint64 FRigVMByteCode::AddExecuteOp(uint16 InFunctionIndex, const FRigVMOperandA
 {
 	FRigVMExecuteOp Op(InFunctionIndex, (uint8)InOperands.Num());
 	uint64 OpByteIndex = AddOp(Op);
+
 	uint64 OperandsByteIndex = (uint64)ByteCode.AddZeroed(sizeof(FRigVMOperand) * InOperands.Num());
 	FMemory::Memcpy(ByteCode.GetData() + OperandsByteIndex, InOperands.GetData(), sizeof(FRigVMOperand) * InOperands.Num());
 	return OpByteIndex;
@@ -758,7 +807,7 @@ FString FRigVMByteCode::DumpToText() const
 				const FRigVMExecuteOp& Op = GetOpAt<FRigVMExecuteOp>(Instruction.ByteCodeIndex);
 				Line += FString::Printf(TEXT(", FunctionIndex %d"), (int32)Op.FunctionIndex);
 
-				FRigVMOperandArray Operands = GetOperandsForExecuteOp(Instruction.ByteCodeIndex);
+				FRigVMOperandArray Operands = GetOperandsForExecuteOp(Instruction);
 				if (Operands.Num() > 0)
 				{
 					TArray<FString> OperandsContent;
@@ -864,3 +913,221 @@ uint64 FRigVMByteCode::AddEndBlockOp()
 	return AddOp(Op);
 }
 
+
+uint64 FRigVMByteCode::GetOpAlignment(ERigVMOpCode InOpCode) const
+{
+	switch (InOpCode)
+	{
+		case ERigVMOpCode::Execute_0_Operands:
+		case ERigVMOpCode::Execute_1_Operands:
+		case ERigVMOpCode::Execute_2_Operands:
+		case ERigVMOpCode::Execute_3_Operands:
+		case ERigVMOpCode::Execute_4_Operands:
+		case ERigVMOpCode::Execute_5_Operands:
+		case ERigVMOpCode::Execute_6_Operands:
+		case ERigVMOpCode::Execute_7_Operands:
+		case ERigVMOpCode::Execute_8_Operands:
+		case ERigVMOpCode::Execute_9_Operands:
+		case ERigVMOpCode::Execute_10_Operands:
+		case ERigVMOpCode::Execute_11_Operands:
+		case ERigVMOpCode::Execute_12_Operands:
+		case ERigVMOpCode::Execute_13_Operands:
+		case ERigVMOpCode::Execute_14_Operands:
+		case ERigVMOpCode::Execute_15_Operands:
+		case ERigVMOpCode::Execute_16_Operands:
+		case ERigVMOpCode::Execute_17_Operands:
+		case ERigVMOpCode::Execute_18_Operands:
+		case ERigVMOpCode::Execute_19_Operands:
+		case ERigVMOpCode::Execute_20_Operands:
+		case ERigVMOpCode::Execute_21_Operands:
+		case ERigVMOpCode::Execute_22_Operands:
+		case ERigVMOpCode::Execute_23_Operands:
+		case ERigVMOpCode::Execute_24_Operands:
+		case ERigVMOpCode::Execute_25_Operands:
+		case ERigVMOpCode::Execute_26_Operands:
+		case ERigVMOpCode::Execute_27_Operands:
+		case ERigVMOpCode::Execute_28_Operands:
+		case ERigVMOpCode::Execute_29_Operands:
+		case ERigVMOpCode::Execute_30_Operands:
+		case ERigVMOpCode::Execute_31_Operands:
+		case ERigVMOpCode::Execute_32_Operands:
+		case ERigVMOpCode::Execute_33_Operands:
+		case ERigVMOpCode::Execute_34_Operands:
+		case ERigVMOpCode::Execute_35_Operands:
+		case ERigVMOpCode::Execute_36_Operands:
+		case ERigVMOpCode::Execute_37_Operands:
+		case ERigVMOpCode::Execute_38_Operands:
+		case ERigVMOpCode::Execute_39_Operands:
+		case ERigVMOpCode::Execute_40_Operands:
+		case ERigVMOpCode::Execute_41_Operands:
+		case ERigVMOpCode::Execute_42_Operands:
+		case ERigVMOpCode::Execute_43_Operands:
+		case ERigVMOpCode::Execute_44_Operands:
+		case ERigVMOpCode::Execute_45_Operands:
+		case ERigVMOpCode::Execute_46_Operands:
+		case ERigVMOpCode::Execute_47_Operands:
+		case ERigVMOpCode::Execute_48_Operands:
+		case ERigVMOpCode::Execute_49_Operands:
+		case ERigVMOpCode::Execute_50_Operands:
+		case ERigVMOpCode::Execute_51_Operands:
+		case ERigVMOpCode::Execute_52_Operands:
+		case ERigVMOpCode::Execute_53_Operands:
+		case ERigVMOpCode::Execute_54_Operands:
+		case ERigVMOpCode::Execute_55_Operands:
+		case ERigVMOpCode::Execute_56_Operands:
+		case ERigVMOpCode::Execute_57_Operands:
+		case ERigVMOpCode::Execute_58_Operands:
+		case ERigVMOpCode::Execute_59_Operands:
+		case ERigVMOpCode::Execute_60_Operands:
+		case ERigVMOpCode::Execute_61_Operands:
+		case ERigVMOpCode::Execute_62_Operands:
+		case ERigVMOpCode::Execute_63_Operands:
+		case ERigVMOpCode::Execute_64_Operands:
+		{
+			static const uint64 Alignment = FRigVMExecuteOp::StaticStruct()->GetCppStructOps()->GetAlignment();
+			return Alignment;
+		}
+		case ERigVMOpCode::Copy:
+		{
+			static const uint64 Alignment = FRigVMCopyOp::StaticStruct()->GetCppStructOps()->GetAlignment();
+			return Alignment;
+		}
+		case ERigVMOpCode::Zero:
+		case ERigVMOpCode::BoolFalse:
+		case ERigVMOpCode::BoolTrue:
+		case ERigVMOpCode::Increment:
+		case ERigVMOpCode::Decrement:
+		{
+			static const uint64 Alignment = FRigVMUnaryOp::StaticStruct()->GetCppStructOps()->GetAlignment();
+			return Alignment;
+		}
+		case ERigVMOpCode::Equals:
+		case ERigVMOpCode::NotEquals:
+		{
+			static const uint64 Alignment = FRigVMComparisonOp::StaticStruct()->GetCppStructOps()->GetAlignment();
+			return Alignment;
+		}
+		case ERigVMOpCode::JumpAbsolute:
+		case ERigVMOpCode::JumpForward:
+		case ERigVMOpCode::JumpBackward:
+		{
+			static const uint64 Alignment = FRigVMJumpOp::StaticStruct()->GetCppStructOps()->GetAlignment();
+			return Alignment;
+		}
+		case ERigVMOpCode::JumpAbsoluteIf:
+		case ERigVMOpCode::JumpForwardIf:
+		case ERigVMOpCode::JumpBackwardIf:
+		{
+			static const uint64 Alignment = FRigVMJumpIfOp::StaticStruct()->GetCppStructOps()->GetAlignment();
+			return Alignment;
+		}
+		case ERigVMOpCode::ChangeType:
+		{
+			static const uint64 Alignment = FRigVMChangeTypeOp::StaticStruct()->GetCppStructOps()->GetAlignment();
+			return Alignment;
+		}
+		case ERigVMOpCode::Exit:
+		{
+			static const uint64 Alignment = FRigVMBaseOp::StaticStruct()->GetCppStructOps()->GetAlignment();
+			return Alignment;
+		}
+		case ERigVMOpCode::BeginBlock:
+		{
+			static const uint64 Alignment = FRigVMBinaryOp::StaticStruct()->GetCppStructOps()->GetAlignment();
+			return Alignment;
+		}
+		case ERigVMOpCode::EndBlock:
+		{
+			static const uint64 Alignment = FRigVMBaseOp::StaticStruct()->GetCppStructOps()->GetAlignment();
+			return Alignment;
+		}
+		case ERigVMOpCode::Invalid:
+		{
+			ensure(false);
+			return 0;
+		}
+	}
+	return 0;
+}
+
+uint64 FRigVMByteCode::GetOperandAlignment() const
+{
+	static const uint64 OperandAlignment = !FRigVMOperand::StaticStruct()->GetCppStructOps()->GetAlignment();
+	return OperandAlignment;
+}
+
+void FRigVMByteCode::AlignByteCode()
+{
+	if (bByteCodeIsAligned)
+	{
+		return;
+	}
+
+	FRigVMInstructionArray Instructions(*this, false);
+	uint64 BytesToReserve = ByteCode.Num();
+
+	for (int32 InstructionIndex = 0; InstructionIndex < Instructions.Num(); InstructionIndex++)
+	{
+		const FRigVMInstruction& Instruction = Instructions[InstructionIndex];
+		BytesToReserve += GetOpAlignment(Instruction.OpCode);
+
+		if (Instruction.OpCode >= ERigVMOpCode::Execute_0_Operands && Instruction.OpCode < ERigVMOpCode::Execute_64_Operands)
+		{
+			BytesToReserve += GetOperandAlignment();
+		}
+	}
+
+	TArray<uint8> AlignedByteCode;
+	AlignedByteCode.Reserve(BytesToReserve);
+	AlignedByteCode.AddZeroed(ByteCode.Num());
+
+	uint64 ShiftedBytes = 0;
+	for (int32 InstructionIndex = 0; InstructionIndex < Instructions.Num(); InstructionIndex++)
+	{
+		const FRigVMInstruction& Instruction = Instructions[InstructionIndex];
+		uint64 OriginalByteCodeIndex = Instruction.ByteCodeIndex;
+		uint64 AlignedByteCodeIndex = OriginalByteCodeIndex + ShiftedBytes;
+		uint64 OpAlignment = GetOpAlignment(Instruction.OpCode);
+
+		if (OpAlignment > 0)
+		{
+			while (!IsAligned(&AlignedByteCode[AlignedByteCodeIndex], OpAlignment))
+			{
+				AlignedByteCode[AlignedByteCodeIndex] = (uint8)Instruction.OpCode;
+				AlignedByteCodeIndex++;
+				ShiftedBytes++;
+				AlignedByteCode.AddZeroed(1);
+			}
+		}
+
+		uint64 NumBytesToCopy = GetOpNumBytesAt(OriginalByteCodeIndex, false);
+		FMemory::Memcpy(&AlignedByteCode[AlignedByteCodeIndex], &ByteCode[OriginalByteCodeIndex], NumBytesToCopy);
+
+		if (Instruction.OpCode >= ERigVMOpCode::Execute_0_Operands && Instruction.OpCode < ERigVMOpCode::Execute_64_Operands)
+		{
+			AlignedByteCodeIndex += NumBytesToCopy;
+
+			static const uint64 OperandAlignment = GetOperandAlignment();
+			if (OperandAlignment > 0)
+			{
+				while (!IsAligned(&AlignedByteCode[AlignedByteCodeIndex], OperandAlignment))
+				{
+					AlignedByteCodeIndex++;
+					ShiftedBytes++;
+					AlignedByteCode.AddZeroed(1);
+				}
+			}
+
+			FRigVMExecuteOp ExecuteOp;
+			FMemory::Memcpy(&ExecuteOp, &ByteCode[OriginalByteCodeIndex], sizeof(FRigVMExecuteOp));
+
+			OriginalByteCodeIndex += NumBytesToCopy;
+			NumBytesToCopy = sizeof(FRigVMOperand) * ExecuteOp.GetOperandCount();
+
+			FMemory::Memcpy(&AlignedByteCode[AlignedByteCodeIndex], &ByteCode[OriginalByteCodeIndex], NumBytesToCopy);
+		}
+	}
+
+	Swap(ByteCode, AlignedByteCode);
+	bByteCodeIsAligned = true;
+}
\ No newline at end of file
diff --git a/Engine/Source/Runtime/RigVM/Private/RigVMCore/RigVMMemory.cpp b/Engine/Source/Runtime/RigVM/Private/RigVMCore/RigVMMemory.cpp
--- a/Engine/Source/Runtime/RigVM/Private/RigVMCore/RigVMMemory.cpp
+++ b/Engine/Source/Runtime/RigVM/Private/RigVMCore/RigVMMemory.cpp
@@ -595,8 +595,8 @@ bool FRigVMMemoryContainer::Serialize(FArchive& Ar)
 						Ar << View;
 						if (!bEncounteredErrorDuringLoad)
 						{
-							ensure(View.Num() == Register.GetAllocatedBytes());
-							FMemory::Memcpy(&Data[Register.GetFirstAllocatedByte()], View.GetData(), View.Num());
+							ensure(View.Num() <= Register.GetAllocatedBytes());
+							FMemory::Memcpy(&Data[Register.GetWorkByteIndex()], View.GetData(), View.Num());
 						}
 						break;
 					}
@@ -607,10 +607,7 @@ bool FRigVMMemoryContainer::Serialize(FArchive& Ar)
 						if (!bEncounteredErrorDuringLoad)
 						{
 							ensure(View.Num() == Register.GetTotalElementCount());
-							for (uint16 ElementIndex = 0; ElementIndex < Register.GetTotalElementCount(); ElementIndex++)
-							{
-								*((FName*)&Data[Register.GetWorkByteIndex() + Register.ElementSize * ElementIndex]) = View[ElementIndex];
-							}
+							RigVMCopy<FName>(&Data[Register.GetWorkByteIndex()], View.GetData(), View.Num());
 						}
 						break;
 					}
@@ -621,10 +618,7 @@ bool FRigVMMemoryContainer::Serialize(FArchive& Ar)
 						if (!bEncounteredErrorDuringLoad)
 						{
 							ensure(View.Num() == Register.GetTotalElementCount());
-							for (uint16 ElementIndex = 0; ElementIndex < Register.GetTotalElementCount(); ElementIndex++)
-							{
-								*((FString*)&Data[Register.GetWorkByteIndex() + Register.ElementSize * ElementIndex]) = View[ElementIndex];
-							}
+							RigVMCopy<FString>(&Data[Register.GetWorkByteIndex()], View.GetData(), View.Num());
 						}
 						break;
 					}
@@ -861,7 +855,7 @@ bool FRigVMMemoryContainer::Serialize(FArchive& Ar)
 					case ERigVMRegisterType::Plain:
 					{
 						FRigVMByteArray View;
-						View.Append(&Data[Register.GetFirstAllocatedByte()], Register.GetAllocatedBytes());
+						View.Append(&Data[Register.GetWorkByteIndex()], Register.GetAllocatedBytes() - Register.GetAlignmentBytes());
 						Ar << View;
 						break;
 					}
@@ -1066,24 +1060,14 @@ bool FRigVMMemoryContainer::Copy(
 		{
 			int32 NumNames = InNumBytes / sizeof(FName);
 			ensure(NumNames * sizeof(FName) == InNumBytes);
-			FRigVMFixedArray<FName> TargetNames((FName*)InTargetPtr, NumNames);
-			FRigVMFixedArray<FName> SourceNames((FName*)InSourcePtr, NumNames);
-			for (int32 Index = 0; Index < NumNames; Index++)
-			{
-				TargetNames[Index] = SourceNames[Index];
-			}
+			RigVMCopy<FName>(InTargetPtr, InSourcePtr, NumNames);
 			break;
 		}
 		case ERigVMRegisterType::String:
 		{
 			int32 NumStrings = InNumBytes / sizeof(FString);
 			ensure(NumStrings * sizeof(FString) == InNumBytes);
-			FRigVMFixedArray<FString> TargetStrings((FString*)InTargetPtr, NumStrings);
-			FRigVMFixedArray<FString> SourceStrings((FString*)InSourcePtr, NumStrings);
-			for (int32 Index = 0; Index < NumStrings; Index++)
-			{
-				TargetStrings[Index] = SourceStrings[Index];
-			}
+			RigVMCopy<FString>(InTargetPtr, InSourcePtr, NumStrings);
 			break;
 		}
 		case ERigVMRegisterType::Invalid:
@@ -1788,38 +1772,46 @@ void FRigVMMemoryContainer::UpdateRegisters()
 		FRigVMRegister& Register = Registers[RegisterIndex];
 		Register.ByteIndex += AlignmentShift;
 
-		UScriptStruct* ScriptStruct = GetScriptStruct(RegisterIndex);
-		if (ScriptStruct != nullptr && !Register.IsDynamic())
+		int32 Alignment = 4;
+
+		if (Register.IsDynamic() ||
+			Register.Type == ERigVMRegisterType::Name ||
+			Register.Type == ERigVMRegisterType::String)
+		{
+			Alignment = 8;
+		}
+		else if (UScriptStruct* ScriptStruct = GetScriptStruct(RegisterIndex))
 		{
-			UScriptStruct::ICppStructOps* TheCppStructOps = ScriptStruct->GetCppStructOps();
-			if (TheCppStructOps != NULL)
+			if (UScriptStruct::ICppStructOps* TheCppStructOps = ScriptStruct->GetCppStructOps())
 			{
-				if (!TheCppStructOps->HasZeroConstructor())
-				{
-					uint8* Pointer = GetData(RegisterIndex);
+				Alignment = TheCppStructOps->GetAlignment();
+			}
+		}
 
-					if (Register.AlignmentBytes > 0)
-					{
-						if (!IsAligned(Pointer, TheCppStructOps->GetAlignment()))
-						{
-							Data.RemoveAt(Register.GetFirstAllocatedByte(), Register.AlignmentBytes);
-							AlignmentShift -= Register.AlignmentBytes;
-							Register.ByteIndex -= Register.AlignmentBytes;
-							Register.AlignmentBytes = 0;
-							Pointer = GetData(RegisterIndex);
-						}
-					}
+		if (Alignment != 0)
+		{
+			uint8* Pointer = GetData(RegisterIndex);
 
-					while (!IsAligned(Pointer, TheCppStructOps->GetAlignment()))
-					{
-						Data.InsertZeroed(Register.GetFirstAllocatedByte(), 1);
-						Register.AlignmentBytes++;
-						Register.ByteIndex++;
-						AlignmentShift++;
-						Pointer = GetData(RegisterIndex);
-					}
+			if (Register.AlignmentBytes > 0)
+			{
+				if (!IsAligned(Pointer, Alignment))
+				{
+					Data.RemoveAt(Register.GetFirstAllocatedByte(), Register.AlignmentBytes);
+					AlignmentShift -= Register.AlignmentBytes;
+					Register.ByteIndex -= Register.AlignmentBytes;
+					Register.AlignmentBytes = 0;
+					Pointer = GetData(RegisterIndex);
 				}
 			}
+
+			while (!IsAligned(Pointer, Alignment))
+			{
+				Data.InsertZeroed(Register.GetFirstAllocatedByte(), 1);
+				Register.AlignmentBytes++;
+				Register.ByteIndex++;
+				AlignmentShift++;
+				Pointer = GetData(RegisterIndex);
+			}
 		}
 	}
 
diff --git a/Engine/Source/Runtime/RigVM/Public/RigVMCore/RigVMByteCode.h b/Engine/Source/Runtime/RigVM/Public/RigVMCore/RigVMByteCode.h
--- a/Engine/Source/Runtime/RigVM/Public/RigVMCore/RigVMByteCode.h
+++ b/Engine/Source/Runtime/RigVM/Public/RigVMCore/RigVMByteCode.h
@@ -100,8 +100,11 @@ enum class ERigVMOpCode : uint8
 };
 
 // Base class for all VM operations
+USTRUCT()
 struct RIGVM_API FRigVMBaseOp
 {
+	GENERATED_USTRUCT_BODY()
+
 	FRigVMBaseOp(ERigVMOpCode InOpCode = ERigVMOpCode::Invalid)
 	: OpCode(InOpCode)
 	{
@@ -112,8 +115,11 @@ struct RIGVM_API FRigVMBaseOp
 
 
 // execute a function
+USTRUCT()
 struct RIGVM_API FRigVMExecuteOp : public FRigVMBaseOp
 {
+	GENERATED_USTRUCT_BODY()
+
 	FRigVMExecuteOp()
 	: FRigVMBaseOp()
 	, FunctionIndex(INDEX_NONE)
@@ -139,8 +145,11 @@ struct RIGVM_API FRigVMExecuteOp : public FRigVMBaseOp
 };
 
 // operator used for zero, false, true, increment, decrement
+USTRUCT()
 struct RIGVM_API FRigVMUnaryOp : public FRigVMBaseOp
 {
+	GENERATED_USTRUCT_BODY()
+
 	FRigVMUnaryOp()
 		: FRigVMBaseOp(ERigVMOpCode::Invalid)
 		, Arg()
@@ -175,8 +184,11 @@ struct RIGVM_API FRigVMUnaryOp : public FRigVMBaseOp
 };
 
 // operator used for beginblock
+USTRUCT()
 struct RIGVM_API FRigVMBinaryOp : public FRigVMBaseOp
 {
+	GENERATED_USTRUCT_BODY()
+
 	FRigVMBinaryOp()
 		: FRigVMBaseOp(ERigVMOpCode::Invalid)
 		, ArgA()
@@ -206,8 +218,11 @@ struct RIGVM_API FRigVMBinaryOp : public FRigVMBaseOp
 };
 
 // copy the content of one register to another
+USTRUCT()
 struct RIGVM_API FRigVMCopyOp : public FRigVMBaseOp
 {
+	GENERATED_USTRUCT_BODY()
+
 	FRigVMCopyOp()
 	: FRigVMBaseOp(ERigVMOpCode::Copy)
 	, Source()
@@ -237,8 +252,11 @@ struct RIGVM_API FRigVMCopyOp : public FRigVMBaseOp
 };
 
 // used for equals and not equals comparisons
+USTRUCT()
 struct RIGVM_API FRigVMComparisonOp : public FRigVMBaseOp
 {
+	GENERATED_USTRUCT_BODY()
+
 	FRigVMComparisonOp()
 		: FRigVMBaseOp(ERigVMOpCode::Invalid)
 		, A()
@@ -279,8 +297,11 @@ struct RIGVM_API FRigVMComparisonOp : public FRigVMBaseOp
 // jump to a new instruction index.
 // the instruction can be absolute, relative forward or relative backward
 // based on the opcode 
+USTRUCT()
 struct RIGVM_API FRigVMJumpOp : public FRigVMBaseOp
 {
+	GENERATED_USTRUCT_BODY()
+
 	FRigVMJumpOp()
 	: FRigVMBaseOp(ERigVMOpCode::Invalid)
 	, InstructionIndex(INDEX_NONE)
@@ -308,8 +329,11 @@ struct RIGVM_API FRigVMJumpOp : public FRigVMBaseOp
 // jump to a new instruction index based on a condition.
 // the instruction can be absolute, relative forward or relative backward
 // based on the opcode 
+USTRUCT()
 struct RIGVM_API FRigVMJumpIfOp : public FRigVMUnaryOp
 {
+	GENERATED_USTRUCT_BODY()
+
 	FRigVMJumpIfOp()
 		: FRigVMUnaryOp()
 		, InstructionIndex(INDEX_NONE)
@@ -338,8 +362,11 @@ struct RIGVM_API FRigVMJumpIfOp : public FRigVMUnaryOp
 };
 
 // change the type of a register
+USTRUCT()
 struct RIGVM_API FRigVMChangeTypeOp : public FRigVMUnaryOp
 {
+	GENERATED_USTRUCT_BODY()
+
 	FRigVMChangeTypeOp()
 		: FRigVMUnaryOp()
 		, Type(ERigVMRegisterType::Invalid)
@@ -380,17 +407,21 @@ struct RIGVM_API FRigVMInstruction
 {
 	GENERATED_USTRUCT_BODY()
 
-	FRigVMInstruction(ERigVMOpCode InOpCode = ERigVMOpCode::Invalid, uint64 InByteCodeIndex = UINT64_MAX)
-		: OpCode(InOpCode)
-		, ByteCodeIndex(InByteCodeIndex)
+	FRigVMInstruction(ERigVMOpCode InOpCode = ERigVMOpCode::Invalid, uint64 InByteCodeIndex = UINT64_MAX, uint8 InOperandAlignment = 0)
+		: ByteCodeIndex(InByteCodeIndex)
+		, OpCode(InOpCode)
+		, OperandAlignment(InOperandAlignment)
 	{
 	}
 
+	UPROPERTY()
+	uint64 ByteCodeIndex;
+
 	UPROPERTY()
 	ERigVMOpCode OpCode;
 
 	UPROPERTY()
-	uint64 ByteCodeIndex;
+	uint8 OperandAlignment;
 };
 
 /**
@@ -425,7 +456,7 @@ public:
 private:
 
 	// hide utility constructor
-	FRigVMInstructionArray(const FRigVMByteCode& InByteCode);
+	FRigVMInstructionArray(const FRigVMByteCode& InByteCode, bool bByteCodeIsAligned = true);
 
 	UPROPERTY()
 	TArray<FRigVMInstruction> Instructions;
@@ -481,6 +512,9 @@ public:
 	// returns the number of instructions in this container
 	uint64 Num() const;
 
+	// const accessor for a byte given its index
+	FORCEINLINE const uint8& operator[](int32 InIndex) const { return ByteCode[InIndex]; }
+
 	// returns the number of entries
 	uint64 NumEntries() const;
 
@@ -538,7 +572,7 @@ public:
 	// returns an instruction array for iterating over all operators
 	FORCEINLINE FRigVMInstructionArray GetInstructions() const
 	{
-		return FRigVMInstructionArray(*this);
+		return FRigVMInstructionArray(*this, bByteCodeIsAligned);
 	}
 
 	// returns the opcode at a given byte index
@@ -549,7 +583,7 @@ public:
 	}
 
 	// returns the size of the operator in bytes at a given byte index
-	uint64 GetOpNumBytesAt(uint64 InByteCodeIndex, bool bIncludeOperands = true) const;
+	uint64 GetOpNumBytesAt(uint64 InByteCodeIndex, bool bIncludeOperands) const;
 
 	// returns an operator at a given byte code index
 	template<class OpType>
@@ -588,18 +622,15 @@ public:
 		return FRigVMOperandArray((FRigVMOperand*)(ByteCode.GetData() + InByteCodeIndex), InArgumentCount);
 	}
 
-	// returns the operands for an execute operator / instruction at a given byte code index
-	FORCEINLINE FRigVMOperandArray GetOperandsForExecuteOp(uint64 InByteCodeIndex) const
-	{
-		const FRigVMExecuteOp& ExecuteOp = GetOpAt<FRigVMExecuteOp>(InByteCodeIndex);
-		return GetOperandsAt(InByteCodeIndex + sizeof(FRigVMExecuteOp), ExecuteOp.GetOperandCount());
-	}
-
 	// returns the operands for a given execute instruction
-	FORCEINLINE FRigVMOperandArray GetOperandsForExecuteOp(const FRigVMInstruction& InInstruction) const
+	FORCEINLINE_DEBUGGABLE FRigVMOperandArray GetOperandsForExecuteOp(const FRigVMInstruction& InInstruction) const
 	{
-		const FRigVMExecuteOp& ExecuteOp = GetOpAt<FRigVMExecuteOp>(InInstruction);
-		return GetOperandsAt(InInstruction.ByteCodeIndex + sizeof(FRigVMExecuteOp), ExecuteOp.GetOperandCount());
+		uint64 ByteCodeIndex = InInstruction.ByteCodeIndex;
+		const FRigVMExecuteOp& ExecuteOp = GetOpAt<FRigVMExecuteOp>(ByteCodeIndex);
+		// if the bytecode is not aligned the OperandAlignment needs to be 0
+		ensure(bByteCodeIsAligned || InInstruction.OperandAlignment == 0);
+		ByteCodeIndex += sizeof(FRigVMExecuteOp) + (uint64)InInstruction.OperandAlignment;
+		return GetOperandsAt(ByteCodeIndex, ExecuteOp.GetOperandCount());
 	}
 
 	// returns the raw data of the byte code
@@ -621,29 +652,43 @@ public:
 	// returns the number of instructions within this byte code
 	int32 GetNumInstructions() const { return NumInstructions; }
 
+	// returns the alignment for an operator given its opcode
+	uint64 GetOpAlignment(ERigVMOpCode InOpCode) const;
+
+	// returns the alignment for an operand
+	uint64 GetOperandAlignment() const;
+
 	FString DumpToText() const;
 
 private:
 
 	template<class OpType>
-	FORCEINLINE uint64 AddOp(const OpType& InOp)
+	FORCEINLINE_DEBUGGABLE uint64 AddOp(const OpType& InOp)
 	{
 		uint64 ByteIndex = (uint64)ByteCode.AddZeroed(sizeof(OpType));
-		FMemory::Memcpy(ByteCode.GetData() + ByteIndex, &InOp, sizeof(OpType));
+		uint8* Pointer = &ByteCode[ByteIndex];
+		FMemory::Memcpy(Pointer, &InOp, sizeof(OpType));
 		NumInstructions++;
 		return ByteIndex;
 	}
 
-	// memory for all functions
+	void AlignByteCode();
+
+	// memory for all instructions
 	UPROPERTY()
 	TArray<uint8> ByteCode;
 
-	// memory for all functions
+	// number of instructions stored here
 	UPROPERTY()
 	int32 NumInstructions;
 
+	// a look up table from entry name to instruction index
 	UPROPERTY()
 	TArray<FRigVMByteCodeEntry> Entries;
 
+	// if this is set to true the stored bytecode is aligned / padded
+	bool bByteCodeIsAligned;
+
 	friend class URigVMCompiler;
+	friend class FRigVMByteCodeTest;
 };
diff --git a/Engine/Source/Runtime/RigVM/Public/RigVMCore/RigVMTraits.h b/Engine/Source/Runtime/RigVM/Public/RigVMCore/RigVMTraits.h
--- a/Engine/Source/Runtime/RigVM/Public/RigVMCore/RigVMTraits.h
+++ b/Engine/Source/Runtime/RigVM/Public/RigVMCore/RigVMTraits.h
@@ -181,12 +181,6 @@ FORCEINLINE_DEBUGGABLE void RigVMInitialize(void* InPtr, int32 InCount)
 {
 	ensure(InCount >= 0);
 	FMemory::Memzero(InPtr, InCount * sizeof(FName));
-
-	FName* Values = (FName*)InPtr;
-	for(int32 Index=0;Index<InCount;Index++)
-	{
-		Values[Index] = FName();
-	}
 }
 
 template<
@@ -196,11 +190,7 @@ template<
 FORCEINLINE_DEBUGGABLE void RigVMDestroy(void* InPtr, int32 InCount)
 {
 	ensure(InCount >= 0);
-	FName* Values = (FName*)InPtr;
-	for(int32 Index=0;Index<InCount;Index++)
-	{
-		Values[Index] = FName();
-	}
+	FMemory::Memzero(InPtr, InCount * sizeof(FName));
 }
 
 template<
@@ -210,13 +200,7 @@ template<
 FORCEINLINE_DEBUGGABLE void RigVMCopy(void* InTargetPtr, const void* InSourcePtr, int32 InCount = 1)
 {
 	ensure(InCount >= 0);
-
-	FName* Target = (FName*)InTargetPtr;
-	const FName* Source = (const FName*)InSourcePtr;
-	for(int32 Index=0;Index<InCount;Index++)
-	{
-		Target[Index] = Source[Index];
-	}
+	FMemory::Memcpy(InTargetPtr, InSourcePtr, InCount * sizeof(FName));
 }
 
 template <
