Editgrated 14490694 from Release-4.26 (Control Rig: RigVM alignment to only apply with ByteCode.Num() > 0)

diff --git a/Engine/Source/Runtime/RigVM/Private/RigVMCore/RigVM.cpp b/Engine/Source/Runtime/RigVM/Private/RigVMCore/RigVM.cpp
--- a/Engine/Source/Runtime/RigVM/Private/RigVMCore/RigVM.cpp
+++ b/Engine/Source/Runtime/RigVM/Private/RigVMCore/RigVM.cpp
@@ -203,6 +203,7 @@ void URigVM::CopyFrom(URigVM* InVM, bool bDeferCopy, bool bReferenceLiteralMemor
 	{
 		ByteCodeStorage = InVM->ByteCodeStorage;
 		ByteCodePtr = &ByteCodeStorage;
+		ByteCodePtr->bByteCodeIsAligned = InVM->ByteCodeStorage.bByteCodeIsAligned;
 	}
 	else
 	{
diff --git a/Engine/Source/Runtime/RigVM/Private/RigVMCore/RigVMByteCode.cpp b/Engine/Source/Runtime/RigVM/Private/RigVMCore/RigVMByteCode.cpp
--- a/Engine/Source/Runtime/RigVM/Private/RigVMCore/RigVMByteCode.cpp
+++ b/Engine/Source/Runtime/RigVM/Private/RigVMCore/RigVMByteCode.cpp
@@ -1063,6 +1063,11 @@ void FRigVMByteCode::AlignByteCode()
 		return;
 	}
 
+	if (ByteCode.Num() == 0)
+	{
+		return;
+	}
+
 	FRigVMInstructionArray Instructions(*this, false);
 	uint64 BytesToReserve = ByteCode.Num();
 
diff --git a/Engine/Source/Runtime/RigVM/Public/RigVMCore/RigVMByteCode.h b/Engine/Source/Runtime/RigVM/Public/RigVMCore/RigVMByteCode.h
--- a/Engine/Source/Runtime/RigVM/Public/RigVMCore/RigVMByteCode.h
+++ b/Engine/Source/Runtime/RigVM/Public/RigVMCore/RigVMByteCode.h
@@ -691,4 +691,5 @@ private:
 
 	friend class URigVMCompiler;
 	friend class FRigVMByteCodeTest;
+	friend class URigVM;
 };
