First steps for Variable accessors. Added virtual method for FRigVMStruct to react to a node being dropped. Added a context to allow to introspect existing variables, create new variables and bind variables to a pin.

Please review in order

Files under Runtime
Files under Developer
Files under EngineTest

diff --git a/Engine/Source/Developer/RigVMDeveloper/Private/RigVMModel/RigVMController.cpp b/Engine/Source/Developer/RigVMDeveloper/Private/RigVMModel/RigVMController.cpp
--- a/Engine/Source/Developer/RigVMDeveloper/Private/RigVMModel/RigVMController.cpp
+++ b/Engine/Source/Developer/RigVMDeveloper/Private/RigVMModel/RigVMController.cpp
@@ -225,6 +225,16 @@ URigVMStructNode* URigVMController::AddStructNode(UScriptStruct* InScriptStruct,
 
 	Notify(ERigVMGraphNotifType::NodeAdded, Node);
 
+	if (StructNodeCreatedContext.IsValid())
+	{
+		if (TSharedPtr<FStructOnScope> StructScope = Node->ConstructStructInstance())
+		{
+			TGuardValue<FName> NodeNameScope(StructNodeCreatedContext.NodeName, Node->GetFName());
+			FRigVMStruct* StructInstance = (FRigVMStruct*)StructScope->GetStructMemory();
+			StructInstance->OnStructNodeCreated(StructNodeCreatedContext);
+		}
+	}
+
 	if (bUndo)
 	{
 		ActionStack->EndAction(Action);
@@ -266,6 +276,12 @@ URigVMVariableNode* URigVMController::AddVariableNode(const FName& InVariableNam
 		InCPPTypeObject = URigVMPin::FindObjectFromCPPTypeObjectPath<UObject>(InCPPType);
 	}
 
+	FString CPPType = InCPPType;
+	if (UScriptStruct* ScriptStruct = Cast<UScriptStruct>(InCPPTypeObject))
+	{
+		CPPType = ScriptStruct->GetStructCPPName();
+	}
+
 	FString Name = GetValidNodeName(InNodeName.IsEmpty() ? FString(TEXT("VariableNode")) : InNodeName);
 	URigVMVariableNode* Node = NewObject<URigVMVariableNode>(Graph, *Name);
 	Node->Position = InPosition;
@@ -288,7 +304,7 @@ URigVMVariableNode* URigVMController::AddVariableNode(const FName& InVariableNam
 	Node->Pins.Add(VariablePin);
 
 	URigVMPin* ValuePin = NewObject<URigVMPin>(Node, *URigVMVariableNode::ValueName);
-	ValuePin->CPPType = InCPPType;
+	ValuePin->CPPType = CPPType;
 
 	if (UScriptStruct* ScriptStruct = Cast<UScriptStruct>(InCPPTypeObject))
 	{
diff --git a/Engine/Source/Developer/RigVMDeveloper/Public/RigVMModel/RigVMController.h b/Engine/Source/Developer/RigVMDeveloper/Public/RigVMModel/RigVMController.h
--- a/Engine/Source/Developer/RigVMDeveloper/Public/RigVMModel/RigVMController.h
+++ b/Engine/Source/Developer/RigVMDeveloper/Public/RigVMModel/RigVMController.h
@@ -443,6 +443,8 @@ public:
 	void RepopulatePinsOnNode(URigVMNode* InNode);
 #endif
 
+	FRigVMStructNodeCreatedContext& GetStructNodeCreatedContext() { return StructNodeCreatedContext; }
+
 private:
 
 	UPROPERTY(BlueprintReadOnly, Category = RigVMController, meta = (ScriptName = "ModifiedEvent", AllowPrivateAccess = "true"))
@@ -563,6 +565,8 @@ private:
 	static TMap<FControlRigStructPinRedirectorKey, FString> PinPathCoreRedirectors;
 	FCriticalSection PinPathCoreRedirectorsLock;
 
+	FRigVMStructNodeCreatedContext StructNodeCreatedContext;
+
 	friend class URigVMGraph;
 	friend class URigVMActionStack;
 	friend class URigVMCompiler;
diff --git a/Engine/Source/Runtime/RigVM/Private/RigVMCore/RigVMStruct.cpp b/Engine/Source/Runtime/RigVM/Private/RigVMCore/RigVMStruct.cpp
--- a/Engine/Source/Runtime/RigVM/Private/RigVMCore/RigVMStruct.cpp
+++ b/Engine/Source/Runtime/RigVM/Private/RigVMCore/RigVMStruct.cpp
@@ -3,6 +3,87 @@
 #include "RigVMCore/RigVMStruct.h"
 #include "UObject/StructOnScope.h"
 
+////////////////////////////////////////////////////////////////////////////////////////////////////////////////////////
+
+bool FRigVMStructNodeCreatedContext::IsValid() const
+{
+	return Reason != ERigVMNodeCreatedReason::Unknown &&
+		AllExternalVariablesDelegate.IsBound() &&
+		CreateExternalVariableDelegate.IsBound() &&
+		BindPinToExternalVariableDelegate.IsBound();
+}
+
+TArray<FRigVMExternalVariable> FRigVMStructNodeCreatedContext::GetExternalVariables() const
+{
+	TArray<FRigVMExternalVariable> ExternalVariables;
+
+	if (AllExternalVariablesDelegate.IsBound())
+	{
+		ExternalVariables = AllExternalVariablesDelegate.Execute();
+	}
+
+	return ExternalVariables;
+}
+
+FName FRigVMStructNodeCreatedContext::AddExternalVariable(const FRigVMExternalVariable& InVariableToCreate)
+{
+	if (CreateExternalVariableDelegate.IsBound())
+	{
+		return CreateExternalVariableDelegate.Execute(InVariableToCreate);
+	}
+	return NAME_None;
+}
+
+bool FRigVMStructNodeCreatedContext::BindPinToExternalVariable(FString InPinPath, FName InVariableName)
+{
+	if (BindPinToExternalVariableDelegate.IsBound())
+	{
+		return BindPinToExternalVariableDelegate.Execute(InPinPath, InVariableName);
+	}
+	return false;
+}
+
+FRigVMExternalVariable FRigVMStructNodeCreatedContext::FindVariable(FName InVariableName) const
+{
+	TArray<FRigVMExternalVariable> ExternalVariables = GetExternalVariables();
+	for (FRigVMExternalVariable ExternalVariable : ExternalVariables)
+	{
+		if (ExternalVariable.Name == InVariableName)
+		{
+			return ExternalVariable;
+		}
+	}
+	return FRigVMExternalVariable();
+}
+
+FName FRigVMStructNodeCreatedContext::FindFirstVariableOfType(FName InCPPTypeName) const
+{
+	TArray<FRigVMExternalVariable> ExternalVariables = GetExternalVariables();
+	for (FRigVMExternalVariable ExternalVariable : ExternalVariables)
+	{
+		if (ExternalVariable.TypeName == InCPPTypeName)
+		{
+			return ExternalVariable.Name;
+		}
+	}
+	return NAME_None;
+}
+
+FName FRigVMStructNodeCreatedContext::FindFirstVariableOfType(UObject* InCPPTypeObject) const
+{
+	TArray<FRigVMExternalVariable> ExternalVariables = GetExternalVariables();
+	for (FRigVMExternalVariable ExternalVariable : ExternalVariables)
+	{
+		if (ExternalVariable.TypeObject == InCPPTypeObject)
+		{
+			return ExternalVariable.Name;
+		}
+	}
+	return NAME_None;
+}
+
+////////////////////////////////////////////////////////////////////////////////////////////////////////////////////////
+
 const FName FRigVMStruct::DeprecatedMetaName("Deprecated");
 const FName FRigVMStruct::InputMetaName("Input");
 const FName FRigVMStruct::OutputMetaName("Output");
@@ -365,5 +446,4 @@ FString FRigVMStruct::ExportToFullyQualifiedText(UScriptStruct* InStruct, const
 	return FString::Printf(TEXT("(%s)"), *FString::Join(FieldValues, TEXT(",")));
 }
 
-
 #endif
\ No newline at end of file
diff --git a/Engine/Source/Runtime/RigVM/Public/RigVMCore/RigVMExternalVariable.h b/Engine/Source/Runtime/RigVM/Public/RigVMCore/RigVMExternalVariable.h
--- a/Engine/Source/Runtime/RigVM/Public/RigVMCore/RigVMExternalVariable.h
+++ b/Engine/Source/Runtime/RigVM/Public/RigVMCore/RigVMExternalVariable.h
@@ -90,6 +90,11 @@ struct RIGVM_API FRigVMExternalVariable
 			ExternalVariable.TypeName = *StructProperty->Struct->GetStructCPPName();
 			ExternalVariable.TypeObject = StructProperty->Struct;
 		}
+		else if (FObjectProperty* ObjectProperty = CastField<FObjectProperty>(Property))
+		{
+			ExternalVariable.TypeName = *ObjectProperty->PropertyClass->GetName();
+			ExternalVariable.TypeObject = ObjectProperty->PropertyClass;
+		}
 
 		return ExternalVariable;
 	}
@@ -334,6 +339,38 @@ struct RIGVM_API FRigVMExternalVariable
 		return Variable;
 	}
 
+	template <
+		typename T,
+		typename TEnableIf<TModels<CRigVMUClass, T>::Value>::Type * = nullptr
+	>
+	FORCEINLINE static FRigVMExternalVariable Make(const FName& InName, T& InValue)
+	{
+		FRigVMExternalVariable Variable;
+		Variable.Name = InName;
+		Variable.TypeName = T::StaticClass()->GetFName();
+		Variable.TypeObject = T::StaticClass();
+		Variable.bIsArray = false;
+		Variable.Size = T::StaticClass()->GetStructureSize();
+		Variable.Memory = (uint8*)&InValue;
+		return Variable;
+	}
+
+	template <
+		typename T,
+		typename TEnableIf<TModels<CRigVMUClass, T>::Value>::Type * = nullptr
+	>
+	FORCEINLINE static FRigVMExternalVariable Make(const FName& InName, TArray<T>& InValue)
+	{
+		FRigVMExternalVariable Variable;
+		Variable.Name = InName;
+		Variable.TypeName = T::StaticClass()->GetFName();
+		Variable.TypeObject = T::StaticClass();
+		Variable.bIsArray = true;
+		Variable.Size = T::StaticClass()->GetStructureSize();
+		Variable.Memory = (uint8*)&InValue;
+		return Variable;
+	}
+
 	template<typename T>
 	FORCEINLINE T GetValue()
 	{
diff --git a/Engine/Source/Runtime/RigVM/Public/RigVMCore/RigVMStruct.h b/Engine/Source/Runtime/RigVM/Public/RigVMCore/RigVMStruct.h
--- a/Engine/Source/Runtime/RigVM/Public/RigVMCore/RigVMStruct.h
+++ b/Engine/Source/Runtime/RigVM/Public/RigVMCore/RigVMStruct.h
@@ -4,8 +4,134 @@
 
 #include "CoreMinimal.h"
 #include "RigVMCore/RigVMRegistry.h"
+#include "RigVMCore/RigVMExternalVariable.h"
+#include "RigVMCore/RigVMTraits.h"
 #include "RigVMStruct.generated.h"
 
+// delegates used for variable introspection / creation
+DECLARE_DELEGATE_RetVal(TArray<FRigVMExternalVariable>, FRigVMGetExternalVariablesDelegate)
+DECLARE_DELEGATE_RetVal_OneParam(FName, FRigVMCreateExternalVariableDelegate, FRigVMExternalVariable)
+DECLARE_DELEGATE_RetVal_TwoParams(bool, FRigVMBindPinToExternalVariableDelegate, FString, FName)
+
+/** Context as of why the node was created */
+enum class ERigVMNodeCreatedReason : uint8
+{
+	NodeSpawner,
+	ScriptedEvent,
+	Paste,
+	BackwardsCompatibility,
+	Unknown,
+};
+
+/**
+ * A context struct passed to FRigVMStruct::OnStructNodeCreated
+ */
+struct RIGVM_API FRigVMStructNodeCreatedContext
+{
+public:
+
+	struct FScope
+	{
+	public:
+		FScope(FRigVMStructNodeCreatedContext& InContext, ERigVMNodeCreatedReason InReason)
+			: Context(InContext)
+			, PreviousReason(InContext.GetReason())
+		{
+			Context.Reason = InReason;
+		}
+
+		~FScope()
+		{
+			Context.Reason = PreviousReason;
+		}
+
+	private:
+		FRigVMStructNodeCreatedContext& Context;
+		ERigVMNodeCreatedReason PreviousReason;
+	};
+
+	/** Returns true if this context is valid to use */
+	bool IsValid() const;
+
+	/** Get the reason why this node was created */
+	ERigVMNodeCreatedReason GetReason() const { return Reason; }
+
+	/** Get the name of this node */
+	FName GetNodeName() const { return NodeName; }
+
+	/** Returns all currently existing external variables */
+	TArray<FRigVMExternalVariable> GetExternalVariables() const;
+
+	/** Creates a new variable within the host of this VM */
+	FName AddExternalVariable(const FRigVMExternalVariable& InVariableToCreate);
+
+	/** Binds a pin to an external variable on the created node */
+	bool BindPinToExternalVariable(FString InPinPath, FName InVariableName);
+
+	/** Returns a variable given a name (or a non-valid variable if not found) */
+	FRigVMExternalVariable FindVariable(FName InVariableName) const;
+
+	/** Returns the name of the first variable given a(or NAME_None if not found) */
+	FName FindFirstVariableOfType(FName InCPPTypeName) const;
+
+	/** Returns the name of the first variable given a type (or NAME_None if not found) */
+	template <
+		typename T,
+		typename TEnableIf<TRigVMIsBaseStructure<T>::Value>::Type * = nullptr
+	>
+	FORCEINLINE FName FindFirstVariableOfType() const
+	{
+		return FindFirstVariableOfType(TBaseStructure<T>::Get());
+	}
+
+	/** Returns the name of the first variable given a type (or NAME_None if not found) */
+	template <
+		typename T,
+		typename TEnableIf<TModels<CRigVMUStruct, T>::Value>::Type * = nullptr
+	>
+	FORCEINLINE FName FindFirstVariableOfType() const
+	{
+		return FindFirstVariableOfType(T::StaticStruct());
+	}
+
+	/** Returns the name of the first variable given a type (or NAME_None if not found) */
+	template <
+		typename T,
+		typename TEnableIf<TModels<CRigVMUClass, T>::Value>::Type * = nullptr
+	>
+	FORCEINLINE FName FindFirstVariableOfType() const
+	{
+		return FindFirstVariableOfType(T::StaticClass());
+	}
+
+	/** Returns the name of the first variable given a type (or NAME_None if not found) */
+	template <
+		typename T,
+		typename TEnableIf<TIsEnum<T>::Value>::Type * = nullptr
+	>
+		FORCEINLINE FName FindFirstVariableOfType() const
+	{
+		return FindFirstVariableOfType(StaticEnum<T>());
+	}
+
+	FRigVMGetExternalVariablesDelegate& GetAllExternalVariablesDelegate() { return AllExternalVariablesDelegate; }
+	FRigVMCreateExternalVariableDelegate& GetCreateExternalVariableDelegate() { return CreateExternalVariableDelegate; }
+	FRigVMBindPinToExternalVariableDelegate& GetBindPinToExternalVariableDelegate() { return BindPinToExternalVariableDelegate; }
+
+private:
+	
+	FName NodeName = NAME_None;
+	ERigVMNodeCreatedReason Reason = ERigVMNodeCreatedReason::Unknown;
+	FRigVMGetExternalVariablesDelegate AllExternalVariablesDelegate;
+	FRigVMCreateExternalVariableDelegate CreateExternalVariableDelegate;
+	FRigVMBindPinToExternalVariableDelegate BindPinToExternalVariableDelegate;
+
+	FName FindFirstVariableOfType(UObject* InCPPTypeObject) const;
+
+	friend class URigVMController;
+	friend struct FScope;
+};
+
 /**
  * The base class for all RigVM enabled structs.
  */
@@ -22,11 +148,15 @@ public:
 
 	FORCEINLINE virtual int32 GetArraySize(const FName& InParameterName, const FRigVMUserDataArray& RigVMUserData) { return INDEX_NONE; }
 
-	// loop realted
+	// loop related
 	FORCEINLINE virtual bool IsForLoop() const { return false; }
 	FORCEINLINE virtual int32 GetNumSlices() const { return 1; }
 
+	// node creation
+	FORCEINLINE virtual void OnStructNodeCreated(FRigVMStructNodeCreatedContext& InContext) const {}
+
 #if WITH_EDITOR
+
 	static bool ValidateStruct(UScriptStruct* InStruct, FString* OutErrorMessage);
 	static bool CheckPinType(UScriptStruct* InStruct, const FName& PinName, const FString& ExpectedType, FString* OutErrorMessage = nullptr);
 	static bool CheckPinDirection(UScriptStruct* InStruct, const FName& PinName, const FName& InDirectionMetaName);
