Split Serialize into Save and Load functions.

diff --git a/Engine/Source/Runtime/RigVM/Private/RigVMCore/RigVM.cpp b/Engine/Source/Runtime/RigVM/Private/RigVMCore/RigVM.cpp
--- a/Engine/Source/Runtime/RigVM/Private/RigVMCore/RigVM.cpp
+++ b/Engine/Source/Runtime/RigVM/Private/RigVMCore/RigVM.cpp
@@ -5,27 +5,47 @@
 #include "UObject/AnimObjectVersion.h"
 #include "HAL/PlatformTLS.h"
 
-bool FRigVMParameter::Serialize(FArchive& Ar)
+void FRigVMParameter::Serialize(FArchive& Ar)
 {
 	Ar.UsingCustomVersion(FAnimObjectVersion::GUID);
 
 	if (Ar.CustomVer(FAnimObjectVersion::GUID) < FAnimObjectVersion::StoreMarkerNamesOnSkeleton)
 	{
-		return false;
+		return;
 	}
 
+	if (Ar.IsSaving() || Ar.IsObjectReferenceCollector())
+	{
+		Save(Ar);
+	}
+	else if (Ar.IsLoading())
+	{
+		Load(Ar);
+	}
+	else
+	{
+		checkNoEntry();
+	}
+}
+
+void FRigVMParameter::Save(FArchive& Ar)
+{
 	Ar << Type;
 	Ar << Name;
 	Ar << RegisterIndex;
 	Ar << CPPType;
 	Ar << ScriptStructPath;
+}
 
-	if (Ar.IsLoading())
-	{
-		ScriptStruct = nullptr;
-	}
+void FRigVMParameter::Load(FArchive& Ar)
+{
+	Ar << Type;
+	Ar << Name;
+	Ar << RegisterIndex;
+	Ar << CPPType;
+	Ar << ScriptStructPath;
 
-	return true;
+	ScriptStruct = nullptr;
 }
 
 UScriptStruct* FRigVMParameter::GetScriptStruct() const
@@ -70,10 +90,32 @@ void URigVM::Serialize(FArchive& Ar)
 
 	ensure(ExecutingThreadId == INDEX_NONE);
 
-	if (Ar.IsLoading())
+	if (Ar.IsSaving() || Ar.IsObjectReferenceCollector())
 	{
-		Reset();
+		Save(Ar);
+	}
+	else if (Ar.IsLoading())
+	{
+		Load(Ar);
+	}
+	else
+	{
+		checkNoEntry();
 	}
+}
+
+void URigVM::Save(FArchive& Ar)
+{
+	Ar << WorkMemoryStorage;
+	Ar << LiteralMemoryStorage;
+	Ar << FunctionNamesStorage;
+	Ar << ByteCodeStorage;
+	Ar << Parameters;
+}
+
+void URigVM::Load(FArchive& Ar)
+{
+	Reset();
 
 	Ar << WorkMemoryStorage;
 	Ar << LiteralMemoryStorage;
@@ -81,26 +123,23 @@ void URigVM::Serialize(FArchive& Ar)
 	Ar << ByteCodeStorage;
 	Ar << Parameters;
 
-	if (Ar.IsLoading())
+	if (WorkMemoryStorage.bEncounteredErrorDuringLoad ||
+		LiteralMemoryStorage.bEncounteredErrorDuringLoad)
 	{
-		if (WorkMemoryStorage.bEncounteredErrorDuringLoad ||
-			LiteralMemoryStorage.bEncounteredErrorDuringLoad)
+		Reset();
+	}
+	else
+	{
+		Instructions.Reset();
+		FunctionsStorage.Reset();
+		ParametersNameMap.Reset();
+
+		for (int32 Index = 0; Index < Parameters.Num(); Index++)
 		{
-			Reset();
+			ParametersNameMap.Add(Parameters[Index].Name, Index);
 		}
-		else
-		{
-			Instructions.Reset();
-			FunctionsStorage.Reset();
-			ParametersNameMap.Reset();
-
-			for (int32 Index = 0; Index < Parameters.Num(); Index++)
-			{
-				ParametersNameMap.Add(Parameters[Index].Name, Index);
-			}
 
-			InvalidateCachedMemory();
-		}
+		InvalidateCachedMemory();
 	}
 }
 
diff --git a/Engine/Source/Runtime/RigVM/Private/RigVMCore/RigVMByteCode.cpp b/Engine/Source/Runtime/RigVM/Private/RigVMCore/RigVMByteCode.cpp
--- a/Engine/Source/Runtime/RigVM/Private/RigVMCore/RigVMByteCode.cpp
+++ b/Engine/Source/Runtime/RigVM/Private/RigVMCore/RigVMByteCode.cpp
@@ -4,65 +4,57 @@
 #include "UObject/AnimObjectVersion.h"
 #include "UObject/FortniteMainBranchObjectVersion.h"
 
-bool FRigVMExecuteOp::Serialize(FArchive& Ar)
+void FRigVMExecuteOp::Serialize(FArchive& Ar)
 {
 	Ar << OpCode;
 	Ar << FunctionIndex;
-	return true;
 }
 
-bool FRigVMUnaryOp::Serialize(FArchive& Ar)
+void FRigVMUnaryOp::Serialize(FArchive& Ar)
 {
 	Ar << OpCode;
 	Ar << Arg;
-	return true;
 }
 
-bool FRigVMBinaryOp::Serialize(FArchive& Ar)
+void FRigVMBinaryOp::Serialize(FArchive& Ar)
 {
 	Ar << OpCode;
 	Ar << ArgA;
 	Ar << ArgB;
-	return true;
 }
 
-bool FRigVMCopyOp::Serialize(FArchive& Ar)
+void FRigVMCopyOp::Serialize(FArchive& Ar)
 {
 	Ar << OpCode;
 	Ar << Source;
 	Ar << Target;
-	return true;
 }
 
-bool FRigVMComparisonOp::Serialize(FArchive& Ar)
+void FRigVMComparisonOp::Serialize(FArchive& Ar)
 {
 	Ar << OpCode;
 	Ar << A;
 	Ar << B;
 	Ar << Result;
-	return true;
 }
 
-bool FRigVMJumpOp::Serialize(FArchive& Ar)
+void FRigVMJumpOp::Serialize(FArchive& Ar)
 {
 	Ar << OpCode;
 	Ar << InstructionIndex;
-	return true;
 }
 
-bool FRigVMJumpIfOp::Serialize(FArchive& Ar)
+void FRigVMJumpIfOp::Serialize(FArchive& Ar)
 {
 	Ar << OpCode;
 	Ar << Arg;
 	Ar << InstructionIndex;
 	Ar << Condition;
-	return true;
 }
 
-bool FRigVMChangeTypeOp::Serialize(FArchive& Ar)
+void FRigVMChangeTypeOp::Serialize(FArchive& Ar)
 {
 	ensure(false);
-	return false;
 }
 
 FRigVMInstructionArray::FRigVMInstructionArray()
@@ -131,36 +123,39 @@ FRigVMByteCode::FRigVMByteCode()
 {
 }
 
-bool FRigVMByteCode::Serialize(FArchive& Ar)
+void FRigVMByteCode::Serialize(FArchive& Ar)
 {
 	Ar.UsingCustomVersion(FAnimObjectVersion::GUID);
 	Ar.UsingCustomVersion(FFortniteMainBranchObjectVersion::GUID);
 
 	if (Ar.CustomVer(FAnimObjectVersion::GUID) < FAnimObjectVersion::StoreMarkerNamesOnSkeleton)
 	{
-		return false;
+		return;
 	}
 
-	if (Ar.IsLoading() && Ar.CustomVer(FFortniteMainBranchObjectVersion::GUID) < FFortniteMainBranchObjectVersion::RigVMByteCodeDeterminism)
+	if (Ar.IsSaving() || Ar.IsObjectReferenceCollector())
 	{
-		Ar << ByteCode;
-		return true;
+		Save(Ar);
 	}
-
-	FRigVMInstructionArray Instructions;
-	
-	int32 InstructionCount = 0;
-	if (Ar.IsSaving())
+	else if (Ar.IsLoading())
 	{
-		Instructions = GetInstructions();
-		InstructionCount = Instructions.Num();
+		Load(Ar);
 	}
 	else
 	{
-		ByteCode.Reset();
-		bByteCodeIsAligned = false;
+		checkNoEntry();
 	}
+}
 
+void FRigVMByteCode::Save(FArchive& Ar) const
+{
+	FRigVMInstructionArray Instructions;
+
+	int32 InstructionCount = 0;
+	
+	Instructions = GetInstructions();
+	InstructionCount = Instructions.Num();
+	
 	Ar << InstructionCount;
 
 	for (int32 InstructionIndex = 0; InstructionIndex < InstructionCount; InstructionIndex++)
@@ -168,12 +163,8 @@ bool FRigVMByteCode::Serialize(FArchive& Ar)
 		FRigVMInstruction Instruction;
 		ERigVMOpCode OpCode = ERigVMOpCode::Invalid;
 
-		if (Ar.IsSaving())
-		{
-			Instruction = Instructions[InstructionIndex];
-			OpCode = Instruction.OpCode;
-		}
-			
+		Instruction = Instructions[InstructionIndex];
+		OpCode = Instruction.OpCode;
 		Ar << OpCode;
 
 		switch (OpCode)
@@ -244,52 +235,24 @@ bool FRigVMByteCode::Serialize(FArchive& Ar)
 			case ERigVMOpCode::Execute_63_Operands:
 			case ERigVMOpCode::Execute_64_Operands:
 			{
-				if (Ar.IsSaving())
-				{
-					FRigVMExecuteOp Op = GetOpAt<FRigVMExecuteOp>(Instruction.ByteCodeIndex);
-					Ar << Op;
+				FRigVMExecuteOp Op = GetOpAt<FRigVMExecuteOp>(Instruction.ByteCodeIndex);
+				Ar << Op;
 
-					FRigVMOperandArray Operands = GetOperandsForExecuteOp(Instruction);
-					int32 OperandCount = (int32)Op.GetOperandCount();
-					ensure(OperandCount == Operands.Num());
+				FRigVMOperandArray Operands = GetOperandsForExecuteOp(Instruction);
+				int32 OperandCount = (int32)Op.GetOperandCount();
+				ensure(OperandCount == Operands.Num());
 
-					for (int32 OperandIndex = 0; OperandIndex < OperandCount; OperandIndex++)
-					{
-						FRigVMOperand Operand = Operands[OperandIndex];
-						Ar << Operand;
-					}
-				}
-				else
+				for (int32 OperandIndex = 0; OperandIndex < OperandCount; OperandIndex++)
 				{
-					FRigVMExecuteOp Op;
-					Ar << Op;
-
-					int32 OperandCount = Op.GetOperandCount();
-					TArray<FRigVMOperand> Operands;
-					for (int32 OperandIndex = 0; OperandIndex < OperandCount; OperandIndex++)
-					{
-						FRigVMOperand Operand;
-						Ar << Operand;
-						Operands.Add(Operand);
-					}
-
-					AddExecuteOp(Op.FunctionIndex, Operands);
+					FRigVMOperand Operand = Operands[OperandIndex];
+					Ar << Operand;
 				}
 				break;
 			}
 			case ERigVMOpCode::Copy:
 			{
-				if (Ar.IsSaving())
-				{
-					FRigVMCopyOp Op = GetOpAt<FRigVMCopyOp>(Instruction.ByteCodeIndex);
-					Ar << Op;
-				}
-				else
-				{
-					FRigVMCopyOp Op;
-					Ar << Op;
-					AddOp<FRigVMCopyOp>(Op);
-				}
+				FRigVMCopyOp Op = GetOpAt<FRigVMCopyOp>(Instruction.ByteCodeIndex);
+				Ar << Op;
 				break;
 			}
 			case ERigVMOpCode::Zero:
@@ -298,108 +261,49 @@ bool FRigVMByteCode::Serialize(FArchive& Ar)
 			case ERigVMOpCode::Increment:
 			case ERigVMOpCode::Decrement:
 			{
-				if (Ar.IsSaving())
-				{
-					FRigVMUnaryOp Op = GetOpAt<FRigVMUnaryOp>(Instruction.ByteCodeIndex);
-					Ar << Op;
-				}
-				else
-				{
-					FRigVMUnaryOp Op;
-					Ar << Op;
-					AddOp<FRigVMUnaryOp>(Op);
-				}
+				FRigVMUnaryOp Op = GetOpAt<FRigVMUnaryOp>(Instruction.ByteCodeIndex);
+				Ar << Op;
 				break;
 			}
 			case ERigVMOpCode::Equals:
 			case ERigVMOpCode::NotEquals:
 			{
-				if (Ar.IsSaving())
-				{
-					FRigVMComparisonOp Op = GetOpAt<FRigVMComparisonOp>(Instruction.ByteCodeIndex);
-					Ar << Op;
-				}
-				else
-				{
-					FRigVMComparisonOp Op;
-					Ar << Op;
-					AddOp<FRigVMComparisonOp>(Op);
-				}
+				FRigVMComparisonOp Op = GetOpAt<FRigVMComparisonOp>(Instruction.ByteCodeIndex);
+				Ar << Op;
 				break;
 			}
 			case ERigVMOpCode::JumpAbsolute:
 			case ERigVMOpCode::JumpForward:
 			case ERigVMOpCode::JumpBackward:
 			{
-				if (Ar.IsSaving())
-				{
-					FRigVMJumpOp Op = GetOpAt<FRigVMJumpOp>(Instruction.ByteCodeIndex);
-					Ar << Op;
-				}
-				else
-				{
-					FRigVMJumpOp Op;
-					Ar << Op;
-					AddOp<FRigVMJumpOp>(Op);
-				}
+				FRigVMJumpOp Op = GetOpAt<FRigVMJumpOp>(Instruction.ByteCodeIndex);
+				Ar << Op;
 				break;
 			}
 			case ERigVMOpCode::JumpAbsoluteIf:
 			case ERigVMOpCode::JumpForwardIf:
 			case ERigVMOpCode::JumpBackwardIf:
 			{
-				if (Ar.IsSaving())
-				{
-					FRigVMJumpIfOp Op = GetOpAt<FRigVMJumpIfOp>(Instruction.ByteCodeIndex);
-					Ar << Op;
-				}
-				else
-				{
-					FRigVMJumpIfOp Op;
-					Ar << Op;
-					AddOp<FRigVMJumpIfOp>(Op);
-				}
+				FRigVMJumpIfOp Op = GetOpAt<FRigVMJumpIfOp>(Instruction.ByteCodeIndex);
+				Ar << Op;
 				break;
 			}
 			case ERigVMOpCode::Exit:
 			{
-				if (Ar.IsSaving())
-				{
-					// nothing todo, the ExitOp has no custom data inside of it
-					// so all we need is the previously saved OpCode.
-				}
-				else
-				{
-					AddExitOp();
-				}
+				// nothing todo, the ExitOp has no custom data inside of it
+				// so all we need is the previously saved OpCode.
 				break;
 			}
 			case ERigVMOpCode::BeginBlock:
 			{
-				if (Ar.IsSaving())
-				{
-					FRigVMBinaryOp Op = GetOpAt<FRigVMBinaryOp>(Instruction.ByteCodeIndex);
-					Ar << Op;
-				}
-				else
-				{
-					FRigVMBinaryOp Op;
-					Ar << Op;
-					AddOp<FRigVMBinaryOp>(Op);
-				}
+				FRigVMBinaryOp Op = GetOpAt<FRigVMBinaryOp>(Instruction.ByteCodeIndex);
+				Ar << Op;
 				break;
 			}
 			case ERigVMOpCode::EndBlock:
 			{
-				if (Ar.IsSaving())
-				{
-					// nothing todo, the EndBlock has no custom data inside of it
-					// so all we need is the previously saved OpCode.
-				}
-				else
-				{
-					AddEndBlockOp();
-				}
+				// nothing todo, the EndBlock has no custom data inside of it
+				// so all we need is the previously saved OpCode.
 				break;
 			}
 			default:
@@ -409,45 +313,213 @@ bool FRigVMByteCode::Serialize(FArchive& Ar)
 		}
 	}
 
-	if (Ar.IsLoading())
+	UScriptStruct* ScriptStruct = FRigVMByteCodeEntry::StaticStruct();
+	TArray<uint8, TAlignedHeapAllocator<16>> DefaultStructData;
+	DefaultStructData.AddZeroed(ScriptStruct->GetStructureSize());
+	ScriptStruct->InitializeDefaultValue(DefaultStructData.GetData());
+
+	TArray<FString> View;
+	for (uint16 EntryIndex = 0; EntryIndex < Entries.Num(); EntryIndex++)
+	{
+		FString Value;
+		ScriptStruct->ExportText(Value, &Entries[EntryIndex], DefaultStructData.GetData(), nullptr, PPF_None, nullptr);
+		View.Add(Value);
+	}
+
+	Ar << View;
+}
+
+void FRigVMByteCode::Load(FArchive& Ar)
+{
+	if (Ar.CustomVer(FFortniteMainBranchObjectVersion::GUID) < FFortniteMainBranchObjectVersion::RigVMByteCodeDeterminism)
+	{
+		Ar << ByteCode;
+		return;
+	}
+
+	FRigVMInstructionArray Instructions;
+
+	int32 InstructionCount = 0;
+	
+	ByteCode.Reset();
+	bByteCodeIsAligned = false;
+	Ar << InstructionCount;
+
+	for (int32 InstructionIndex = 0; InstructionIndex < InstructionCount; InstructionIndex++)
 	{
-		AlignByteCode();
+		FRigVMInstruction Instruction;
+		ERigVMOpCode OpCode = ERigVMOpCode::Invalid;
 
-		Entries.Reset();
-		if (Ar.CustomVer(FAnimObjectVersion::GUID) >= FAnimObjectVersion::SerializeRigVMEntries)
+		Ar << OpCode;
+
+		switch (OpCode)
 		{
-			UScriptStruct* ScriptStruct = FRigVMByteCodeEntry::StaticStruct();
+			case ERigVMOpCode::Execute_0_Operands:
+			case ERigVMOpCode::Execute_1_Operands:
+			case ERigVMOpCode::Execute_2_Operands:
+			case ERigVMOpCode::Execute_3_Operands:
+			case ERigVMOpCode::Execute_4_Operands:
+			case ERigVMOpCode::Execute_5_Operands:
+			case ERigVMOpCode::Execute_6_Operands:
+			case ERigVMOpCode::Execute_7_Operands:
+			case ERigVMOpCode::Execute_8_Operands:
+			case ERigVMOpCode::Execute_9_Operands:
+			case ERigVMOpCode::Execute_10_Operands:
+			case ERigVMOpCode::Execute_11_Operands:
+			case ERigVMOpCode::Execute_12_Operands:
+			case ERigVMOpCode::Execute_13_Operands:
+			case ERigVMOpCode::Execute_14_Operands:
+			case ERigVMOpCode::Execute_15_Operands:
+			case ERigVMOpCode::Execute_16_Operands:
+			case ERigVMOpCode::Execute_17_Operands:
+			case ERigVMOpCode::Execute_18_Operands:
+			case ERigVMOpCode::Execute_19_Operands:
+			case ERigVMOpCode::Execute_20_Operands:
+			case ERigVMOpCode::Execute_21_Operands:
+			case ERigVMOpCode::Execute_22_Operands:
+			case ERigVMOpCode::Execute_23_Operands:
+			case ERigVMOpCode::Execute_24_Operands:
+			case ERigVMOpCode::Execute_25_Operands:
+			case ERigVMOpCode::Execute_26_Operands:
+			case ERigVMOpCode::Execute_27_Operands:
+			case ERigVMOpCode::Execute_28_Operands:
+			case ERigVMOpCode::Execute_29_Operands:
+			case ERigVMOpCode::Execute_30_Operands:
+			case ERigVMOpCode::Execute_31_Operands:
+			case ERigVMOpCode::Execute_32_Operands:
+			case ERigVMOpCode::Execute_33_Operands:
+			case ERigVMOpCode::Execute_34_Operands:
+			case ERigVMOpCode::Execute_35_Operands:
+			case ERigVMOpCode::Execute_36_Operands:
+			case ERigVMOpCode::Execute_37_Operands:
+			case ERigVMOpCode::Execute_38_Operands:
+			case ERigVMOpCode::Execute_39_Operands:
+			case ERigVMOpCode::Execute_40_Operands:
+			case ERigVMOpCode::Execute_41_Operands:
+			case ERigVMOpCode::Execute_42_Operands:
+			case ERigVMOpCode::Execute_43_Operands:
+			case ERigVMOpCode::Execute_44_Operands:
+			case ERigVMOpCode::Execute_45_Operands:
+			case ERigVMOpCode::Execute_46_Operands:
+			case ERigVMOpCode::Execute_47_Operands:
+			case ERigVMOpCode::Execute_48_Operands:
+			case ERigVMOpCode::Execute_49_Operands:
+			case ERigVMOpCode::Execute_50_Operands:
+			case ERigVMOpCode::Execute_51_Operands:
+			case ERigVMOpCode::Execute_52_Operands:
+			case ERigVMOpCode::Execute_53_Operands:
+			case ERigVMOpCode::Execute_54_Operands:
+			case ERigVMOpCode::Execute_55_Operands:
+			case ERigVMOpCode::Execute_56_Operands:
+			case ERigVMOpCode::Execute_57_Operands:
+			case ERigVMOpCode::Execute_58_Operands:
+			case ERigVMOpCode::Execute_59_Operands:
+			case ERigVMOpCode::Execute_60_Operands:
+			case ERigVMOpCode::Execute_61_Operands:
+			case ERigVMOpCode::Execute_62_Operands:
+			case ERigVMOpCode::Execute_63_Operands:
+			case ERigVMOpCode::Execute_64_Operands:
+			{
+				FRigVMExecuteOp Op;
+				Ar << Op;
 
-			TArray<FString> View;
-			Ar << View;
+				int32 OperandCount = Op.GetOperandCount();
+				TArray<FRigVMOperand> Operands;
+				for (int32 OperandIndex = 0; OperandIndex < OperandCount; OperandIndex++)
+				{
+					FRigVMOperand Operand;
+					Ar << Operand;
+					Operands.Add(Operand);
+				}
 
-			for (int32 EntryIndex = 0; EntryIndex < View.Num(); EntryIndex++)
+				AddExecuteOp(Op.FunctionIndex, Operands);
+				break;
+			}
+			case ERigVMOpCode::Copy:
+			{
+				FRigVMCopyOp Op;
+				Ar << Op;
+				AddOp<FRigVMCopyOp>(Op);
+				break;
+			}
+			case ERigVMOpCode::Zero:
+			case ERigVMOpCode::BoolFalse:
+			case ERigVMOpCode::BoolTrue:
+			case ERigVMOpCode::Increment:
+			case ERigVMOpCode::Decrement:
+			{
+				FRigVMUnaryOp Op;
+				Ar << Op;
+				AddOp<FRigVMUnaryOp>(Op);
+				break;
+			}
+			case ERigVMOpCode::Equals:
+			case ERigVMOpCode::NotEquals:
+			{
+				FRigVMComparisonOp Op;
+				Ar << Op;
+				AddOp<FRigVMComparisonOp>(Op);
+				break;
+			}
+			case ERigVMOpCode::JumpAbsolute:
+			case ERigVMOpCode::JumpForward:
+			case ERigVMOpCode::JumpBackward:
+			{
+				FRigVMJumpOp Op;
+				Ar << Op;
+				AddOp<FRigVMJumpOp>(Op);
+				break;
+			}
+			case ERigVMOpCode::JumpAbsoluteIf:
+			case ERigVMOpCode::JumpForwardIf:
+			case ERigVMOpCode::JumpBackwardIf:
 			{
-				FRigVMByteCodeEntry Entry;
-				ScriptStruct->ImportText(*View[EntryIndex], &Entry, nullptr, PPF_None, nullptr, ScriptStruct->GetName());
-				Entries.Add(Entry);
+				FRigVMJumpIfOp Op;
+				Ar << Op;
+				AddOp<FRigVMJumpIfOp>(Op);
+				break;
+			}
+			case ERigVMOpCode::Exit:
+			{	
+				AddExitOp();	
+				break;
+			}
+			case ERigVMOpCode::BeginBlock:
+			{
+				FRigVMBinaryOp Op;
+				Ar << Op;
+				AddOp<FRigVMBinaryOp>(Op);
+				break;
+			}
+			case ERigVMOpCode::EndBlock:
+			{
+				AddEndBlockOp();
+				break;
+			}
+			default:
+			{
+				ensure(false);
 			}
 		}
 	}
-	else if(Ar.IsSaving())
+	
+	AlignByteCode();
+
+	Entries.Reset();
+	if (Ar.CustomVer(FAnimObjectVersion::GUID) >= FAnimObjectVersion::SerializeRigVMEntries)
 	{
 		UScriptStruct* ScriptStruct = FRigVMByteCodeEntry::StaticStruct();
-		TArray<uint8, TAlignedHeapAllocator<16>> DefaultStructData;
-		DefaultStructData.AddZeroed(ScriptStruct->GetStructureSize());
-		ScriptStruct->InitializeDefaultValue(DefaultStructData.GetData());
 
 		TArray<FString> View;
-		for (uint16 EntryIndex = 0; EntryIndex < Entries.Num(); EntryIndex++)
+		Ar << View;
+
+		for (int32 EntryIndex = 0; EntryIndex < View.Num(); EntryIndex++)
 		{
-			FString Value;
-			ScriptStruct->ExportText(Value, &Entries[EntryIndex], DefaultStructData.GetData(), nullptr, PPF_None, nullptr);
-			View.Add(Value);
+			FRigVMByteCodeEntry Entry;
+			ScriptStruct->ImportText(*View[EntryIndex], &Entry, nullptr, PPF_None, nullptr, ScriptStruct->GetName());
+			Entries.Add(Entry);
 		}
-
-		Ar << View;
 	}
-
-	return true;
 }
 
 void FRigVMByteCode::Reset()
diff --git a/Engine/Source/Runtime/RigVM/Private/RigVMCore/RigVMMemory.cpp b/Engine/Source/Runtime/RigVM/Private/RigVMCore/RigVMMemory.cpp
--- a/Engine/Source/Runtime/RigVM/Private/RigVMCore/RigVMMemory.cpp
+++ b/Engine/Source/Runtime/RigVM/Private/RigVMCore/RigVMMemory.cpp
@@ -13,21 +13,20 @@
 
 ////////////////////////////////////////////////////////////////////////////////////////////////////////////////////////
 
-bool FRigVMOperand::Serialize(FArchive& Ar)
+void FRigVMOperand::Serialize(FArchive& Ar)
 {
 	Ar << MemoryType;
 	Ar << RegisterIndex;
 	Ar << RegisterOffset;
-	return true;
 }
 
-bool FRigVMRegister::Serialize(FArchive& Ar)
+void FRigVMRegister::Serialize(FArchive& Ar)
 {
 	Ar.UsingCustomVersion(FAnimObjectVersion::GUID);
 
 	if (Ar.CustomVer(FAnimObjectVersion::GUID) < FAnimObjectVersion::StoreMarkerNamesOnSkeleton)
 	{
-		return false;
+		return;
 	}
 
 	uint16 SliceIndex = 0;
@@ -43,52 +42,94 @@ bool FRigVMRegister::Serialize(FArchive& Ar)
 	Ar << Name;
 	Ar << ScriptStructIndex;
 
-	if (Ar.IsLoading())
+	if (Ar.IsSaving() || Ar.IsObjectReferenceCollector())
 	{
-		if (Ar.CustomVer(FAnimObjectVersion::GUID) >= FAnimObjectVersion::SerializeRigVMRegisterArrayState)
-		{
-			Ar << bIsArray;
-		}
-		else
-		{
-			bIsArray = false;
-		}
-
-		if (Ar.CustomVer(FAnimObjectVersion::GUID) >= FAnimObjectVersion::SerializeRigVMRegisterDynamicState)
-		{
-			Ar << bIsDynamic;
-		}
-		else
-		{
-			bIsDynamic = false;
-		}
+		Save(Ar);
+	}
+	else if (Ar.IsLoading())
+	{
+		Load(Ar);
 	}
 	else
+	{
+		checkNoEntry();
+	}
+}
+
+void FRigVMRegister::Save(FArchive& Ar)
+{
+	Ar << bIsArray;
+	Ar << bIsDynamic;
+}
+
+void FRigVMRegister::Load(FArchive& Ar)
+{
+	if (Ar.CustomVer(FAnimObjectVersion::GUID) >= FAnimObjectVersion::SerializeRigVMRegisterArrayState)
 	{
 		Ar << bIsArray;
-		Ar << bIsDynamic;
+	}
+	else
+	{
+		bIsArray = false;
 	}
 
-	return true;
+	if (Ar.CustomVer(FAnimObjectVersion::GUID) >= FAnimObjectVersion::SerializeRigVMRegisterDynamicState)
+	{
+		Ar << bIsDynamic;
+	}
+	else
+	{
+		bIsDynamic = false;
+	}
 }
 
 ////////////////////////////////////////////////////////////////////////////////////////////////////////////////////////
 
-bool FRigVMRegisterOffset::Serialize(FArchive& Ar)
+void FRigVMRegisterOffset::Serialize(FArchive& Ar)
 {
 	Ar.UsingCustomVersion(FAnimObjectVersion::GUID);
 	Ar.UsingCustomVersion(FReleaseObjectVersion::GUID);
 
 	if (Ar.CustomVer(FAnimObjectVersion::GUID) < FAnimObjectVersion::StoreMarkerNamesOnSkeleton)
 	{
-		return false;
+		return;
+	}
+
+	if (Ar.IsSaving() || Ar.IsObjectReferenceCollector())
+	{
+		Save(Ar);
+	}
+	else if (Ar.IsLoading())
+	{
+		Load(Ar);
+	}
+	else
+	{
+		checkNoEntry();
 	}
+}
+
+void FRigVMRegisterOffset::Save(FArchive& Ar)
+{
+	Ar << Segments;
+	Ar << Type;
+	Ar << CPPType;
 
+	Ar << ScriptStruct;
+	
+	Ar << ElementSize;
+	Ar << ParentScriptStruct;
+	Ar << CachedSegmentPath;
+	Ar << ArrayIndex;
+}
+
+void FRigVMRegisterOffset::Load(FArchive& Ar)
+{
 	Ar << Segments;
 	Ar << Type;
 	Ar << CPPType;
 
-	if (Ar.IsLoading() && Ar.CustomVer(FReleaseObjectVersion::GUID) < FReleaseObjectVersion::SerializeRigVMOffsetSegmentPaths)
+	if (Ar.CustomVer(FReleaseObjectVersion::GUID) < FReleaseObjectVersion::SerializeRigVMOffsetSegmentPaths)
 	{
 		FName ScriptStructPath;
 		Ar << ScriptStructPath;
@@ -102,43 +143,32 @@ bool FRigVMRegisterOffset::Serialize(FArchive& Ar)
 
 	Ar << ElementSize;
 
-	if (Ar.IsLoading())
+	if (Ar.CustomVer(FReleaseObjectVersion::GUID) >= FReleaseObjectVersion::SerializeRigVMOffsetSegmentPaths)
 	{
-		if (Ar.CustomVer(FReleaseObjectVersion::GUID) >= FReleaseObjectVersion::SerializeRigVMOffsetSegmentPaths)
-		{
-			FString SegmentPath;
-			Ar << ParentScriptStruct;
-			Ar << SegmentPath;
-			Ar << ArrayIndex;
+		FString SegmentPath;
+		Ar << ParentScriptStruct;
+		Ar << SegmentPath;
+		Ar << ArrayIndex;
 
-			if (Ar.IsTransacting())
+		if (Ar.IsTransacting())
+		{
+			CachedSegmentPath = SegmentPath;
+		}
+		else if (ParentScriptStruct != nullptr && !SegmentPath.IsEmpty())
+		{
+			int32 InitialOffset = ArrayIndex * ParentScriptStruct->GetStructureSize();
+			FRigVMRegisterOffset TempOffset(ParentScriptStruct, SegmentPath, InitialOffset, ElementSize);
+			if (TempOffset.GetSegments().Num() == Segments.Num())
 			{
+				Segments = TempOffset.GetSegments();
 				CachedSegmentPath = SegmentPath;
 			}
-			else if (ParentScriptStruct != nullptr && !SegmentPath.IsEmpty())
+			else
 			{
-				int32 InitialOffset = ArrayIndex * ParentScriptStruct->GetStructureSize();
-				FRigVMRegisterOffset TempOffset(ParentScriptStruct, SegmentPath, InitialOffset, ElementSize);
-				if (TempOffset.GetSegments().Num() == Segments.Num())
-				{
-					Segments = TempOffset.GetSegments();
-					CachedSegmentPath = SegmentPath;
-				}
-				else
-				{
-					checkNoEntry();
-				}
+				checkNoEntry();
 			}
 		}
 	}
-	else
-	{
-		Ar << ParentScriptStruct;
-		Ar << CachedSegmentPath;
-		Ar << ArrayIndex;
-	}
-
-	return true;
 }
 
 FRigVMRegisterOffset::FRigVMRegisterOffset(UScriptStruct* InScriptStruct, const FString& InSegmentPath, int32 InInitialOffset, uint16 InElementSize)
@@ -457,37 +487,42 @@ FRigVMMemoryContainer& FRigVMMemoryContainer::operator= (const FRigVMMemoryConta
 	return *this;
 }
 
-bool FRigVMMemoryContainer::Serialize(FArchive& Ar)
+void FRigVMMemoryContainer::Serialize(FArchive& Ar)
 {
 	Ar.UsingCustomVersion(FAnimObjectVersion::GUID);
 
  	if (Ar.CustomVer(FAnimObjectVersion::GUID) < FAnimObjectVersion::StoreMarkerNamesOnSkeleton)
  	{
- 		return false;
+ 		return;
  	}
 
-	if (Ar.IsLoading())
+	if (Ar.IsSaving() || Ar.IsObjectReferenceCollector())
 	{
-		for (int32 RegisterIndex = 0; RegisterIndex < Registers.Num(); RegisterIndex++)
-		{
-			Destroy(RegisterIndex);
-		}
+		Save(Ar);
 	}
+	else if (Ar.IsLoading())
+	{
+		Load(Ar);
+	}
+	else
+	{
+		checkNoEntry();
+	}
+}
 
-	if (Ar.IsSaving())
+void FRigVMMemoryContainer::Save(FArchive& Ar)
+{
+	for (FRigVMRegister& Register : Registers)
 	{
-		for (FRigVMRegister& Register : Registers)
+		if (Register.IsNestedDynamic())
 		{
-			if (Register.IsNestedDynamic())
-			{
-				FRigVMNestedByteArray* NestedArrayStorage = (FRigVMNestedByteArray*)&Data[Register.GetWorkByteIndex()];
-				Register.SliceCount = NestedArrayStorage->Num();
-			}
-			else if (Register.IsDynamic())
-			{
-				FRigVMByteArray* ArrayStorage = (FRigVMByteArray*)&Data[Register.GetWorkByteIndex()];
-				Register.SliceCount = ArrayStorage->Num();
-			}
+			FRigVMNestedByteArray* NestedArrayStorage = (FRigVMNestedByteArray*)&Data[Register.GetWorkByteIndex()];
+			Register.SliceCount = NestedArrayStorage->Num();
+		}
+		else if (Register.IsDynamic())
+		{
+			FRigVMByteArray* ArrayStorage = (FRigVMByteArray*)&Data[Register.GetWorkByteIndex()];
+			Register.SliceCount = ArrayStorage->Num();
 		}
 	}
 
@@ -496,222 +531,329 @@ bool FRigVMMemoryContainer::Serialize(FArchive& Ar)
 	Ar << Registers;
 	Ar << RegisterOffsets;
 
-	if (Ar.IsLoading())
+	
+	TArray<FString> ScriptStructPaths;
+	for (UScriptStruct* ScriptStruct : ScriptStructs)
 	{
-#if DEBUG_RIGVMMEMORY
-		UE_LOG_RIGVMMEMORY(TEXT("%d Memory - Begin Loading..."), (int32)GetMemoryType());
-#endif
+		ScriptStructPaths.Add(ScriptStruct->GetPathName());
+	}
+	Ar << ScriptStructPaths;
 
-		bEncounteredErrorDuringLoad = false;
+	uint64 TotalBytes = Data.Num();
+	Ar << TotalBytes;
 
-		ScriptStructs.Reset();
-		TArray<FString> ScriptStructPaths;
-		Ar << ScriptStructPaths;
-		
-		for (const FString& ScriptStructPath : ScriptStructPaths)
+	for (FRigVMRegister& Register : Registers)
+	{
+		if (Register.ElementCount == 0 && !Register.IsDynamic())
 		{
-			UScriptStruct* ScriptStruct = FindObject<UScriptStruct>(nullptr, *ScriptStructPath);
-
-			// this might have happened if a given script struct no longer 
-			// exists or cannot be loaded.
-			if (ScriptStruct == nullptr)
-			{
-				FString PackagePath = Ar.GetArchiveName();
-				UE_LOG(LogRigVM, Error, TEXT("Struct '%s' cannot be found. Asset '%s' no longer functional."), *ScriptStructPath, *PackagePath);
-				bEncounteredErrorDuringLoad = true;
-			}
-
-			ScriptStructs.Add(ScriptStruct);
+			continue;
 		}
 
-		uint64 TotalBytes = 0;
-		Ar << TotalBytes;
-		
-		Data.Empty();
-
-		if (!bEncounteredErrorDuringLoad)
+		if (!Register.IsDynamic())
 		{
-			// during load we'll recreate the memory for all registers.
-			// the size for structs might have changed, so we need to reallocate.
-			for (int32 RegisterIndex = 0; RegisterIndex < Registers.Num(); RegisterIndex++)
+			switch (Register.Type)
 			{
-				FRigVMRegister& Register = Registers[RegisterIndex];
-
-				UScriptStruct* ScriptStruct = GetScriptStruct(Register);
-				if (ScriptStruct)
+				case ERigVMRegisterType::Plain:
 				{
-					Register.ElementSize = ScriptStruct->GetStructureSize();
+					FRigVMByteArray View;
+					View.Append(&Data[Register.GetWorkByteIndex()], Register.GetAllocatedBytes() - Register.GetAlignmentBytes());
+					Ar << View;
+					break;
 				}
-				else if (Register.Type == ERigVMRegisterType::Name)
+				case ERigVMRegisterType::Name:
 				{
-					Register.ElementSize = sizeof(FName);
+					TArray<FName> View;
+					View.Append((FName*)&Data[Register.GetWorkByteIndex()], Register.GetTotalElementCount());
+					Ar << View;
+					break;
 				}
-				else if (Register.Type == ERigVMRegisterType::String)
+				case ERigVMRegisterType::String:
 				{
-					Register.ElementSize = sizeof(FString);
+					TArray<FString> View;
+					View.Append((FString*)&Data[Register.GetWorkByteIndex()], Register.GetTotalElementCount());
+					Ar << View;
+					break;
 				}
+				case ERigVMRegisterType::Struct:
+				{
+					uint8* DataPtr = &Data[Register.GetWorkByteIndex()];
+					UScriptStruct* ScriptStruct = GetScriptStruct(Register);
 
-				Register.AlignmentBytes = 0;
-				Register.TrailingBytes = 0;
+					TArray<uint8, TAlignedHeapAllocator<16>> DefaultStructData;
+					DefaultStructData.AddZeroed(ScriptStruct->GetStructureSize());
+					ScriptStruct->InitializeDefaultValue(DefaultStructData.GetData());
 
-				if (Registers[RegisterIndex].IsDynamic())
-				{
-					Register.ByteIndex = Data.AddZeroed(sizeof(FRigVMByteArray));
-				}
-				else
-				{
-					Register.ByteIndex = Data.AddZeroed(Register.GetNumBytesAllSlices());
+					TArray<FString> View;
+					for (uint16 ElementIndex = 0; ElementIndex < Register.GetTotalElementCount(); ElementIndex++)
+					{
+						FString Value;
+						ScriptStruct->ExportText(Value, DataPtr, DefaultStructData.GetData(), nullptr, PPF_None, nullptr);
+						View.Add(Value);
+						DataPtr += Register.ElementSize;
+					}
+
+					ScriptStruct->DestroyStruct(DefaultStructData.GetData(), 1);
+
+					Ar << View;
+					break;
 				}
 			}
-			UpdateRegisters();
+		}
+		else if (!Register.IsNestedDynamic())
+		{
+			FRigVMByteArray* ArrayStorage = (FRigVMByteArray*)&Data[Register.GetWorkByteIndex()];
 
-			// once the register memory is allocated we can construt its contents.
-			for (int32 RegisterIndex = 0; RegisterIndex < Registers.Num(); RegisterIndex++)
+			switch (Register.Type)
+			{
+			case ERigVMRegisterType::Plain:
+			{
+				Ar << *ArrayStorage;
+				break;
+			}
+			case ERigVMRegisterType::Name:
+			{
+				TArray<FName> View = FRigVMFixedArray<FName>(*ArrayStorage);
+				Ar << View;
+				break;
+			}
+			case ERigVMRegisterType::String:
+			{
+				TArray<FString> View = FRigVMFixedArray<FString>(*ArrayStorage);
+				Ar << View;
+				break;
+			}
+			case ERigVMRegisterType::Struct:
 			{
-				const FRigVMRegister& Register = Registers[RegisterIndex];
-				if (!Registers[RegisterIndex].IsDynamic())
+				uint8* DataPtr = ArrayStorage->GetData();
+				uint16 NumElements = (uint16)(ArrayStorage->Num() / Register.ElementSize);
+				ensure(NumElements * Register.ElementSize == ArrayStorage->Num());
+				UScriptStruct* ScriptStruct = GetScriptStruct(Register);
+
+				TArray<FString> View;
+				for (uint16 ElementIndex = 0; ElementIndex < NumElements; ElementIndex++)
 				{
-					Construct(RegisterIndex);
+					FString Value;
+					ScriptStruct->ExportText(Value, DataPtr, nullptr, nullptr, PPF_None, nullptr);
+					View.Add(Value);
+					DataPtr += Register.ElementSize;
 				}
+
+				Ar << View;
+				break;
+			}
 			}
 		}
-
-		for (int32 RegisterIndex=0;RegisterIndex<Registers.Num();RegisterIndex++)
+		else
 		{
-			FRigVMRegister& Register = Registers[RegisterIndex];
+			FRigVMNestedByteArray* NestedArrayStorage = (FRigVMNestedByteArray*)&Data[Register.GetWorkByteIndex()];
 
-			if (Register.ElementCount == 0 && !Register.IsDynamic())
+			for (int32 SliceIndex = 0; SliceIndex < Register.SliceCount; SliceIndex++)
 			{
-				continue;
-			}
+				if (NestedArrayStorage->Num() <= SliceIndex)
+				{
+					FRigVMByteArray EmptyStorage;
+					Ar << EmptyStorage;
+					continue;
+				}
+
+				FRigVMByteArray& ArrayStorage = (*NestedArrayStorage)[SliceIndex];
 
-			if (!Register.IsDynamic())
-			{
 				switch (Register.Type)
 				{
 					case ERigVMRegisterType::Plain:
 					{
-						FRigVMByteArray View;
-						Ar << View;
-						if (!bEncounteredErrorDuringLoad)
-						{
-							ensure(View.Num() <= Register.GetAllocatedBytes());
-							FMemory::Memcpy(&Data[Register.GetWorkByteIndex()], View.GetData(), View.Num());
-						}
+						Ar << ArrayStorage;
 						break;
 					}
 					case ERigVMRegisterType::Name:
 					{
-						TArray<FName> View;
+						TArray<FName> View = FRigVMFixedArray<FName>(ArrayStorage);
 						Ar << View;
-						if (!bEncounteredErrorDuringLoad)
-						{
-							ensure(View.Num() == Register.GetTotalElementCount());
-							RigVMCopy<FName>(&Data[Register.GetWorkByteIndex()], View.GetData(), View.Num());
-						}
 						break;
 					}
 					case ERigVMRegisterType::String:
 					{
-						TArray<FString> View;
+						TArray<FString> View = FRigVMFixedArray<FString>(ArrayStorage);
 						Ar << View;
-						if (!bEncounteredErrorDuringLoad)
-						{
-							ensure(View.Num() == Register.GetTotalElementCount());
-							RigVMCopy<FString>(&Data[Register.GetWorkByteIndex()], View.GetData(), View.Num());
-						}
 						break;
 					}
 					case ERigVMRegisterType::Struct:
 					{
+						uint8* DataPtr = ArrayStorage.GetData();
+						uint16 NumElements = (uint16)(ArrayStorage.Num() / Register.ElementSize);
+						ensure(NumElements * Register.ElementSize == ArrayStorage.Num());
+						UScriptStruct* ScriptStruct = GetScriptStruct(Register);
+
 						TArray<FString> View;
-						Ar << View;
-						if (!bEncounteredErrorDuringLoad)
+						for (uint16 ElementIndex = 0; ElementIndex < NumElements; ElementIndex++)
 						{
-							ensure(View.Num() == Register.GetTotalElementCount());
-
-							uint8* DataPtr = &Data[Register.GetWorkByteIndex()];
-							UScriptStruct* ScriptStruct = GetScriptStruct(Register);
-							if (ScriptStruct && !bEncounteredErrorDuringLoad)
-							{
-								ensure(ScriptStruct->GetStructureSize() == Register.ElementSize);
-
-								for (uint16 ElementIndex = 0; ElementIndex < Register.GetTotalElementCount(); ElementIndex++)
-								{
-									FRigVMMemoryContainerImportErrorContext ErrorPipe;
-									ScriptStruct->ImportText(*View[ElementIndex], DataPtr, nullptr, PPF_None, &ErrorPipe, ScriptStruct->GetName());
-									if (ErrorPipe.NumErrors > 0)
-									{
-										bEncounteredErrorDuringLoad = true;
-										break;
-									}
-									DataPtr += Register.ElementSize;
-								}
-							}
+							FString Value;
+							ScriptStruct->ExportText(Value, DataPtr, nullptr, nullptr, PPF_None, nullptr);
+							View.Add(Value);
+							DataPtr += Register.ElementSize;
 						}
+
+						Ar << View;
 						break;
 					}
 				}
 			}
-			else if (!Register.IsNestedDynamic())
+		}
+	}
+}
+
+void FRigVMMemoryContainer::Load(FArchive& Ar)
+{
+	for (int32 RegisterIndex = 0; RegisterIndex < Registers.Num(); RegisterIndex++)
+	{
+		Destroy(RegisterIndex);
+	}
+	
+	Ar << bUseNameMap;
+	Ar << MemoryType;
+	Ar << Registers;
+	Ar << RegisterOffsets;
+
+
+#if DEBUG_RIGVMMEMORY
+	UE_LOG_RIGVMMEMORY(TEXT("%d Memory - Begin Loading..."), (int32)GetMemoryType());
+#endif
+
+	bEncounteredErrorDuringLoad = false;
+
+	ScriptStructs.Reset();
+	TArray<FString> ScriptStructPaths;
+	Ar << ScriptStructPaths;
+
+	for (const FString& ScriptStructPath : ScriptStructPaths)
+	{
+		UScriptStruct* ScriptStruct = FindObject<UScriptStruct>(nullptr, *ScriptStructPath);
+
+		// this might have happened if a given script struct no longer 
+		// exists or cannot be loaded.
+		if (ScriptStruct == nullptr)
+		{
+			FString PackagePath = Ar.GetArchiveName();
+			UE_LOG(LogRigVM, Error, TEXT("Struct '%s' cannot be found. Asset '%s' no longer functional."), *ScriptStructPath, *PackagePath);
+			bEncounteredErrorDuringLoad = true;
+		}
+
+		ScriptStructs.Add(ScriptStruct);
+	}
+
+	uint64 TotalBytes = 0;
+	Ar << TotalBytes;
+
+	Data.Empty();
+
+	if (!bEncounteredErrorDuringLoad)
+	{
+		// during load we'll recreate the memory for all registers.
+		// the size for structs might have changed, so we need to reallocate.
+		for (int32 RegisterIndex = 0; RegisterIndex < Registers.Num(); RegisterIndex++)
+		{
+			FRigVMRegister& Register = Registers[RegisterIndex];
+
+			UScriptStruct* ScriptStruct = GetScriptStruct(Register);
+			if (ScriptStruct)
+			{
+				Register.ElementSize = ScriptStruct->GetStructureSize();
+			}
+			else if (Register.Type == ERigVMRegisterType::Name)
 			{
-				FRigVMByteArray* ArrayStorage = nullptr;
+				Register.ElementSize = sizeof(FName);
+			}
+			else if (Register.Type == ERigVMRegisterType::String)
+			{
+				Register.ElementSize = sizeof(FString);
+			}
 
-				if (!bEncounteredErrorDuringLoad)
-				{
-					ArrayStorage = (FRigVMByteArray*)&Data[Register.GetWorkByteIndex()];
-				}
+			Register.AlignmentBytes = 0;
+			Register.TrailingBytes = 0;
 
-				switch (Register.Type)
+			if (Registers[RegisterIndex].IsDynamic())
+			{
+				Register.ByteIndex = Data.AddZeroed(sizeof(FRigVMByteArray));
+			}
+			else
+			{
+				Register.ByteIndex = Data.AddZeroed(Register.GetNumBytesAllSlices());
+			}
+		}
+		UpdateRegisters();
+
+		// once the register memory is allocated we can construt its contents.
+		for (int32 RegisterIndex = 0; RegisterIndex < Registers.Num(); RegisterIndex++)
+		{
+			const FRigVMRegister& Register = Registers[RegisterIndex];
+			if (!Registers[RegisterIndex].IsDynamic())
+			{
+				Construct(RegisterIndex);
+			}
+		}
+	}
+
+	for (int32 RegisterIndex = 0; RegisterIndex < Registers.Num(); RegisterIndex++)
+	{
+		FRigVMRegister& Register = Registers[RegisterIndex];
+
+		if (Register.ElementCount == 0 && !Register.IsDynamic())
+		{
+			continue;
+		}
+
+		if (!Register.IsDynamic())
+		{
+			switch (Register.Type)
+			{
+				case ERigVMRegisterType::Plain:
 				{
-					case ERigVMRegisterType::Plain:
+					FRigVMByteArray View;
+					Ar << View;
+					if (!bEncounteredErrorDuringLoad)
 					{
-						FRigVMByteArray View;
-						Ar << View;
-						if (ArrayStorage)
-						{
-							*ArrayStorage = View;
-						}
-						break;
+						ensure(View.Num() <= Register.GetAllocatedBytes());
+						FMemory::Memcpy(&Data[Register.GetWorkByteIndex()], View.GetData(), View.Num());
 					}
-					case ERigVMRegisterType::Name:
+					break;
+				}
+				case ERigVMRegisterType::Name:
+				{
+					TArray<FName> View;
+					Ar << View;
+					if (!bEncounteredErrorDuringLoad)
 					{
-						TArray<FName> View;
-						Ar << View;
-
-						if (ArrayStorage)
-						{
-							FRigVMDynamicArray<FName> ValueArray(*ArrayStorage);
-							ValueArray.CopyFrom(View);
-						}
-						break;
+						ensure(View.Num() == Register.GetTotalElementCount());
+						RigVMCopy<FName>(&Data[Register.GetWorkByteIndex()], View.GetData(), View.Num());
 					}
-					case ERigVMRegisterType::String:
+					break;
+				}
+				case ERigVMRegisterType::String:
+				{
+					TArray<FString> View;
+					Ar << View;
+					if (!bEncounteredErrorDuringLoad)
 					{
-						TArray<FString> View;
-						Ar << View;
-
-						if (ArrayStorage)
-						{
-							FRigVMDynamicArray<FString> ValueArray(*ArrayStorage);
-							ValueArray.CopyFrom(View);
-						}
-						break;
+						ensure(View.Num() == Register.GetTotalElementCount());
+						RigVMCopy<FString>(&Data[Register.GetWorkByteIndex()], View.GetData(), View.Num());
 					}
-					case ERigVMRegisterType::Struct:
+					break;
+				}
+				case ERigVMRegisterType::Struct:
+				{
+					TArray<FString> View;
+					Ar << View;
+					if (!bEncounteredErrorDuringLoad)
 					{
-						TArray<FString> View;
-						Ar << View;
+						ensure(View.Num() == Register.GetTotalElementCount());
 
+						uint8* DataPtr = &Data[Register.GetWorkByteIndex()];
 						UScriptStruct* ScriptStruct = GetScriptStruct(Register);
-						if (ScriptStruct && !bEncounteredErrorDuringLoad && ArrayStorage)
+						if (ScriptStruct && !bEncounteredErrorDuringLoad)
 						{
 							ensure(ScriptStruct->GetStructureSize() == Register.ElementSize);
 
-							ArrayStorage->SetNumZeroed(View.Num() * Register.ElementSize);
-							uint8* DataPtr = ArrayStorage->GetData();
-
-							for (uint16 ElementIndex = 0; ElementIndex < View.Num(); ElementIndex++)
+							for (uint16 ElementIndex = 0; ElementIndex < Register.GetTotalElementCount(); ElementIndex++)
 							{
 								FRigVMMemoryContainerImportErrorContext ErrorPipe;
 								ScriptStruct->ImportText(*View[ElementIndex], DataPtr, nullptr, PPF_None, &ErrorPipe, ScriptStruct->GetName());
@@ -723,285 +865,187 @@ bool FRigVMMemoryContainer::Serialize(FArchive& Ar)
 								DataPtr += Register.ElementSize;
 							}
 						}
-						break;
 					}
+					break;
 				}
 			}
-			else
+		}
+		else if (!Register.IsNestedDynamic())
+		{
+			FRigVMByteArray* ArrayStorage = nullptr;
+
+			if (!bEncounteredErrorDuringLoad)
 			{
-				FRigVMNestedByteArray* NestedArrayStorage = nullptr; 
-				if (!bEncounteredErrorDuringLoad)
+				ArrayStorage = (FRigVMByteArray*)&Data[Register.GetWorkByteIndex()];
+			}
+
+			switch (Register.Type)
+			{
+				case ERigVMRegisterType::Plain:
 				{
-					NestedArrayStorage = (FRigVMNestedByteArray*)&Data[Register.GetWorkByteIndex()];
-					NestedArrayStorage->Reset();
-					NestedArrayStorage->SetNumZeroed(Register.SliceCount);
+					FRigVMByteArray View;
+					Ar << View;
+					if (ArrayStorage)
+					{
+						*ArrayStorage = View;
+					}
+					break;
 				}
+				case ERigVMRegisterType::Name:
+				{
+					TArray<FName> View;
+					Ar << View;
 
-				for (int32 SliceIndex = 0; SliceIndex < Register.SliceCount; SliceIndex++)
+					if (ArrayStorage)
+					{
+						FRigVMDynamicArray<FName> ValueArray(*ArrayStorage);
+						ValueArray.CopyFrom(View);
+					}
+					break;
+				}
+				case ERigVMRegisterType::String:
 				{
-					FRigVMByteArray* ArrayStorage = nullptr;
-					if (NestedArrayStorage)
+					TArray<FString> View;
+					Ar << View;
+
+					if (ArrayStorage)
 					{
-						ArrayStorage = &(*NestedArrayStorage)[SliceIndex];
+						FRigVMDynamicArray<FString> ValueArray(*ArrayStorage);
+						ValueArray.CopyFrom(View);
 					}
+					break;
+				}
+				case ERigVMRegisterType::Struct:
+				{
+					TArray<FString> View;
+					Ar << View;
 
-					switch (Register.Type)
+					UScriptStruct* ScriptStruct = GetScriptStruct(Register);
+					if (ScriptStruct && !bEncounteredErrorDuringLoad && ArrayStorage)
 					{
-						case ERigVMRegisterType::Plain:
-						{
-							FRigVMByteArray View;
-							Ar << View;
-							if (ArrayStorage)
-							{
-								*ArrayStorage = View;
-							}
-							break;
-						}
-						case ERigVMRegisterType::Name:
-						{
-							TArray<FName> View;
-							Ar << View;
+						ensure(ScriptStruct->GetStructureSize() == Register.ElementSize);
 
-							if (ArrayStorage)
-							{
-								FRigVMDynamicArray<FName> ValueArray(*ArrayStorage);
-								ValueArray.CopyFrom(View);
-							}
-							break;
-						}
-						case ERigVMRegisterType::String:
-						{
-							TArray<FString> View;
-							Ar << View;
+						ArrayStorage->SetNumZeroed(View.Num() * Register.ElementSize);
+						uint8* DataPtr = ArrayStorage->GetData();
 
-							if (ArrayStorage)
-							{
-								FRigVMDynamicArray<FString> ValueArray(*ArrayStorage);
-								ValueArray.CopyFrom(View);
-							}
-							break;
-						}
-						case ERigVMRegisterType::Struct:
+						for (uint16 ElementIndex = 0; ElementIndex < View.Num(); ElementIndex++)
 						{
-							TArray<FString> View;
-							Ar << View;
-
-							UScriptStruct* ScriptStruct = GetScriptStruct(Register);
-							if (ScriptStruct && !bEncounteredErrorDuringLoad && ArrayStorage)
+							FRigVMMemoryContainerImportErrorContext ErrorPipe;
+							ScriptStruct->ImportText(*View[ElementIndex], DataPtr, nullptr, PPF_None, &ErrorPipe, ScriptStruct->GetName());
+							if (ErrorPipe.NumErrors > 0)
 							{
-								ensure(ScriptStruct->GetStructureSize() == Register.ElementSize);
-
-								ArrayStorage->SetNumZeroed(View.Num() * Register.ElementSize);
-								uint8* DataPtr = ArrayStorage->GetData();
-
-								for (uint16 ElementIndex = 0; ElementIndex < View.Num(); ElementIndex++)
-								{
-									FRigVMMemoryContainerImportErrorContext ErrorPipe;
-									ScriptStruct->ImportText(*View[ElementIndex], DataPtr, nullptr, PPF_None, &ErrorPipe, ScriptStruct->GetName());
-									if (ErrorPipe.NumErrors > 0)
-									{
-										bEncounteredErrorDuringLoad = true;
-										break;
-									}
-									DataPtr += Register.ElementSize;
-								}
+								bEncounteredErrorDuringLoad = true;
+								break;
 							}
-							break;
+							DataPtr += Register.ElementSize;
 						}
 					}
+					break;
 				}
 			}
 		}
-
-		if (bEncounteredErrorDuringLoad)
-		{
-#if DEBUG_RIGVMMEMORY
-			UE_LOG_RIGVMMEMORY(TEXT("%d Memory - Encountered errors during load."), (int32)GetMemoryType());
-#endif
-			Reset();
-		}
 		else
 		{
-			UpdateRegisters();
-		}
-
-#if DEBUG_RIGVMMEMORY
-		UE_LOG_RIGVMMEMORY(TEXT("%d Memory - Finished Loading."), (int32)GetMemoryType());
-#endif
-	}
-	else
-	{
-		TArray<FString> ScriptStructPaths;
-		for (UScriptStruct* ScriptStruct : ScriptStructs)
-		{
-			ScriptStructPaths.Add(ScriptStruct->GetPathName());
-		}
-		Ar << ScriptStructPaths;
-
-		uint64 TotalBytes = Data.Num();
-		Ar << TotalBytes;
-
-		for (FRigVMRegister& Register : Registers)
-		{
-			if (Register.ElementCount == 0 && !Register.IsDynamic())
+			FRigVMNestedByteArray* NestedArrayStorage = nullptr;
+			if (!bEncounteredErrorDuringLoad)
 			{
-				continue;
+				NestedArrayStorage = (FRigVMNestedByteArray*)&Data[Register.GetWorkByteIndex()];
+				NestedArrayStorage->Reset();
+				NestedArrayStorage->SetNumZeroed(Register.SliceCount);
 			}
 
-			if (!Register.IsDynamic())
+			for (int32 SliceIndex = 0; SliceIndex < Register.SliceCount; SliceIndex++)
 			{
+				FRigVMByteArray* ArrayStorage = nullptr;
+				if (NestedArrayStorage)
+				{
+					ArrayStorage = &(*NestedArrayStorage)[SliceIndex];
+				}
+
 				switch (Register.Type)
 				{
 					case ERigVMRegisterType::Plain:
 					{
 						FRigVMByteArray View;
-						View.Append(&Data[Register.GetWorkByteIndex()], Register.GetAllocatedBytes() - Register.GetAlignmentBytes());
 						Ar << View;
+						if (ArrayStorage)
+						{
+							*ArrayStorage = View;
+						}
 						break;
 					}
 					case ERigVMRegisterType::Name:
 					{
 						TArray<FName> View;
-						View.Append((FName*)&Data[Register.GetWorkByteIndex()], Register.GetTotalElementCount());
 						Ar << View;
+
+						if (ArrayStorage)
+						{
+							FRigVMDynamicArray<FName> ValueArray(*ArrayStorage);
+							ValueArray.CopyFrom(View);
+						}
 						break;
 					}
 					case ERigVMRegisterType::String:
 					{
 						TArray<FString> View;
-						View.Append((FString*)&Data[Register.GetWorkByteIndex()], Register.GetTotalElementCount());
 						Ar << View;
-						break;
-					}
-					case ERigVMRegisterType::Struct:
-					{
-						uint8* DataPtr = &Data[Register.GetWorkByteIndex()];
-						UScriptStruct* ScriptStruct = GetScriptStruct(Register);
 
-						TArray<uint8, TAlignedHeapAllocator<16>> DefaultStructData;
-						DefaultStructData.AddZeroed(ScriptStruct->GetStructureSize());
-						ScriptStruct->InitializeDefaultValue(DefaultStructData.GetData());
-
-						TArray<FString> View;
-						for (uint16 ElementIndex = 0; ElementIndex < Register.GetTotalElementCount(); ElementIndex++)
+						if (ArrayStorage)
 						{
-							FString Value;
-							ScriptStruct->ExportText(Value, DataPtr, DefaultStructData.GetData(), nullptr, PPF_None, nullptr);
-							View.Add(Value);
-							DataPtr += Register.ElementSize;
+							FRigVMDynamicArray<FString> ValueArray(*ArrayStorage);
+							ValueArray.CopyFrom(View);
 						}
-
-						ScriptStruct->DestroyStruct(DefaultStructData.GetData(), 1);
-
-						Ar << View;
-						break;
-					}
-				}
-			}
-			else if (!Register.IsNestedDynamic())
-			{
-				FRigVMByteArray* ArrayStorage = (FRigVMByteArray*)&Data[Register.GetWorkByteIndex()];
-
-				switch (Register.Type)
-				{
-					case ERigVMRegisterType::Plain:
-					{
-						Ar << *ArrayStorage;
-						break;
-					}
-					case ERigVMRegisterType::Name:
-					{
-						TArray<FName> View = FRigVMFixedArray<FName>(*ArrayStorage);
-						Ar << View;
-						break;
-					}
-					case ERigVMRegisterType::String:
-					{
-						TArray<FString> View = FRigVMFixedArray<FString>(*ArrayStorage);
-						Ar << View;
 						break;
 					}
 					case ERigVMRegisterType::Struct:
 					{
-						uint8* DataPtr = ArrayStorage->GetData();
-						uint16 NumElements = (uint16)(ArrayStorage->Num() / Register.ElementSize);
-						ensure(NumElements * Register.ElementSize == ArrayStorage->Num());
-						UScriptStruct* ScriptStruct = GetScriptStruct(Register);
-
 						TArray<FString> View;
-						for (uint16 ElementIndex = 0; ElementIndex < NumElements; ElementIndex++)
-						{
-							FString Value;
-							ScriptStruct->ExportText(Value, DataPtr, nullptr, nullptr, PPF_None, nullptr);
-							View.Add(Value);
-							DataPtr += Register.ElementSize;
-						}
-
 						Ar << View;
-						break;
-					}
-				}
-			}
-			else
-			{
-				FRigVMNestedByteArray* NestedArrayStorage = (FRigVMNestedByteArray*)&Data[Register.GetWorkByteIndex()];
 
-				for (int32 SliceIndex = 0; SliceIndex < Register.SliceCount; SliceIndex++)
-				{
-					if (NestedArrayStorage->Num() <= SliceIndex)
-					{
-						FRigVMByteArray EmptyStorage;
-						Ar << EmptyStorage;
-						continue;
-					}
-
-					FRigVMByteArray& ArrayStorage = (*NestedArrayStorage)[SliceIndex];
-
-					switch (Register.Type)
-					{
-						case ERigVMRegisterType::Plain:
-						{
-							Ar << ArrayStorage;
-							break;
-						}
-						case ERigVMRegisterType::Name:
-						{
-							TArray<FName> View = FRigVMFixedArray<FName>(ArrayStorage);
-							Ar << View;
-							break;
-						}
-						case ERigVMRegisterType::String:
-						{
-							TArray<FString> View = FRigVMFixedArray<FString>(ArrayStorage);
-							Ar << View;
-							break;
-						}
-						case ERigVMRegisterType::Struct:
+						UScriptStruct* ScriptStruct = GetScriptStruct(Register);
+						if (ScriptStruct && !bEncounteredErrorDuringLoad && ArrayStorage)
 						{
-							uint8* DataPtr = ArrayStorage.GetData();
-							uint16 NumElements = (uint16)(ArrayStorage.Num() / Register.ElementSize);
-							ensure(NumElements * Register.ElementSize == ArrayStorage.Num());
-							UScriptStruct* ScriptStruct = GetScriptStruct(Register);
+							ensure(ScriptStruct->GetStructureSize() == Register.ElementSize);
 
-							TArray<FString> View;
-							for (uint16 ElementIndex = 0; ElementIndex < NumElements; ElementIndex++)
+							ArrayStorage->SetNumZeroed(View.Num() * Register.ElementSize);
+							uint8* DataPtr = ArrayStorage->GetData();
+
+							for (uint16 ElementIndex = 0; ElementIndex < View.Num(); ElementIndex++)
 							{
-								FString Value;
-								ScriptStruct->ExportText(Value, DataPtr, nullptr, nullptr, PPF_None, nullptr);
-								View.Add(Value);
+								FRigVMMemoryContainerImportErrorContext ErrorPipe;
+								ScriptStruct->ImportText(*View[ElementIndex], DataPtr, nullptr, PPF_None, &ErrorPipe, ScriptStruct->GetName());
+								if (ErrorPipe.NumErrors > 0)
+								{
+									bEncounteredErrorDuringLoad = true;
+									break;
+								}
 								DataPtr += Register.ElementSize;
 							}
-
-							Ar << View;
-							break;
 						}
+						break;
 					}
 				}
 			}
 		}
 	}
 
-	return true;
+	if (bEncounteredErrorDuringLoad)
+	{
+#if DEBUG_RIGVMMEMORY
+		UE_LOG_RIGVMMEMORY(TEXT("%d Memory - Encountered errors during load."), (int32)GetMemoryType());
+#endif
+		Reset();
+	}
+	else
+	{
+		UpdateRegisters();
+	}
+
+#if DEBUG_RIGVMMEMORY
+	UE_LOG_RIGVMMEMORY(TEXT("%d Memory - Finished Loading."), (int32)GetMemoryType());
+#endif
+	
 }
 
 void FRigVMMemoryContainer::Reset()
diff --git a/Engine/Source/Runtime/RigVM/Public/RigVMCore/RigVM.h b/Engine/Source/Runtime/RigVM/Public/RigVMCore/RigVM.h
--- a/Engine/Source/Runtime/RigVM/Public/RigVMCore/RigVM.h
+++ b/Engine/Source/Runtime/RigVM/Public/RigVMCore/RigVM.h
@@ -41,7 +41,9 @@ public:
 	{
 	}
 
-	bool Serialize(FArchive& Ar);
+	void Serialize(FArchive& Ar);
+	void Save(FArchive& Ar);
+	void Load(FArchive& Ar);
 	FORCEINLINE friend FArchive& operator<<(FArchive& Ar, FRigVMParameter& P)
 	{
 		P.Serialize(Ar);
@@ -121,6 +123,8 @@ public:
 
 	// UObject interface
 	void Serialize(FArchive& Ar);
+	void Save(FArchive& Ar);
+	void Load(FArchive& Ar);
 
 	// resets the container and maintains all memory
 	void Reset();
diff --git a/Engine/Source/Runtime/RigVM/Public/RigVMCore/RigVMByteCode.h b/Engine/Source/Runtime/RigVM/Public/RigVMCore/RigVMByteCode.h
--- a/Engine/Source/Runtime/RigVM/Public/RigVMCore/RigVMByteCode.h
+++ b/Engine/Source/Runtime/RigVM/Public/RigVMCore/RigVMByteCode.h
@@ -136,7 +136,7 @@ struct RIGVM_API FRigVMExecuteOp : public FRigVMBaseOp
 
 	FORCEINLINE uint8 GetOperandCount() const { return uint8(OpCode) - uint8(ERigVMOpCode::Execute_0_Operands); }
 
-	bool Serialize(FArchive& Ar);
+	void Serialize(FArchive& Ar);
 	FORCEINLINE friend FArchive& operator<<(FArchive& Ar, FRigVMExecuteOp& P)
 	{
 		P.Serialize(Ar);
@@ -175,7 +175,7 @@ struct RIGVM_API FRigVMUnaryOp : public FRigVMBaseOp
 
 	FRigVMOperand Arg;
 
-	bool Serialize(FArchive& Ar);
+	void Serialize(FArchive& Ar);
 	FORCEINLINE friend FArchive& operator<<(FArchive& Ar, FRigVMUnaryOp& P)
 	{
 		P.Serialize(Ar);
@@ -209,7 +209,7 @@ struct RIGVM_API FRigVMBinaryOp : public FRigVMBaseOp
 	FRigVMOperand ArgA;
 	FRigVMOperand ArgB;
 
-	bool Serialize(FArchive& Ar);
+	void Serialize(FArchive& Ar);
 	FORCEINLINE friend FArchive& operator<<(FArchive& Ar, FRigVMBinaryOp& P)
 	{
 		P.Serialize(Ar);
@@ -243,7 +243,7 @@ struct RIGVM_API FRigVMCopyOp : public FRigVMBaseOp
 	FRigVMOperand Source;
 	FRigVMOperand Target;
 
-	bool Serialize(FArchive& Ar);
+	void Serialize(FArchive& Ar);
 	FORCEINLINE friend FArchive& operator<<(FArchive& Ar, FRigVMCopyOp& P)
 	{
 		P.Serialize(Ar);
@@ -286,7 +286,7 @@ struct RIGVM_API FRigVMComparisonOp : public FRigVMBaseOp
 	FRigVMOperand B;
 	FRigVMOperand Result;
 
-	bool Serialize(FArchive& Ar);
+	void Serialize(FArchive& Ar);
 	FORCEINLINE friend FArchive& operator<<(FArchive& Ar, FRigVMComparisonOp& P)
 	{
 		P.Serialize(Ar);
@@ -318,7 +318,7 @@ struct RIGVM_API FRigVMJumpOp : public FRigVMBaseOp
 
 	int32 InstructionIndex;
 
-	bool Serialize(FArchive& Ar);
+	void Serialize(FArchive& Ar);
 	FORCEINLINE friend FArchive& operator<<(FArchive& Ar, FRigVMJumpOp& P)
 	{
 		P.Serialize(Ar);
@@ -353,7 +353,7 @@ struct RIGVM_API FRigVMJumpIfOp : public FRigVMUnaryOp
 	int32 InstructionIndex;
 	bool Condition;
 
-	bool Serialize(FArchive& Ar);
+	void Serialize(FArchive& Ar);
 	FORCEINLINE friend FArchive& operator<<(FArchive& Ar, FRigVMJumpIfOp& P)
 	{
 		P.Serialize(Ar);
@@ -390,7 +390,7 @@ struct RIGVM_API FRigVMChangeTypeOp : public FRigVMUnaryOp
 	uint16 ElementCount;
 	uint16 SliceCount;
 
-	bool Serialize(FArchive& Ar);
+	void Serialize(FArchive& Ar);
 	FORCEINLINE friend FArchive& operator<<(FArchive& Ar, FRigVMChangeTypeOp& P)
 	{
 		P.Serialize(Ar);
@@ -496,7 +496,9 @@ public:
 
 	FRigVMByteCode();
 
-	bool Serialize(FArchive& Ar);
+	void Serialize(FArchive& Ar);
+	void Save(FArchive& Ar) const;
+	void Load(FArchive& Ar);
 	FORCEINLINE friend FArchive& operator<<(FArchive& Ar, FRigVMByteCode& P)
 	{
 		P.Serialize(Ar);
diff --git a/Engine/Source/Runtime/RigVM/Public/RigVMCore/RigVMMemory.h b/Engine/Source/Runtime/RigVM/Public/RigVMCore/RigVMMemory.h
--- a/Engine/Source/Runtime/RigVM/Public/RigVMCore/RigVMMemory.h
+++ b/Engine/Source/Runtime/RigVM/Public/RigVMCore/RigVMMemory.h
@@ -92,7 +92,9 @@ public:
 	// returns the index of the register of this argument
 	FORCEINLINE_DEBUGGABLE int32 GetRegisterOffset() const { return RegisterOffset == UINT16_MAX ? INDEX_NONE : (int32)RegisterOffset; }
 
-	bool Serialize(FArchive& Ar);
+	void Serialize(FArchive& Ar);
+	void Save(FArchive& Ar) const;
+	void Load(FArchive& Ar);
 	FORCEINLINE friend FArchive& operator<<(FArchive& Ar, FRigVMOperand& P)
 	{
 		P.Serialize(Ar);
@@ -202,7 +204,9 @@ struct RIGVM_API FRigVMRegister
 	UPROPERTY()
 	bool bIsDynamic;
 
-	bool Serialize(FArchive& Ar);
+	void Serialize(FArchive& Ar);
+	void Save(FArchive& Ar);
+	void Load(FArchive& Ar);
 	FORCEINLINE_DEBUGGABLE friend FArchive& operator<<(FArchive& Ar, FRigVMRegister& P)
 	{
 		P.Serialize(Ar);
@@ -278,7 +282,9 @@ public:
 	{
 	}
 
-	bool Serialize(FArchive& Ar);
+	void Serialize(FArchive& Ar);
+	void Save(FArchive& Ar);
+	void Load(FArchive& Ar);
 	FORCEINLINE_DEBUGGABLE friend FArchive& operator<<(FArchive& Ar, FRigVMRegisterOffset& P)
 	{
 		P.Serialize(Ar);
@@ -574,7 +580,9 @@ public:
 	// accessor for a register based on a a name. note: only works if SupportsNames() == true
 	FORCEINLINE_DEBUGGABLE FRigVMRegister& GetRegister(const FName& InName) { return Registers[GetIndex(InName)]; }
 
-	bool Serialize(FArchive& Ar);
+	void Serialize(FArchive& Ar);
+	void Save(FArchive& Ar);
+	void Load(FArchive& Ar);
 	FORCEINLINE_DEBUGGABLE friend FArchive& operator<<(FArchive& Ar, FRigVMMemoryContainer& P)
 	{
 		P.Serialize(Ar);
