Control Rig / RigVM: Add accessor for RegisterOffsets

diff --git a/Engine/Source/Runtime/RigVM/Private/RigVMCore/RigVMMemory.cpp b/Engine/Source/Runtime/RigVM/Private/RigVMCore/RigVMMemory.cpp
--- a/Engine/Source/Runtime/RigVM/Private/RigVMCore/RigVMMemory.cpp
+++ b/Engine/Source/Runtime/RigVM/Private/RigVMCore/RigVMMemory.cpp
@@ -1808,6 +1808,19 @@ TArray<FString> FRigVMMemoryContainer::GetRegisterValueAsString(const FRigVMOper
 	return DefaultValues;
 }
 
+const FRigVMRegisterOffset FRigVMMemoryContainer::GetRegisterOffsetForOperand(const FRigVMOperand& InOperand) const
+{
+	ensure(InOperand.GetMemoryType() == MemoryType);
+
+	int32 RegisterOffsetIndex = InOperand.GetRegisterOffset();
+	if (RegisterOffsets.IsValidIndex(RegisterOffsetIndex))
+	{
+		return RegisterOffsets[RegisterOffsetIndex];
+	}
+
+	return FRigVMRegisterOffset();
+}
+
 void FRigVMMemoryContainer::UpdateRegisters()
 {
 	int32 AlignmentShift = 0;
diff --git a/Engine/Source/Runtime/RigVM/Public/RigVMCore/RigVMMemory.h b/Engine/Source/Runtime/RigVM/Public/RigVMCore/RigVMMemory.h
--- a/Engine/Source/Runtime/RigVM/Public/RigVMCore/RigVMMemory.h
+++ b/Engine/Source/Runtime/RigVM/Public/RigVMCore/RigVMMemory.h
@@ -302,6 +302,7 @@ public:
 
 	bool operator == (const FRigVMRegisterOffset& InOther) const;
 
+	FORCEINLINE_DEBUGGABLE bool IsValid() const { return Type != ERigVMRegisterType::Invalid; }
 	FORCEINLINE_DEBUGGABLE ERigVMRegisterType GetType() const { return Type; }
 	FORCEINLINE_DEBUGGABLE FName GetCPPType() const { return CPPType; }
 	uint16 GetElementSize() const;
@@ -1105,6 +1106,9 @@ public:
 	void SetRegisterValueFromString(const FRigVMOperand& InOperand, const FString& InCPPType, const UObject* InCPPTypeObject, const TArray<FString>& InDefaultValues);
 	TArray<FString> GetRegisterValueAsString(const FRigVMOperand& InOperand, const FString& InCPPType, const UObject* InCPPTypeObject);
 
+	// Returns the register offset for a given operand
+	const FRigVMRegisterOffset GetRegisterOffsetForOperand(const FRigVMOperand& InOperand) const;
+
 	// returns the statistics information
 	FRigVMMemoryStatistics GetStatistics() const
 	{
