Control Rig: Incorporate jack.cai's CL-14364519 modifications into UE5

diff --git a/Engine/Plugins/Experimental/ControlRig/Source/ControlRigEditor/Private/Editor/SControlRigStackView.cpp b/Engine/Plugins/Experimental/ControlRig/Source/ControlRigEditor/Private/Editor/SControlRigStackView.cpp
--- a/Engine/Plugins/Experimental/ControlRig/Source/ControlRigEditor/Private/Editor/SControlRigStackView.cpp
+++ b/Engine/Plugins/Experimental/ControlRig/Source/ControlRigEditor/Private/Editor/SControlRigStackView.cpp
@@ -297,33 +297,100 @@ void SControlRigStackView::HandleGetChildrenForTree(TSharedPtr<FRigStackEntry> I
 	OutChildren = InItem->Children;
 }
 
-void SControlRigStackView::RefreshTreeView(URigVM* InVM)
+void SControlRigStackView::PopulateStackView(URigVM* InVM)
 {
-	Operators.Reset();
-	
 	if (InVM)
 	{
-		TArray<FString> Labels = InVM->DumpByteCodeAsTextArray(TArray<int32>(), false);
 		FRigVMInstructionArray Instructions = InVM->GetInstructions();
-		ensure(Labels.Num() == Instructions.Num());
+		const TArray<URigVMNode*>& Nodes = ControlRigBlueprint->Model->GetNodes();
+		
+		// 1. cache information about instructions/nodes, which will be used later 
+		TArray<int32> InstructionIndexToNodeIndex;
+		InstructionIndexToNodeIndex.Init(-1, Instructions.Num());
+		TMap<FString, FString> NodeNameToDisplayName;
+		for (int32 NodeIndex = 0; NodeIndex < Nodes.Num(); NodeIndex++)
+		{
+			URigVMNode* Node = Nodes[NodeIndex];
+			// only struct nodes among all nodes has StaticExecute() that generates actual instructions
+			if (URigVMStructNode* StructNode = Cast<URigVMStructNode>(Node))
+			{
+				int32 InstructionIndex = StructNode->GetInstructionIndex();
+				// only active nodes are assigned a valid instruction index, hence the check here
+				if (InstructionIndex >= 0 && InstructionIndex < InstructionIndexToNodeIndex.Num())
+				{
+					InstructionIndexToNodeIndex[InstructionIndex] = NodeIndex;
+				}
+				FString DisplayName = StructNode->GetNodeTitle();
+#if WITH_EDITOR
+				UScriptStruct* Struct = StructNode->GetScriptStruct();
+				FString MenuDescSuffixMetadata;
+				if (Struct)
+				{
+					Struct->GetStringMetaDataHierarchical(FRigVMStruct::MenuDescSuffixMetaName, &MenuDescSuffixMetadata);
+				}
+				if (!MenuDescSuffixMetadata.IsEmpty())
+				{
+					DisplayName = FString::Printf(TEXT("%s %s"), *StructNode->GetNodeTitle(), *MenuDescSuffixMetadata);
+				}
+#endif
+				// this is needed for name replacement later
+				NodeNameToDisplayName.Add(StructNode->GetName(), DisplayName);
+			}
+		}
+
+		// 2. replace raw operand names with NodeTitle.PinName/PropertyName.OffsetName
+		TArray<FString> Labels = InVM->DumpByteCodeAsTextArray(TArray<int32>(), false, [NodeNameToDisplayName](const FString& RegisterName, const FString& RegisterOffsetName)
+		{
+			FString NewRegisterName = RegisterName;
+			FString NodeName;
+			FString PinName;
+			if (RegisterName.Split(TEXT("."), &NodeName, &PinName))
+			{
+				const FString* NodeTitle = NodeNameToDisplayName.Find(NodeName);
+				NewRegisterName = FString::Printf(TEXT("%s.%s"), NodeTitle ? **NodeTitle : *NodeName, *PinName);
+			}
+			FString OperandLabel;
+			OperandLabel = NewRegisterName;
+			if (!RegisterOffsetName.IsEmpty())
+			{
+				OperandLabel = FString::Printf(TEXT("%s.%s"), *OperandLabel, *RegisterOffsetName);
+			}
+			return OperandLabel;
+		});
 
+		ensure(Labels.Num() == Instructions.Num());
+		
+		// 3. replace instruction names with node titles
 		for (int32 InstructionIndex = 0; InstructionIndex < Labels.Num(); InstructionIndex++)
 		{
-			const FRigVMInstruction& Instruction = Instructions[InstructionIndex];
 			FString Label = Labels[InstructionIndex];
-			FString Left, Right;
-			if (Label.Split(TEXT("("), &Left, &Right))
+			int32 NodeIndex = InstructionIndexToNodeIndex[InstructionIndex];
+			// note that some instructions don't map to a node, hence the check here
+			if (NodeIndex >= 0 && NodeIndex < Nodes.Num())
 			{
-				Label = FString::Printf(TEXT("%s(...)"), *Left);
+				URigVMNode* Node = Nodes[NodeIndex];
+				Label = NodeNameToDisplayName[Node->GetName()];
 			}
-
+			// add the entry with the new label to the stack view
+			const FRigVMInstruction& Instruction = Instructions[InstructionIndex];
 			if (FilterText.IsEmpty() || Label.Contains(FilterText.ToString()))
-			{ 
+			{
 				TSharedPtr<FRigStackEntry> NewEntry = MakeShared<FRigStackEntry>(InstructionIndex, ERigStackEntry::Operator, InstructionIndex, Instruction.OpCode, Label);
 				Operators.Add(NewEntry);
 			}
 		}
+	}
+}
+
+void SControlRigStackView::RefreshTreeView(URigVM* InVM)
+{
+	Operators.Reset();
 
+	// populate the stack with node names/instruction names
+	PopulateStackView(InVM);
+	
+	if (InVM)
+	{
 		// fill the children from the log
 		if (ControlRigEditor.IsValid())
 		{
diff --git a/Engine/Plugins/Experimental/ControlRig/Source/ControlRigEditor/Private/Editor/SControlRigStackView.h b/Engine/Plugins/Experimental/ControlRig/Source/ControlRigEditor/Private/Editor/SControlRigStackView.h
--- a/Engine/Plugins/Experimental/ControlRig/Source/ControlRigEditor/Private/Editor/SControlRigStackView.h
+++ b/Engine/Plugins/Experimental/ControlRig/Source/ControlRigEditor/Private/Editor/SControlRigStackView.h
@@ -103,6 +103,9 @@ private:
 	void HandleControlRigInitializedEvent(UControlRig* InControlRig, const EControlRigState InState, const FName& InEventName);
 	void HandlePreviewControlRigUpdated(FControlRigEditor* InEditor);
 
+	/** Populate the execution stack with descriptive names for each instruction */
+	void PopulateStackView(URigVM* InVM);
+
 	TSharedPtr<STreeView<TSharedPtr<FRigStackEntry>>> TreeView;
 
 	/** Command list we bind to */
diff --git a/Engine/Source/Runtime/RigVM/Private/RigVMCore/RigVM.cpp b/Engine/Source/Runtime/RigVM/Private/RigVMCore/RigVM.cpp
--- a/Engine/Source/Runtime/RigVM/Private/RigVMCore/RigVM.cpp
+++ b/Engine/Source/Runtime/RigVM/Private/RigVMCore/RigVM.cpp
@@ -1361,7 +1361,7 @@ void URigVM::SetRegisterValueFromString(const FRigVMOperand& InOperand, const FS
 
 #if WITH_EDITOR
 
-TArray<FString> URigVM::DumpByteCodeAsTextArray(const TArray<int32>& InInstructionOrder, bool bIncludeLineNumbers)
+TArray<FString> URigVM::DumpByteCodeAsTextArray(const TArray<int32>& InInstructionOrder, bool bIncludeLineNumbers, TFunction<FString(const FString& RegisterName, const FString& RegisterOffsetName)> OperandFormatFunction)
 {
 	RefreshInstructionsIfRequired();
 	const FRigVMByteCode& ByteCode = GetByteCode();
@@ -1458,7 +1458,7 @@ TArray<FString> URigVM::DumpByteCodeAsTextArray(const TArray<int32>& InInstructi
 				TArray<FString> Labels;
 				for (const FRigVMOperand& Operand : Operands)
 				{
-					Labels.Add(GetOperandLabel(Operand));
+					Labels.Add(GetOperandLabel(Operand, OperandFormatFunction));
 				}
 
 				ResultLine = FString::Printf(TEXT("%s(%s)"), *FunctionName, *FString::Join(Labels, TEXT(",")));
@@ -1467,49 +1467,49 @@ TArray<FString> URigVM::DumpByteCodeAsTextArray(const TArray<int32>& InInstructi
 			case ERigVMOpCode::Zero:
 			{
 				const FRigVMUnaryOp& Op = ByteCode.GetOpAt<FRigVMUnaryOp>(Instructions[InstructionIndex]);
-				ResultLine = FString::Printf(TEXT("Set %s to 0"), *GetOperandLabel(Op.Arg));
+				ResultLine = FString::Printf(TEXT("Set %s to 0"), *GetOperandLabel(Op.Arg, OperandFormatFunction));
 				break;
 			}
 			case ERigVMOpCode::BoolFalse:
 			{
 				const FRigVMUnaryOp& Op = ByteCode.GetOpAt<FRigVMUnaryOp>(Instructions[InstructionIndex]);
-				ResultLine = FString::Printf(TEXT("Set %s to False"), *GetOperandLabel(Op.Arg));
+				ResultLine = FString::Printf(TEXT("Set %s to False"), *GetOperandLabel(Op.Arg, OperandFormatFunction));
 				break;
 			}
 			case ERigVMOpCode::BoolTrue:
 			{
 				const FRigVMUnaryOp& Op = ByteCode.GetOpAt<FRigVMUnaryOp>(Instructions[InstructionIndex]);
-				ResultLine = FString::Printf(TEXT("Set %s to True"), *GetOperandLabel(Op.Arg));
+				ResultLine = FString::Printf(TEXT("Set %s to True"), *GetOperandLabel(Op.Arg, OperandFormatFunction));
 				break;
 			}
 			case ERigVMOpCode::Increment:
 			{
 				const FRigVMUnaryOp& Op = ByteCode.GetOpAt<FRigVMUnaryOp>(Instructions[InstructionIndex]);
-				ResultLine = FString::Printf(TEXT("Inc %s ++"), *GetOperandLabel(Op.Arg));
+				ResultLine = FString::Printf(TEXT("Inc %s ++"), *GetOperandLabel(Op.Arg, OperandFormatFunction));
 				break;
 			}
 			case ERigVMOpCode::Decrement:
 			{
 				const FRigVMUnaryOp& Op = ByteCode.GetOpAt<FRigVMUnaryOp>(Instructions[InstructionIndex]);
-				ResultLine = FString::Printf(TEXT("Dec %s --"), *GetOperandLabel(Op.Arg));
+				ResultLine = FString::Printf(TEXT("Dec %s --"), *GetOperandLabel(Op.Arg, OperandFormatFunction));
 				break;
 			}
 			case ERigVMOpCode::Copy:
 			{
 				const FRigVMCopyOp& Op = ByteCode.GetOpAt<FRigVMCopyOp>(Instructions[InstructionIndex]);
-				ResultLine = FString::Printf(TEXT("Copy %s to %s"), *GetOperandLabel(Op.Source), *GetOperandLabel(Op.Target));
+				ResultLine = FString::Printf(TEXT("Copy %s to %s"), *GetOperandLabel(Op.Source, OperandFormatFunction), *GetOperandLabel(Op.Target, OperandFormatFunction));
 				break;
 			}
 			case ERigVMOpCode::Equals:
 			{
 				const FRigVMComparisonOp& Op = ByteCode.GetOpAt<FRigVMComparisonOp>(Instructions[InstructionIndex]);
-				ResultLine = FString::Printf(TEXT("Set %s to %s == %s "), *GetOperandLabel(Op.Result), *GetOperandLabel(Op.A), *GetOperandLabel(Op.B));
+				ResultLine = FString::Printf(TEXT("Set %s to %s == %s "), *GetOperandLabel(Op.Result, OperandFormatFunction), *GetOperandLabel(Op.A, OperandFormatFunction), *GetOperandLabel(Op.B, OperandFormatFunction));
 				break;
 			}
 			case ERigVMOpCode::NotEquals:
 			{
 				const FRigVMComparisonOp& Op = ByteCode.GetOpAt<FRigVMComparisonOp>(Instructions[InstructionIndex]);
-				ResultLine = FString::Printf(TEXT("Set %s to %s != %s"), *GetOperandLabel(Op.Result), *GetOperandLabel(Op.A), *GetOperandLabel(Op.B));
+				ResultLine = FString::Printf(TEXT("Set %s to %s != %s"), *GetOperandLabel(Op.Result, OperandFormatFunction), *GetOperandLabel(Op.A, OperandFormatFunction), *GetOperandLabel(Op.B, OperandFormatFunction));
 				break;
 			}
 			case ERigVMOpCode::JumpAbsolute:
@@ -1535,11 +1535,11 @@ TArray<FString> URigVM::DumpByteCodeAsTextArray(const TArray<int32>& InInstructi
 				const FRigVMJumpIfOp& Op = ByteCode.GetOpAt<FRigVMJumpIfOp>(Instructions[InstructionIndex]);
 				if (Op.Condition)
 				{
-					ResultLine = FString::Printf(TEXT("Jump to instruction %d if %s"), Op.InstructionIndex, *GetOperandLabel(Op.Arg));
+					ResultLine = FString::Printf(TEXT("Jump to instruction %d if %s"), Op.InstructionIndex, *GetOperandLabel(Op.Arg, OperandFormatFunction));
 				}
 				else
 				{
-					ResultLine = FString::Printf(TEXT("Jump to instruction %d if !%s"), Op.InstructionIndex, *GetOperandLabel(Op.Arg));
+					ResultLine = FString::Printf(TEXT("Jump to instruction %d if !%s"), Op.InstructionIndex, *GetOperandLabel(Op.Arg, OperandFormatFunction));
 				}
 				break;
 			}
@@ -1548,11 +1548,11 @@ TArray<FString> URigVM::DumpByteCodeAsTextArray(const TArray<int32>& InInstructi
 				const FRigVMJumpIfOp& Op = ByteCode.GetOpAt<FRigVMJumpIfOp>(Instructions[InstructionIndex]);
 				if (Op.Condition)
 				{
-					ResultLine = FString::Printf(TEXT("Jump %d instructions forwards if %s"), Op.InstructionIndex, *GetOperandLabel(Op.Arg));
+					ResultLine = FString::Printf(TEXT("Jump %d instructions forwards if %s"), Op.InstructionIndex, *GetOperandLabel(Op.Arg, OperandFormatFunction));
 				}
 				else
 				{
-					ResultLine = FString::Printf(TEXT("Jump %d instructions forwards if !%s"), Op.InstructionIndex, *GetOperandLabel(Op.Arg));
+					ResultLine = FString::Printf(TEXT("Jump %d instructions forwards if !%s"), Op.InstructionIndex, *GetOperandLabel(Op.Arg, OperandFormatFunction));
 				}
 				break;
 			}
@@ -1561,18 +1561,18 @@ TArray<FString> URigVM::DumpByteCodeAsTextArray(const TArray<int32>& InInstructi
 				const FRigVMJumpIfOp& Op = ByteCode.GetOpAt<FRigVMJumpIfOp>(Instructions[InstructionIndex]);
 				if (Op.Condition)
 				{
-					ResultLine = FString::Printf(TEXT("Jump %d instructions backwards if %s"), Op.InstructionIndex, *GetOperandLabel(Op.Arg));
+					ResultLine = FString::Printf(TEXT("Jump %d instructions backwards if %s"), Op.InstructionIndex, *GetOperandLabel(Op.Arg, OperandFormatFunction));
 				}
 				else
 				{
-					ResultLine = FString::Printf(TEXT("Jump %d instructions backwards if !%s"), Op.InstructionIndex, *GetOperandLabel(Op.Arg));
+					ResultLine = FString::Printf(TEXT("Jump %d instructions backwards if !%s"), Op.InstructionIndex, *GetOperandLabel(Op.Arg, OperandFormatFunction));
 				}
 				break;
 			}
 			case ERigVMOpCode::ChangeType:
 			{
 				const FRigVMChangeTypeOp& Op = ByteCode.GetOpAt<FRigVMChangeTypeOp>(Instructions[InstructionIndex]);
-				ResultLine = FString::Printf(TEXT("Change type of %s"), *GetOperandLabel(Op.Arg));
+				ResultLine = FString::Printf(TEXT("Change type of %s"), *GetOperandLabel(Op.Arg, OperandFormatFunction));
 				break;
 			}
 			case ERigVMOpCode::Exit:
@@ -1620,7 +1620,7 @@ FString URigVM::DumpByteCodeAsText(const TArray<int32>& InInstructionOrder, bool
 	return FString::Join(DumpByteCodeAsTextArray(InInstructionOrder, bIncludeLineNumbers), TEXT("\n"));
 }
 
-FString URigVM::GetOperandLabel(const FRigVMOperand& InOperand) const
+FString URigVM::GetOperandLabel(const FRigVMOperand& InOperand, TFunction<FString(const FString& RegisterName, const FString& RegisterOffsetName)> FormatFunction) const
 {
 	const FRigVMMemoryContainer* MemoryPtr = nullptr;
 
@@ -1635,21 +1635,33 @@ FString URigVM::GetOperandLabel(const FRigVMOperand& InOperand) const
 
 	const FRigVMMemoryContainer& Memory = *MemoryPtr;
 
-	FString OperandLabel;
+	FString RegisterName;
 	if (InOperand.GetMemoryType() == ERigVMMemoryType::External)
 	{
 		const FRigVMExternalVariable& ExternalVariable = ExternalVariables[InOperand.GetRegisterIndex()];
-		OperandLabel = FString::Printf(TEXT("Variable::%s"), *ExternalVariable.Name.ToString());
+		RegisterName = FString::Printf(TEXT("Variable::%s"), *ExternalVariable.Name.ToString());
 	}
 	else
 	{
 		FRigVMRegister Register = Memory[InOperand];
-		OperandLabel = Register.Name.ToString();
+		RegisterName = Register.Name.ToString();
 	}
 
+	FString OperandLabel;
+	OperandLabel = RegisterName;
+	
+	// append an offset name if it exists
+	FString RegisterOffsetName;
 	if (InOperand.GetRegisterOffset() != INDEX_NONE)
 	{
-		return FString::Printf(TEXT("%s.%s"), *OperandLabel, *Memory.RegisterOffsets[InOperand.GetRegisterOffset()].CachedSegmentPath);
+		RegisterOffsetName = Memory.RegisterOffsets[InOperand.GetRegisterOffset()].CachedSegmentPath;
+		OperandLabel = FString::Printf(TEXT("%s.%s"), *OperandLabel, *RegisterOffsetName);
+	}
+
+	// caller can provide an alternative format to override the default format(optional)
+	if (FormatFunction)
+	{
+		OperandLabel = FormatFunction(RegisterName, RegisterOffsetName);
 	}
 
 	return OperandLabel;
diff --git a/Engine/Source/Runtime/RigVM/Public/RigVMCore/RigVM.h b/Engine/Source/Runtime/RigVM/Public/RigVMCore/RigVM.h
--- a/Engine/Source/Runtime/RigVM/Public/RigVMCore/RigVM.h
+++ b/Engine/Source/Runtime/RigVM/Public/RigVMCore/RigVM.h
@@ -482,15 +482,14 @@ public:
 
 
 #if WITH_EDITOR
-
-	// returns the instructions as text
-	TArray<FString> DumpByteCodeAsTextArray(const TArray<int32>& InInstructionOrder = TArray<int32>(), bool bIncludeLineNumbers = true);
+	// returns the instructions as text, OperandFormatFunction is an optional argument that allows you to override how operands are displayed, for example, see SControlRigStackView::PopulateStackView 
+	TArray<FString> DumpByteCodeAsTextArray(const TArray<int32> & InInstructionOrder = TArray<int32>(), bool bIncludeLineNumbers = true, TFunction<FString(const FString& RegisterName, const FString& RegisterOffsetName)> OperandFormatFunction = nullptr);
 	FString DumpByteCodeAsText(const TArray<int32>& InInstructionOrder = TArray<int32>(), bool bIncludeLineNumbers = true);
-
 #endif
 
 #if WITH_EDITOR
-	FString GetOperandLabel(const FRigVMOperand& InOperand) const;
+	// FormatFunction is an optional argument that allows you to override how operands are displayed, for example, see SControlRigStackView::PopulateStackView
+	FString GetOperandLabel(const FRigVMOperand & InOperand, TFunction<FString(const FString& RegisterName, const FString& RegisterOffsetName)> FormatFunction = nullptr) const;
 #endif
 
 private:
