RigVM: Add header file path registration for RigVM Functions, required for Control Rig Nativization Research Project

diff --git a/Engine/Source/Programs/UnrealHeaderTool/Private/CodeGenerator.cpp b/Engine/Source/Programs/UnrealHeaderTool/Private/CodeGenerator.cpp
--- a/Engine/Source/Programs/UnrealHeaderTool/Private/CodeGenerator.cpp
+++ b/Engine/Source/Programs/UnrealHeaderTool/Private/CodeGenerator.cpp
@@ -3872,10 +3872,16 @@ void FNativeClassHeaderGenerator::ExportGeneratedStructBodyMacros(FOutputDevice&
 		// registry on construction of the static struct
 		if (StructRigVMInfo)
 		{
+			// All RigVM methods are in the Units folder
+			FString HeaderFilePath = SourceFile.GetFilename();
+			check(HeaderFilePath.Find(TEXT("Units")) > 0);
+			HeaderFilePath = HeaderFilePath.RightChop(HeaderFilePath.Find(TEXT("Units"))); 
+			HeaderFilePath.ReplaceInline(TEXT("\\"), TEXT("/"));
+
 			for (const FRigVMMethodInfo& MethodInfo : StructRigVMInfo->Methods)
 			{
-				Out.Logf(TEXT("\t\tFRigVMRegistry::Get().Register(TEXT(\"%s::%s\"), &%s::RigVM%s, Singleton);\r\n"),
-					*StructNameCPP, *MethodInfo.Name, *StructNameCPP, *MethodInfo.Name);
+				Out.Logf(TEXT("\t\tFRigVMRegistry::Get().Register(TEXT(\"%s::%s\"), &%s::RigVM%s, Singleton, TEXT(\"%s\"));\r\n"),
+					*StructNameCPP, *MethodInfo.Name, *StructNameCPP, *MethodInfo.Name, *HeaderFilePath);
 			}
 		}
 
diff --git a/Engine/Source/Runtime/RigVM/Private/RigVMCore/RigVMRegistry.cpp b/Engine/Source/Runtime/RigVM/Private/RigVMCore/RigVMRegistry.cpp
--- a/Engine/Source/Runtime/RigVM/Private/RigVMCore/RigVMRegistry.cpp
+++ b/Engine/Source/Runtime/RigVM/Private/RigVMCore/RigVMRegistry.cpp
@@ -16,14 +16,14 @@ void FRigVMRegistry::Refresh()
 {
 }
 
-void FRigVMRegistry::Register(const TCHAR* InName, FRigVMFunctionPtr InFunctionPtr, UScriptStruct* InStruct)
+void FRigVMRegistry::Register(const TCHAR* InName, FRigVMFunctionPtr InFunctionPtr, UScriptStruct* InStruct, FString InHeaderFilePath)
 {
 	if (FindFunction(InName) != nullptr)
 	{
 		return;
 	}
 
-	FRigVMFunction Function(InName, InFunctionPtr, InStruct, Functions.Num());
+	FRigVMFunction Function(InName, InFunctionPtr, InStruct, Functions.Num(), InHeaderFilePath);
 	Functions.Add(Function);
 
 #if WITH_EDITOR
@@ -104,6 +104,19 @@ const FRigVMPrototype* FRigVMRegistry::FindPrototype(UScriptStruct* InStruct, co
 	return FindPrototype(Notation);
 }
 
+FRigVMFunction FRigVMRegistry::FindFunctionInfo(const TCHAR* InName) const
+{
+	for (const FRigVMFunction& Function : Functions)
+	{
+		if (FCString::Strcmp(Function.Name, InName) == 0)
+		{
+			return Function;
+		}
+	} 
+
+	return FRigVMFunction();
+}
+
 const TArray<FRigVMFunction>& FRigVMRegistry::GetFunctions() const
 {
 	return Functions;
diff --git a/Engine/Source/Runtime/RigVM/Public/RigVMCore/RigVMFunction.h b/Engine/Source/Runtime/RigVM/Public/RigVMCore/RigVMFunction.h
--- a/Engine/Source/Runtime/RigVM/Public/RigVMCore/RigVMFunction.h
+++ b/Engine/Source/Runtime/RigVM/Public/RigVMCore/RigVMFunction.h
@@ -38,6 +38,7 @@ struct RIGVM_API FRigVMFunction
 	FRigVMFunctionPtr FunctionPtr;
 	int32 Index;
 	int32 PrototypeIndex;
+	FString HeaderFilePath;
 
 	FRigVMFunction()
 		: Name(nullptr)
@@ -45,15 +46,17 @@ struct RIGVM_API FRigVMFunction
 		, FunctionPtr(nullptr)
 		, Index(INDEX_NONE)
 		, PrototypeIndex(INDEX_NONE)
+		, HeaderFilePath()
 	{
 	}
 
-	FRigVMFunction(const TCHAR* InName, FRigVMFunctionPtr InFunctionPtr, UScriptStruct* InStruct = nullptr, int32 InIndex = INDEX_NONE)
+	FRigVMFunction(const TCHAR* InName, FRigVMFunctionPtr InFunctionPtr, UScriptStruct* InStruct = nullptr, int32 InIndex = INDEX_NONE, FString InHeaderFilePath = FString())
 		: Name(InName)
 		, Struct(InStruct)
 		, FunctionPtr(InFunctionPtr)
 		, Index(InIndex)
 		, PrototypeIndex(INDEX_NONE)
+		, HeaderFilePath(InHeaderFilePath)
 	{
 	}
 
diff --git a/Engine/Source/Runtime/RigVM/Public/RigVMCore/RigVMRegistry.h b/Engine/Source/Runtime/RigVM/Public/RigVMCore/RigVMRegistry.h
--- a/Engine/Source/Runtime/RigVM/Public/RigVMCore/RigVMRegistry.h
+++ b/Engine/Source/Runtime/RigVM/Public/RigVMCore/RigVMRegistry.h
@@ -23,7 +23,7 @@ public:
 	// Registers a function given its name.
 	// The name will be the name of the struct and virtual method,
 	// for example "FMyStruct::MyVirtualMethod"
-	void Register(const TCHAR* InName, FRigVMFunctionPtr InFunctionPtr, UScriptStruct* InStruct = nullptr);
+	void Register(const TCHAR* InName, FRigVMFunctionPtr InFunctionPtr, UScriptStruct* InStruct = nullptr, FString InHeaderFilePath = FString());
 
 	// Refreshes the list and finds the function pointers
 	// based on the names.
@@ -38,6 +38,9 @@ public:
 	// Returns a prototype pointer given its notation (or nullptr)
 	const FRigVMPrototype* FindPrototype(UScriptStruct* InStruct, const FString& InPrototypeName) const;
 
+	// Returns registry info about a function given its name
+	FRigVMFunction FindFunctionInfo(const TCHAR* InName) const;
+
 	// Returns all current rigvm functions
 	const TArray<FRigVMFunction>& GetFunctions() const;
 
