RigVM: Created accessor for RegisterOffset.CachedSegmentPath for Control Rig Nativization Reasearch Project

diff --git a/Engine/Source/Runtime/RigVM/Public/RigVMCore/RigVMMemory.h b/Engine/Source/Runtime/RigVM/Public/RigVMCore/RigVMMemory.h
--- a/Engine/Source/Runtime/RigVM/Public/RigVMCore/RigVMMemory.h
+++ b/Engine/Source/Runtime/RigVM/Public/RigVMCore/RigVMMemory.h
@@ -305,6 +305,7 @@ public:
 	FORCEINLINE_DEBUGGABLE bool IsValid() const { return Type != ERigVMRegisterType::Invalid; }
 	FORCEINLINE_DEBUGGABLE ERigVMRegisterType GetType() const { return Type; }
 	FORCEINLINE_DEBUGGABLE FName GetCPPType() const { return CPPType; }
+	FORCEINLINE_DEBUGGABLE FString GetCachedSegmentPath() const { return CachedSegmentPath; }
 	uint16 GetElementSize() const;
 	UScriptStruct* GetScriptStruct() const;
 
