RigVM: Registry: Access to module name and relative file path of the header for a RigVMFunction + unit test

diff --git a/Engine/Source/Runtime/RigVM/Private/RigVMCore/RigVMFunction.cpp b/Engine/Source/Runtime/RigVM/Private/RigVMCore/RigVMFunction.cpp
--- a/Engine/Source/Runtime/RigVM/Private/RigVMCore/RigVMFunction.cpp
+++ b/Engine/Source/Runtime/RigVM/Private/RigVMCore/RigVMFunction.cpp
@@ -12,3 +12,29 @@ FName FRigVMFunction::GetMethodName() const
 	}
 	return NAME_None;
 }
+
+FString FRigVMFunction::GetModuleName() const
+{
+	if (Struct)
+	{
+		if (UPackage* Package = Struct->GetPackage())
+		{
+			return Package->GetName();
+		}
+	}
+	return FString();
+}
+
+FString FRigVMFunction::GetModuleRelativeHeaderPath() const
+{
+	if (Struct)
+	{
+		FString ModuleRelativePath;
+		if (Struct->GetStringMetaDataHierarchical(TEXT("ModuleRelativePath"), &ModuleRelativePath))
+		{
+			return ModuleRelativePath;
+		}
+	}
+	return FString();
+}
+
diff --git a/Engine/Source/Runtime/RigVM/Private/RigVMCore/RigVMRegistry.cpp b/Engine/Source/Runtime/RigVM/Private/RigVMCore/RigVMRegistry.cpp
--- a/Engine/Source/Runtime/RigVM/Private/RigVMCore/RigVMRegistry.cpp
+++ b/Engine/Source/Runtime/RigVM/Private/RigVMCore/RigVMRegistry.cpp
@@ -16,14 +16,14 @@ void FRigVMRegistry::Refresh()
 {
 }
 
-void FRigVMRegistry::Register(const TCHAR* InName, FRigVMFunctionPtr InFunctionPtr, UScriptStruct* InStruct, FString InHeaderFilePath)
+void FRigVMRegistry::Register(const TCHAR* InName, FRigVMFunctionPtr InFunctionPtr, UScriptStruct* InStruct)
 {
 	if (FindFunction(InName) != nullptr)
 	{
 		return;
 	}
 
-	FRigVMFunction Function(InName, InFunctionPtr, InStruct, Functions.Num(), InHeaderFilePath);
+	FRigVMFunction Function(InName, InFunctionPtr, InStruct, Functions.Num());
 	Functions.Add(Function);
 
 #if WITH_EDITOR
diff --git a/Engine/Source/Runtime/RigVM/Public/RigVMCore/RigVMFunction.h b/Engine/Source/Runtime/RigVM/Public/RigVMCore/RigVMFunction.h
--- a/Engine/Source/Runtime/RigVM/Public/RigVMCore/RigVMFunction.h
+++ b/Engine/Source/Runtime/RigVM/Public/RigVMCore/RigVMFunction.h
@@ -38,7 +38,6 @@ struct RIGVM_API FRigVMFunction
 	FRigVMFunctionPtr FunctionPtr;
 	int32 Index;
 	int32 PrototypeIndex;
-	FString HeaderFilePath;
 
 	FRigVMFunction()
 		: Name(nullptr)
@@ -46,19 +45,20 @@ struct RIGVM_API FRigVMFunction
 		, FunctionPtr(nullptr)
 		, Index(INDEX_NONE)
 		, PrototypeIndex(INDEX_NONE)
-		, HeaderFilePath()
 	{
 	}
 
-	FRigVMFunction(const TCHAR* InName, FRigVMFunctionPtr InFunctionPtr, UScriptStruct* InStruct = nullptr, int32 InIndex = INDEX_NONE, FString InHeaderFilePath = FString())
+	FRigVMFunction(const TCHAR* InName, FRigVMFunctionPtr InFunctionPtr, UScriptStruct* InStruct = nullptr, int32 InIndex = INDEX_NONE)
 		: Name(InName)
 		, Struct(InStruct)
 		, FunctionPtr(InFunctionPtr)
 		, Index(InIndex)
 		, PrototypeIndex(INDEX_NONE)
-		, HeaderFilePath(InHeaderFilePath)
 	{
 	}
 
+	FORCEINLINE bool IsValid() const { return Name != nullptr && FunctionPtr != nullptr; }
 	FName GetMethodName() const;
+	FString GetModuleName() const;
+	FString GetModuleRelativeHeaderPath() const;
 };
diff --git a/Engine/Source/Runtime/RigVM/Public/RigVMCore/RigVMRegistry.h b/Engine/Source/Runtime/RigVM/Public/RigVMCore/RigVMRegistry.h
--- a/Engine/Source/Runtime/RigVM/Public/RigVMCore/RigVMRegistry.h
+++ b/Engine/Source/Runtime/RigVM/Public/RigVMCore/RigVMRegistry.h
@@ -23,7 +23,7 @@ public:
 	// Registers a function given its name.
 	// The name will be the name of the struct and virtual method,
 	// for example "FMyStruct::MyVirtualMethod"
-	void Register(const TCHAR* InName, FRigVMFunctionPtr InFunctionPtr, UScriptStruct* InStruct = nullptr, FString InHeaderFilePath = FString());
+	void Register(const TCHAR* InName, FRigVMFunctionPtr InFunctionPtr, UScriptStruct* InStruct = nullptr);
 
 	// Refreshes the list and finds the function pointers
 	// based on the names.
