RigVMFunction: Fix build break for non-editor

diff --git a/Engine/Source/Runtime/RigVM/Private/RigVMCore/RigVMFunction.cpp b/Engine/Source/Runtime/RigVM/Private/RigVMCore/RigVMFunction.cpp
--- a/Engine/Source/Runtime/RigVM/Private/RigVMCore/RigVMFunction.cpp
+++ b/Engine/Source/Runtime/RigVM/Private/RigVMCore/RigVMFunction.cpp
@@ -15,6 +15,7 @@ FName FRigVMFunction::GetMethodName() const
 
 FString FRigVMFunction::GetModuleName() const
 {
+#if WITH_EDITOR
 	if (Struct)
 	{
 		if (UPackage* Package = Struct->GetPackage())
@@ -22,11 +23,13 @@ FString FRigVMFunction::GetModuleName() const
 			return Package->GetName();
 		}
 	}
+#endif
 	return FString();
 }
 
 FString FRigVMFunction::GetModuleRelativeHeaderPath() const
 {
+#if WITH_EDITOR
 	if (Struct)
 	{
 		FString ModuleRelativePath;
@@ -35,6 +38,7 @@ FString FRigVMFunction::GetModuleRelativeHeaderPath() const
 			return ModuleRelativePath;
 		}
 	}
+#endif
 	return FString();
 }
 
