RigVM: Add GetArrayIndex() to RegisterOffset to get the array index of the operand
, also minor bug/typo fixes

diff --git a/Engine/Source/Runtime/RigVM/Private/RigVMCore/RigVMMemory.cpp b/Engine/Source/Runtime/RigVM/Private/RigVMCore/RigVMMemory.cpp
--- a/Engine/Source/Runtime/RigVM/Private/RigVMCore/RigVMMemory.cpp
+++ b/Engine/Source/Runtime/RigVM/Private/RigVMCore/RigVMMemory.cpp
@@ -1727,7 +1727,7 @@ TArray<FString> FRigVMMemoryContainer::GetRegisterValueAsString(const FRigVMOper
 
 	FRigVMRegister& Register = Registers[InOperand.GetRegisterIndex()];
 
-	int32 SliceCount = Register.ElementCount;
+	int32 SliceCount = Register.SliceCount;
 	
 	if (Register.IsNestedDynamic())
 	{
@@ -1771,7 +1771,7 @@ TArray<FString> FRigVMMemoryContainer::GetRegisterValueAsString(const FRigVMOper
 			}
 			else if (const UEnum* Enum = Cast<const UEnum>(InCPPTypeObject))
 			{
-				DefaultValue = Enum->GetNameStringByValue((int64)GetFixedArray<int32>(Register, INDEX_NONE, SliceIndex)[ElementIndex]);
+				DefaultValue = Enum->GetNameStringByValue((int64)GetFixedArray<uint8>(Register, INDEX_NONE, SliceIndex)[ElementIndex]);
 			}
 			else if (InCPPType == TEXT("bool") && Register.GetNumBytesPerSlice() == sizeof(bool))
 			{
diff --git a/Engine/Source/Runtime/RigVM/Public/RigVMCore/RigVMMemory.h b/Engine/Source/Runtime/RigVM/Public/RigVMCore/RigVMMemory.h
--- a/Engine/Source/Runtime/RigVM/Public/RigVMCore/RigVMMemory.h
+++ b/Engine/Source/Runtime/RigVM/Public/RigVMCore/RigVMMemory.h
@@ -308,6 +308,7 @@ public:
 	FORCEINLINE_DEBUGGABLE ERigVMRegisterType GetType() const { return Type; }
 	FORCEINLINE_DEBUGGABLE FName GetCPPType() const { return CPPType; }
 	FORCEINLINE_DEBUGGABLE FString GetCachedSegmentPath() const { return CachedSegmentPath; }
+	FORCEINLINE_DEBUGGABLE int32 GetArrayIndex() const { return ArrayIndex; }
 	uint16 GetElementSize() const;
 	UScriptStruct* GetScriptStruct() const;
 
