RigVM: Relay walking of RegisterOffsets for arrays to execute, rather than caching phase.

diff --git a/Engine/Source/Runtime/RigVM/Private/RigVMCore/RigVM.cpp b/Engine/Source/Runtime/RigVM/Private/RigVMCore/RigVM.cpp
--- a/Engine/Source/Runtime/RigVM/Private/RigVMCore/RigVM.cpp
+++ b/Engine/Source/Runtime/RigVM/Private/RigVMCore/RigVM.cpp
@@ -1699,21 +1699,15 @@ void URigVM::CacheSingleMemoryHandle(const FRigVMOperand& InArg, bool bForExecut
 		ensure(ExternalVariables.IsValidIndex(InArg.GetRegisterIndex()));
 
 		FRigVMExternalVariable& ExternalVariable = ExternalVariables[InArg.GetRegisterIndex()];
-		FRigVMRegisterOffset RegisterOffset = GetWorkMemory().GetRegisterOffsetForOperand(InArg);
+		const FRigVMRegisterOffset& RegisterOffset = GetWorkMemory().GetRegisterOffsetForOperand(InArg);
 		check(ExternalVariable.Memory);
 
+		FRigVMMemoryHandle Handle = ExternalVariable.GetHandle();
 		if (RegisterOffset.IsValid())
 		{
-			CachedMemoryHandles.Add(FRigVMMemoryHandle(
-				RegisterOffset.GetData(ExternalVariable.Memory),
-				RegisterOffset.GetElementSize(),
-				FRigVMMemoryHandle::FType::Plain
-			));
-		}
-		else
-		{
-			CachedMemoryHandles.Add(ExternalVariable.GetHandle(InArg.GetRegisterOffset()));
+			Handle.RegisterOffset = &RegisterOffset;
 		}
+		CachedMemoryHandles.Add(Handle);
 		return;
 	}
 
diff --git a/Engine/Source/Runtime/RigVM/Private/RigVMCore/RigVMMemory.cpp b/Engine/Source/Runtime/RigVM/Private/RigVMCore/RigVMMemory.cpp
--- a/Engine/Source/Runtime/RigVM/Private/RigVMCore/RigVMMemory.cpp
+++ b/Engine/Source/Runtime/RigVM/Private/RigVMCore/RigVMMemory.cpp
@@ -364,6 +364,18 @@ uint8* FRigVMRegisterOffset::GetData(uint8* InContainer) const
 	return Data;
 }
 
+bool FRigVMRegisterOffset::ContainsArraySegment() const
+{
+	for (int32 SegmentIndex : Segments)
+	{
+		if (SegmentIndex < 0)
+		{
+			return true;
+		}
+	}
+	return false;
+}
+
 bool FRigVMRegisterOffset::operator == (const FRigVMRegisterOffset& InOther) const
 {
 	if (Segments.Num() != InOther.Segments.Num())
@@ -430,6 +442,7 @@ public:
 ////////////////////////////////////////////////////////////////////////////////////////////////////////////////////////
 
 FRigVMByteArray FRigVMMemoryContainer::DefaultByteArray;
+FRigVMRegisterOffset FRigVMMemoryContainer::InvalidRegisterOffset;
 
 FRigVMMemoryContainer::FRigVMMemoryContainer(bool bInUseNames)
 	: bUseNameMap(bInUseNames)
@@ -1806,18 +1819,22 @@ TArray<FString> FRigVMMemoryContainer::GetRegisterValueAsString(const FRigVMOper
 	return DefaultValues;
 }
 
-const FRigVMRegisterOffset FRigVMMemoryContainer::GetRegisterOffsetForOperand(const FRigVMOperand& InOperand) const
+const FRigVMRegisterOffset& FRigVMMemoryContainer::GetRegisterOffset(int32 InRegisterOffsetIndex) const
 {
-	ensure(InOperand.GetMemoryType() == MemoryType || 
-		(InOperand.GetMemoryType() == ERigVMMemoryType::External && MemoryType == ERigVMMemoryType::Work));
-
-	int32 RegisterOffsetIndex = InOperand.GetRegisterOffset();
-	if (RegisterOffsets.IsValidIndex(RegisterOffsetIndex))
+	if (RegisterOffsets.IsValidIndex(InRegisterOffsetIndex))
 	{
-		return RegisterOffsets[RegisterOffsetIndex];
+		return RegisterOffsets[InRegisterOffsetIndex];
 	}
 
-	return FRigVMRegisterOffset();
+	return InvalidRegisterOffset;
+}
+
+const FRigVMRegisterOffset& FRigVMMemoryContainer::GetRegisterOffsetForOperand(const FRigVMOperand& InOperand) const
+{
+	ensure(InOperand.GetMemoryType() == MemoryType || 
+		(InOperand.GetMemoryType() == ERigVMMemoryType::External && MemoryType == ERigVMMemoryType::Work));
+	
+	return GetRegisterOffset(InOperand.GetRegisterOffset());
 }
 
 void FRigVMMemoryContainer::UpdateRegisters()
diff --git a/Engine/Source/Runtime/RigVM/Public/RigVMCore/RigVMExternalVariable.h b/Engine/Source/Runtime/RigVM/Public/RigVMCore/RigVMExternalVariable.h
--- a/Engine/Source/Runtime/RigVM/Public/RigVMCore/RigVMExternalVariable.h
+++ b/Engine/Source/Runtime/RigVM/Public/RigVMCore/RigVMExternalVariable.h
@@ -408,13 +408,12 @@ struct RIGVM_API FRigVMExternalVariable
 			(bAllowNullPtr || Memory != nullptr);
 	}
 
-	FORCEINLINE FRigVMMemoryHandle GetHandle(int32 InOffset = INDEX_NONE) const
+	FORCEINLINE FRigVMMemoryHandle GetHandle() const
 	{
 		return FRigVMMemoryHandle(
 			Memory,
 			Size,
-			bIsArray ? FRigVMMemoryHandle::FType::Dynamic : FRigVMMemoryHandle::FType::Plain,
-			InOffset
+			bIsArray ? FRigVMMemoryHandle::FType::Dynamic : FRigVMMemoryHandle::FType::Plain
 		);
 	}
 
diff --git a/Engine/Source/Runtime/RigVM/Public/RigVMCore/RigVMMemory.h b/Engine/Source/Runtime/RigVM/Public/RigVMCore/RigVMMemory.h
--- a/Engine/Source/Runtime/RigVM/Public/RigVMCore/RigVMMemory.h
+++ b/Engine/Source/Runtime/RigVM/Public/RigVMCore/RigVMMemory.h
@@ -317,6 +317,11 @@ public:
 	// returns the segments of this path
 	const TArray<int32>& GetSegments() const { return Segments; }
 
+	// returns true if this register offset contains an array segment.
+	// this means the register offset has to traverse a pointer and
+	// cannot be cached / collapsed easily
+	bool ContainsArraySegment() const;
+
 	bool operator == (const FRigVMRegisterOffset& InOther) const;
 
 	FORCEINLINE_DEBUGGABLE bool IsValid() const { return Type != ERigVMRegisterType::Invalid; }
@@ -374,21 +379,21 @@ public:
 		: Ptr(nullptr)
 		, Type(FType::Plain)
 		, Size(1)
-		, Offset(UINT16_MAX)
+		, RegisterOffset(nullptr)
 	{}
 
-	FORCEINLINE_DEBUGGABLE FRigVMMemoryHandle(uint8* InPtr, uint16 InSize = 1, FType InType = FType::Plain, int32 InRegisterOffset = INDEX_NONE)
+	FORCEINLINE_DEBUGGABLE FRigVMMemoryHandle(uint8* InPtr, uint16 InSize = 1, FType InType = FType::Plain, const FRigVMRegisterOffset* InRegisterOffset = nullptr)
 		: Ptr(InPtr)
 		, Type(InType)
 		, Size(InSize)
-		, Offset((uint16)InRegisterOffset)
+		, RegisterOffset(InRegisterOffset)
 	{}
 
-	FORCEINLINE_DEBUGGABLE FRigVMMemoryHandle(uint8* InPtr, const FRigVMRegister& InRegister, int32 InRegisterOffset = INDEX_NONE)
+	FORCEINLINE_DEBUGGABLE FRigVMMemoryHandle(uint8* InPtr, const FRigVMRegister& InRegister, const FRigVMRegisterOffset* InRegisterOffset = nullptr)
 		: Ptr(InPtr)
 		, Type(FType::Plain)
 		, Size(InRegister.ElementSize)
-		, Offset((uint16)InRegisterOffset)
+		, RegisterOffset(InRegisterOffset)
 	{
 		if (InRegister.IsNestedDynamic())
 		{
@@ -404,18 +409,18 @@ public:
 		}
 	}
 
-	FORCEINLINE_DEBUGGABLE FRigVMMemoryHandle(FRigVMByteArray* InPtr, uint16 InSize = 1, int32 InRegisterOffset = INDEX_NONE)
+	FORCEINLINE_DEBUGGABLE FRigVMMemoryHandle(FRigVMByteArray* InPtr, uint16 InSize = 1, const FRigVMRegisterOffset* InRegisterOffset = nullptr)
 		: Ptr((uint8*)InPtr)
 		, Type(FType::Dynamic)
 		, Size(InSize)
-		, Offset((uint16)InRegisterOffset)
+		, RegisterOffset(InRegisterOffset)
 	{}
 
-	FORCEINLINE_DEBUGGABLE FRigVMMemoryHandle(FRigVMNestedByteArray* InPtr, uint16 InSize = 1, int32 InRegisterOffset = INDEX_NONE)
+	FORCEINLINE_DEBUGGABLE FRigVMMemoryHandle(FRigVMNestedByteArray* InPtr, uint16 InSize = 1, const FRigVMRegisterOffset* InRegisterOffset = nullptr)
 		: Ptr((uint8*)InPtr)
 		, Type(FType::NestedDynamic)
 		, Size(InSize)
-		, Offset((uint16)InRegisterOffset)
+		, RegisterOffset(InRegisterOffset)
 	{}
 
 	FORCEINLINE_DEBUGGABLE operator const uint8*() const
@@ -440,7 +445,7 @@ public:
 
 private:
 
-	FORCEINLINE_DEBUGGABLE uint8* GetData_Internal(int32 SliceIndex, bool bGetArrayData = false) const
+	FORCEINLINE_DEBUGGABLE uint8* GetData_Internal_NoOffset(int32 SliceIndex, bool bGetArrayData = false) const
 	{
 		ensure(SliceIndex >= 0);
 
@@ -457,7 +462,7 @@ private:
 			}
 			case FType::Dynamic:
 			{
-				ensure(Offset == UINT16_MAX);
+				ensure(RegisterOffset == nullptr);
 				if (!bGetArrayData)
 				{
 					return Ptr;
@@ -477,7 +482,7 @@ private:
 			}
 			case FType::NestedDynamic:
 			{
-				ensure(Offset == UINT16_MAX);
+				ensure(RegisterOffset == nullptr);
 
 				FRigVMNestedByteArray* ArrayStorage = (FRigVMNestedByteArray*)Ptr;
 				if (ArrayStorage->Num() == 0 && bGetArrayData)
@@ -504,10 +509,20 @@ private:
 		}
 	}
 
+	FORCEINLINE_DEBUGGABLE uint8* GetData_Internal(int32 SliceIndex, bool bGetArrayData = false) const
+	{
+		uint8* Result = GetData_Internal_NoOffset(SliceIndex, bGetArrayData);
+		if (Result && RegisterOffset)
+		{
+			Result = RegisterOffset->GetData(Result);
+		}
+		return Result;
+	}
+
 	uint8* Ptr;
 	FType Type;
 	uint16 Size;
-	uint16 Offset;
+	const FRigVMRegisterOffset* RegisterOffset;
 
 	friend class URigVM;
 };
@@ -695,14 +710,30 @@ public:
 	// Returns a memory handle for a given register
 	FORCEINLINE_DEBUGGABLE FRigVMMemoryHandle GetHandle(const FRigVMRegister& Register, int32 InRegisterOffset = INDEX_NONE) const
 	{
+		const FRigVMRegisterOffset& RegisterOffset = GetRegisterOffset(InRegisterOffset);
+
 		if (Register.IsDynamic())
 		{
 			uint8* Ptr = (uint8*)&Data[Register.GetWorkByteIndex()];
-			return FRigVMMemoryHandle(Ptr, Register);
+			return FRigVMMemoryHandle(Ptr, Register, RegisterOffset.IsValid() ? &RegisterOffset : nullptr);
 		}
 
-		uint8* Ptr = GetDataPtr(Register, InRegisterOffset);
-		return FRigVMMemoryHandle(Ptr, Register.GetNumBytesPerSlice());
+		uint8* Ptr = GetDataPtr(Register);
+		uint16 NumBytes = Register.GetNumBytesPerSlice();
+
+		if (RegisterOffset.IsValid())
+		{
+			NumBytes = RegisterOffset.GetElementSize();
+
+			if (!RegisterOffset.ContainsArraySegment())
+			{
+				// flatten the offset and cache the offset ptr directly
+				Ptr = RegisterOffset.GetData(Ptr);
+				return FRigVMMemoryHandle(Ptr, NumBytes);
+			}
+		}
+
+		return FRigVMMemoryHandle(Ptr, NumBytes, FRigVMMemoryHandle::FType::Plain, RegisterOffset.IsValid() ? &RegisterOffset : nullptr);
 	}
 
 	// Returns a memory handle for a given register
@@ -1125,8 +1156,11 @@ public:
 	void SetRegisterValueFromString(const FRigVMOperand& InOperand, const FString& InCPPType, const UObject* InCPPTypeObject, const TArray<FString>& InDefaultValues);
 	TArray<FString> GetRegisterValueAsString(const FRigVMOperand& InOperand, const FString& InCPPType, const UObject* InCPPTypeObject);
 
+	// Returns the register offset given an index
+	const FRigVMRegisterOffset& GetRegisterOffset(int32 InRegisterOffsetIndex) const;
+
 	// Returns the register offset for a given operand
-	const FRigVMRegisterOffset GetRegisterOffsetForOperand(const FRigVMOperand& InOperand) const;
+	const FRigVMRegisterOffset& GetRegisterOffsetForOperand(const FRigVMOperand& InOperand) const;
 
 	// returns the statistics information
 	FRigVMMemoryStatistics GetStatistics() const
@@ -1374,6 +1408,7 @@ private:
 	bool bEncounteredErrorDuringLoad;
 
 	static FRigVMByteArray DefaultByteArray;
+	static FRigVMRegisterOffset InvalidRegisterOffset;
 
 	friend class URigVM;
 	friend class URigVMCompiler;
