Control Rig Array Support: Deprecate fixed array

diff --git a/Engine/Source/Runtime/RigVM/Private/RigVMCore/RigVM.cpp b/Engine/Source/Runtime/RigVM/Private/RigVMCore/RigVM.cpp
--- a/Engine/Source/Runtime/RigVM/Private/RigVMCore/RigVM.cpp
+++ b/Engine/Source/Runtime/RigVM/Private/RigVMCore/RigVM.cpp
@@ -829,14 +829,39 @@ bool URigVM::Initialize(FRigVMMemoryContainerPtrArray Memory, FRigVMFixedArray<v
 
 				FRigVMMemoryHandle& SourceHandle = CachedMemoryHandles[FirstHandleForInstruction[Context.InstructionIndex]];
 				FRigVMMemoryHandle& TargetHandle = CachedMemoryHandles[FirstHandleForInstruction[Context.InstructionIndex] + 1];
+
+				// register offset traversal for plain type is performed here when converting from handle to pointer here
 				void* SourcePtr = SourceHandle;
 				void* TargetPtr = TargetHandle;
 
 				uint64 NumBytes = reinterpret_cast<uint64>(CachedMemoryHandles[FirstHandleForInstruction[Context.InstructionIndex] + 2].GetData());
 				ERigVMRegisterType MemoryType = (ERigVMRegisterType)reinterpret_cast<uint64>(CachedMemoryHandles[FirstHandleForInstruction[Context.InstructionIndex] + 3].GetData());
 
+				// dynamic memory is managed by individual rig units or generated for debug watches, never copy from it
+				ensure(SourceHandle.Type != FRigVMMemoryHandle::Dynamic);
+				ensure(SourceHandle.Type != FRigVMMemoryHandle::NestedDynamic);
+
+				if (SourceHandle.Type == FRigVMMemoryHandle::StaticArray)
+				{
+					// Array is implemented using a single slice
+					FRigVMByteArray* Storage = (FRigVMByteArray*)SourceHandle.Ptr;
+
+					// source ptr now points to the beginning of the slice
+					SourcePtr = Storage->GetData();
+
+					if (SourceHandle.RegisterOffset)
+					{
+						SourcePtr = SourceHandle.RegisterOffset->GetData((uint8*)SourcePtr);
+					}
+				}
+
+				// debug watch an array pin is not supported at the moment
+				ensure(TargetHandle.Type != FRigVMMemoryHandle::NestedDynamic);
+
 				if (TargetHandle.Type == FRigVMMemoryHandle::Dynamic)
 				{
+					ensure(!TargetHandle.RegisterOffset);
+					// Copy op is used when copy to debug watch values, which is implemented as a dynamic register
 					FRigVMByteArray* Storage = (FRigVMByteArray*)TargetHandle.Ptr;
 					if (Context.GetSlice().GetIndex() == 0)
 					{
@@ -845,16 +870,21 @@ bool URigVM::Initialize(FRigVMMemoryContainerPtrArray Memory, FRigVMFixedArray<v
 					int32 ByteIndex = Storage->AddZeroed(NumBytes);
 					TargetPtr = Storage->GetData() + ByteIndex;
 				}
-				else if (TargetHandle.Type == FRigVMMemoryHandle::NestedDynamic)
+				else if (TargetHandle.Type == FRigVMMemoryHandle::StaticArray)
 				{
-					FRigVMNestedByteArray* Storage = (FRigVMNestedByteArray*)TargetHandle.Ptr;
-					if (Context.GetSlice().GetIndex() == 0)
+					FRigVMByteArray* Storage = (FRigVMByteArray*)TargetHandle.Ptr;
+
+					// static array's space are allocated during AddStaticArray()
+					// todo: allow array-array copy to change the size of target array
+					ensure(Storage->Num() >= NumBytes);
+
+					TargetPtr = Storage->GetData();
+
+					if (TargetHandle.RegisterOffset)
 					{
-						Storage->Reset();
+						// for an example of this use case, see CopyValueToArrayInStructTest unit test
+						TargetPtr = TargetHandle.RegisterOffset->GetData((uint8*)TargetPtr);
 					}
-					int32 ArrayIndex = Storage->Add(FRigVMByteArray());
-					(*Storage)[ArrayIndex].AddZeroed(NumBytes);
-					TargetPtr = (*Storage)[ArrayIndex].GetData();
 				}
 
 				switch (MemoryType)
@@ -1094,14 +1124,39 @@ bool URigVM::Execute(FRigVMMemoryContainerPtrArray Memory, FRigVMFixedArray<void
 
 				FRigVMMemoryHandle& SourceHandle = CachedMemoryHandles[FirstHandleForInstruction[Context.InstructionIndex]];
 				FRigVMMemoryHandle& TargetHandle = CachedMemoryHandles[FirstHandleForInstruction[Context.InstructionIndex] + 1];
+
+				// register offset traversal for plain type is performed here when converting from handle to pointer here
 				void* SourcePtr = SourceHandle;
 				void* TargetPtr = TargetHandle;
 
 				uint64 NumBytes = reinterpret_cast<uint64>(CachedMemoryHandles[FirstHandleForInstruction[Context.InstructionIndex] + 2].GetData());
 				ERigVMRegisterType MemoryType = (ERigVMRegisterType)reinterpret_cast<uint64>(CachedMemoryHandles[FirstHandleForInstruction[Context.InstructionIndex] + 3].GetData());
 
-				if (TargetHandle.Type == FRigVMMemoryHandle::Dynamic)
+				// dynamic memory is managed by individual rig units or generated for debug watches, never copy from it
+				ensure(SourceHandle.Type != FRigVMMemoryHandle::Dynamic);
+				ensure(SourceHandle.Type != FRigVMMemoryHandle::NestedDynamic);
+
+				if (SourceHandle.Type == FRigVMMemoryHandle::StaticArray)
 				{
+					// Array is implemented using a single slice
+					FRigVMByteArray* Storage = (FRigVMByteArray*)SourceHandle.Ptr;
+
+					// source ptr now points to the beginning of the slice
+					SourcePtr = Storage->GetData();
+
+					if (SourceHandle.RegisterOffset)
+					{
+						SourcePtr = SourceHandle.RegisterOffset->GetData((uint8*)SourcePtr);
+					} 
+				}
+
+				// debug watch an array pin is not supported at the moment
+				ensure(TargetHandle.Type != FRigVMMemoryHandle::NestedDynamic);
+
+				if (TargetHandle.Type == FRigVMMemoryHandle::Dynamic)
+				{ 
+					ensure(!TargetHandle.RegisterOffset);
+					// Copy op is used when copy to debug watch values, which is implemented as a dynamic register
 					FRigVMByteArray* Storage = (FRigVMByteArray*)TargetHandle.Ptr;
 					if (Context.GetSlice().GetIndex() == 0)
 					{
@@ -1110,16 +1165,21 @@ bool URigVM::Execute(FRigVMMemoryContainerPtrArray Memory, FRigVMFixedArray<void
 					int32 ByteIndex = Storage->AddZeroed(NumBytes);
 					TargetPtr = Storage->GetData() + ByteIndex;
 				}
-				else if (TargetHandle.Type == FRigVMMemoryHandle::NestedDynamic)
+				else if (TargetHandle.Type == FRigVMMemoryHandle::StaticArray)
 				{
-					FRigVMNestedByteArray* Storage = (FRigVMNestedByteArray*)TargetHandle.Ptr;
-					if (Context.GetSlice().GetIndex() == 0)
-					{
-						Storage->Reset();
+					FRigVMByteArray* Storage = (FRigVMByteArray*)TargetHandle.Ptr;
+ 
+					// static array's space are allocated during AddStaticArray()
+					// todo: allow array-array copy to change the size of target array
+					ensure(Storage->Num() >= NumBytes);
+
+					TargetPtr = Storage->GetData(); 
+
+					if (TargetHandle.RegisterOffset)
+					{ 
+						// for an example of this use case, see CopyValueToArrayInStructTest unit test
+						TargetPtr = TargetHandle.RegisterOffset->GetData((uint8*)TargetPtr);
 					}
-					int32 ArrayIndex = Storage->Add(FRigVMByteArray());
-					(*Storage)[ArrayIndex].AddZeroed(NumBytes);
-					TargetPtr = (*Storage)[ArrayIndex].GetData();
 				}
 
 				switch (MemoryType)
diff --git a/Engine/Source/Runtime/RigVM/Private/RigVMCore/RigVMMemory.cpp b/Engine/Source/Runtime/RigVM/Private/RigVMCore/RigVMMemory.cpp
--- a/Engine/Source/Runtime/RigVM/Private/RigVMCore/RigVMMemory.cpp
+++ b/Engine/Source/Runtime/RigVM/Private/RigVMCore/RigVMMemory.cpp
@@ -478,10 +478,18 @@
 	for (int32 Index = 0; Index < Registers.Num(); Index++)
 	{
 		const FRigVMRegister& Register = Registers[Index];
-		if (Register.IsNestedDynamic())
+
+		if (Register.IsArray())
 		{
 			FMemory::Memzero(&Data[Register.GetWorkByteIndex()], sizeof(FRigVMNestedByteArray));
 			*(FRigVMNestedByteArray*)&Data[Register.GetWorkByteIndex()] = FRigVMNestedByteArray();
+
+			if (!Register.IsDynamic())
+			{
+				// static array always uses a single slice, see AddStaticArray() for similar memory initialization
+				((FRigVMNestedByteArray*)&Data[Register.GetWorkByteIndex()])->SetNum(1);
+				(*((FRigVMNestedByteArray*)&Data[Register.GetWorkByteIndex()]))[0].AddZeroed(Register.GetNumBytesPerSlice());
+			}
 		}
 		else if(Register.IsDynamic())
 		{
@@ -527,7 +535,7 @@
 {
 	for (FRigVMRegister& Register : Registers)
 	{
-		if (Register.IsNestedDynamic())
+		if (Register.IsArray())
 		{
 			FRigVMNestedByteArray* NestedArrayStorage = (FRigVMNestedByteArray*)&Data[Register.GetWorkByteIndex()];
 			Register.SliceCount = NestedArrayStorage->Num();
@@ -562,101 +570,7 @@
 			continue;
 		}
 
-		if (!Register.IsDynamic())
+		if (Register.IsArray())
 		{
-			switch (Register.Type)
-			{
-				case ERigVMRegisterType::Plain:
-				{
-					FRigVMByteArray View;
-					View.Append(&Data[Register.GetWorkByteIndex()], Register.GetAllocatedBytes() - Register.GetAlignmentBytes());
-					Ar << View;
-					break;
-				}
-				case ERigVMRegisterType::Name:
-				{
-					TArray<FName> View;
-					View.Append((FName*)&Data[Register.GetWorkByteIndex()], Register.GetTotalElementCount());
-					Ar << View;
-					break;
-				}
-				case ERigVMRegisterType::String:
-				{
-					TArray<FString> View;
-					View.Append((FString*)&Data[Register.GetWorkByteIndex()], Register.GetTotalElementCount());
-					Ar << View;
-					break;
-				}
-				case ERigVMRegisterType::Struct:
-				{
-					uint8* DataPtr = &Data[Register.GetWorkByteIndex()];
-					UScriptStruct* ScriptStruct = GetScriptStruct(Register);
-
-					TArray<uint8, TAlignedHeapAllocator<16>> DefaultStructData;
-					DefaultStructData.AddZeroed(ScriptStruct->GetStructureSize());
-					ScriptStruct->InitializeDefaultValue(DefaultStructData.GetData());
-
-					TArray<FString> View;
-					for (uint16 ElementIndex = 0; ElementIndex < Register.GetTotalElementCount(); ElementIndex++)
-					{
-						FString Value;
-						ScriptStruct->ExportText(Value, DataPtr, DefaultStructData.GetData(), nullptr, PPF_None, nullptr);
-						View.Add(Value);
-						DataPtr += Register.ElementSize;
-					}
-
-					ScriptStruct->DestroyStruct(DefaultStructData.GetData(), 1);
-
-					Ar << View;
-					break;
-				}
-			}
-		}
-		else if (!Register.IsNestedDynamic())
-		{
-			FRigVMByteArray* ArrayStorage = (FRigVMByteArray*)&Data[Register.GetWorkByteIndex()];
-
-			switch (Register.Type)
-			{
-			case ERigVMRegisterType::Plain:
-			{
-				Ar << *ArrayStorage;
-				break;
-			}
-			case ERigVMRegisterType::Name:
-			{
-				TArray<FName> View = FRigVMFixedArray<FName>(*ArrayStorage);
-				Ar << View;
-				break;
-			}
-			case ERigVMRegisterType::String:
-			{
-				TArray<FString> View = FRigVMFixedArray<FString>(*ArrayStorage);
-				Ar << View;
-				break;
-			}
-			case ERigVMRegisterType::Struct:
-			{
-				uint8* DataPtr = ArrayStorage->GetData();
-				uint16 NumElements = (uint16)(ArrayStorage->Num() / Register.ElementSize);
-				ensure(NumElements * Register.ElementSize == ArrayStorage->Num());
-				UScriptStruct* ScriptStruct = GetScriptStruct(Register);
-
-				TArray<FString> View;
-				for (uint16 ElementIndex = 0; ElementIndex < NumElements; ElementIndex++)
-				{
-					FString Value;
-					ScriptStruct->ExportText(Value, DataPtr, nullptr, nullptr, PPF_None, nullptr);
-					View.Add(Value);
-					DataPtr += Register.ElementSize;
-				}
-
-				Ar << View;
-				break;
-			}
-			}
-		}
-		else
-		{
 			FRigVMNestedByteArray* NestedArrayStorage = (FRigVMNestedByteArray*)&Data[Register.GetWorkByteIndex()];
 
@@ -712,6 +626,100 @@
 				}
 			}
 		}
+		else if (Register.IsDynamic())
+		{
+			FRigVMByteArray* ArrayStorage = (FRigVMByteArray*)&Data[Register.GetWorkByteIndex()];
+
+			switch (Register.Type)
+			{
+				case ERigVMRegisterType::Plain:
+				{
+					Ar << *ArrayStorage;
+					break;
+				}
+				case ERigVMRegisterType::Name:
+				{
+					TArray<FName> View = FRigVMFixedArray<FName>(*ArrayStorage);
+					Ar << View;
+					break;
+				}
+				case ERigVMRegisterType::String:
+				{
+					TArray<FString> View = FRigVMFixedArray<FString>(*ArrayStorage);
+					Ar << View;
+					break;
+				}
+				case ERigVMRegisterType::Struct:
+				{
+					uint8* DataPtr = ArrayStorage->GetData();
+					uint16 NumElements = (uint16)(ArrayStorage->Num() / Register.ElementSize);
+					ensure(NumElements * Register.ElementSize == ArrayStorage->Num());
+					UScriptStruct* ScriptStruct = GetScriptStruct(Register);
+
+					TArray<FString> View;
+					for (uint16 ElementIndex = 0; ElementIndex < NumElements; ElementIndex++)
+					{
+						FString Value;
+						ScriptStruct->ExportText(Value, DataPtr, nullptr, nullptr, PPF_None, nullptr);
+						View.Add(Value);
+						DataPtr += Register.ElementSize;
+					}
+
+					Ar << View;
+					break;
+				}
+			}
+		}
+		else
+		{
+			switch (Register.Type)
+			{
+				case ERigVMRegisterType::Plain:
+				{
+					FRigVMByteArray View;
+					View.Append(&Data[Register.GetWorkByteIndex()], Register.GetAllocatedBytes() - Register.GetAlignmentBytes());
+					Ar << View;
+					break;
+				}
+				case ERigVMRegisterType::Name:
+				{
+					TArray<FName> View;
+					View.Append((FName*)&Data[Register.GetWorkByteIndex()], Register.GetTotalElementCount());
+					Ar << View;
+					break;
+				}
+				case ERigVMRegisterType::String:
+				{
+					TArray<FString> View;
+					View.Append((FString*)&Data[Register.GetWorkByteIndex()], Register.GetTotalElementCount());
+					Ar << View;
+					break;
+				}
+				case ERigVMRegisterType::Struct:
+				{
+					uint8* DataPtr = &Data[Register.GetWorkByteIndex()];
+					UScriptStruct* ScriptStruct = GetScriptStruct(Register);
+
+					TArray<uint8, TAlignedHeapAllocator<16>> DefaultStructData;
+					DefaultStructData.AddZeroed(ScriptStruct->GetStructureSize());
+					ScriptStruct->InitializeDefaultValue(DefaultStructData.GetData());
+
+					TArray<FString> View;
+					for (uint16 ElementIndex = 0; ElementIndex < Register.GetTotalElementCount(); ElementIndex++)
+					{
+						FString Value;
+						ScriptStruct->ExportText(Value, DataPtr, DefaultStructData.GetData(), nullptr, PPF_None, nullptr);
+						View.Add(Value);
+						DataPtr += Register.ElementSize;
+					}
+
+					ScriptStruct->DestroyStruct(DefaultStructData.GetData(), 1);
+
+					Ar << View;
+					break;
+				}
+			}
+		} 
 	}
 }
 
@@ -815,58 +823,74 @@
 			continue;
 		}
 
-		if (!Register.IsDynamic())
+		if (Register.IsArray())
 		{
-			switch (Register.Type)
+			FRigVMNestedByteArray* NestedArrayStorage = nullptr;
+			if (!bEncounteredErrorDuringLoad)
 			{
-				case ERigVMRegisterType::Plain:
+				NestedArrayStorage = (FRigVMNestedByteArray*)&Data[Register.GetWorkByteIndex()];
+				NestedArrayStorage->Reset();
+				NestedArrayStorage->SetNumZeroed(Register.SliceCount);
+			}
+
+			for (int32 SliceIndex = 0; SliceIndex < Register.SliceCount; SliceIndex++)
+			{
+				FRigVMByteArray* ArrayStorage = nullptr;
+				if (NestedArrayStorage)
 				{
-					FRigVMByteArray View;
-					Ar << View;
-					if (!bEncounteredErrorDuringLoad)
+					ArrayStorage = &(*NestedArrayStorage)[SliceIndex];
+				}
+
+				switch (Register.Type)
+				{
+					case ERigVMRegisterType::Plain:
 					{
-						ensure(View.Num() <= Register.GetAllocatedBytes());
-						FMemory::Memcpy(&Data[Register.GetWorkByteIndex()], View.GetData(), View.Num());
+						FRigVMByteArray View;
+						Ar << View;
+						if (ArrayStorage)
+						{
+							*ArrayStorage = View;
+						}
+						break;
 					}
-					break;
-				}
-				case ERigVMRegisterType::Name:
-				{
-					TArray<FName> View;
-					Ar << View;
-					if (!bEncounteredErrorDuringLoad)
+					case ERigVMRegisterType::Name:
 					{
-						ensure(View.Num() == Register.GetTotalElementCount());
-						RigVMCopy<FName>(&Data[Register.GetWorkByteIndex()], View.GetData(), View.Num());
+						TArray<FName> View;
+						Ar << View;
+
+						if (ArrayStorage)
+						{
+							FRigVMDynamicArray<FName> ValueArray(*ArrayStorage);
+							ValueArray.CopyFrom(View);
+						}
+						break;
 					}
-					break;
-				}
-				case ERigVMRegisterType::String:
-				{
-					TArray<FString> View;
-					Ar << View;
-					if (!bEncounteredErrorDuringLoad)
+					case ERigVMRegisterType::String:
 					{
-						ensure(View.Num() == Register.GetTotalElementCount());
-						RigVMCopy<FString>(&Data[Register.GetWorkByteIndex()], View.GetData(), View.Num());
+						TArray<FString> View;
+						Ar << View;
+
+						if (ArrayStorage)
+						{
+							FRigVMDynamicArray<FString> ValueArray(*ArrayStorage);
+							ValueArray.CopyFrom(View);
+						}
+						break;
 					}
-					break;
-				}
-				case ERigVMRegisterType::Struct:
-				{
-					TArray<FString> View;
-					Ar << View;
-					if (!bEncounteredErrorDuringLoad)
+					case ERigVMRegisterType::Struct:
 					{
-						ensure(View.Num() == Register.GetTotalElementCount());
+						TArray<FString> View;
+						Ar << View;
 
-						uint8* DataPtr = &Data[Register.GetWorkByteIndex()];
 						UScriptStruct* ScriptStruct = GetScriptStruct(Register);
-						if (ScriptStruct && !bEncounteredErrorDuringLoad)
+						if (ScriptStruct && !bEncounteredErrorDuringLoad && ArrayStorage)
 						{
 							ensure(ScriptStruct->GetStructureSize() == Register.ElementSize);
 
-							for (uint16 ElementIndex = 0; ElementIndex < Register.GetTotalElementCount(); ElementIndex++)
+							ArrayStorage->SetNumZeroed(View.Num() * Register.ElementSize);
+							uint8* DataPtr = ArrayStorage->GetData();
+
+							for (uint16 ElementIndex = 0; ElementIndex < View.Num(); ElementIndex++)
 							{
 								FRigVMMemoryContainerImportErrorContext ErrorPipe;
 								ScriptStruct->ImportText(*View[ElementIndex], DataPtr, nullptr, PPF_None, &ErrorPipe, ScriptStruct->GetName());
@@ -878,12 +902,12 @@
 								DataPtr += Register.ElementSize;
 							}
 						}
+						break;
 					}
-					break;
 				}
 			}
 		}
-		else if (!Register.IsNestedDynamic())
+		else if (Register.IsDynamic())
 		{
 			FRigVMByteArray* ArrayStorage = nullptr;
 
@@ -959,72 +983,56 @@
 		}
 		else
 		{
-			FRigVMNestedByteArray* NestedArrayStorage = nullptr;
-			if (!bEncounteredErrorDuringLoad)
+			switch (Register.Type)
 			{
-				NestedArrayStorage = (FRigVMNestedByteArray*)&Data[Register.GetWorkByteIndex()];
-				NestedArrayStorage->Reset();
-				NestedArrayStorage->SetNumZeroed(Register.SliceCount);
-			}
+				case ERigVMRegisterType::Plain:
+				{
+					FRigVMByteArray View;
+					Ar << View;
+					if (!bEncounteredErrorDuringLoad)
+					{
+						ensure(View.Num() <= Register.GetAllocatedBytes());
+						FMemory::Memcpy(&Data[Register.GetWorkByteIndex()], View.GetData(), View.Num());
+					}
+					break;
+				}
+				case ERigVMRegisterType::Name:
+				{
+					TArray<FName> View;
+					Ar << View;
+					if (!bEncounteredErrorDuringLoad)
+					{
+						ensure(View.Num() == Register.GetTotalElementCount());
+						RigVMCopy<FName>(&Data[Register.GetWorkByteIndex()], View.GetData(), View.Num());
+					}
+					break;
+				}
+				case ERigVMRegisterType::String:
+				{
+					TArray<FString> View;
+					Ar << View;
+					if (!bEncounteredErrorDuringLoad)
+					{
+						ensure(View.Num() == Register.GetTotalElementCount());
+						RigVMCopy<FString>(&Data[Register.GetWorkByteIndex()], View.GetData(), View.Num());
+					}
+					break;
+				}
+				case ERigVMRegisterType::Struct:
+				{
+					TArray<FString> View;
+					Ar << View;
+					if (!bEncounteredErrorDuringLoad)
+					{
+						ensure(View.Num() == Register.GetTotalElementCount());
 
-			for (int32 SliceIndex = 0; SliceIndex < Register.SliceCount; SliceIndex++)
-			{
-				FRigVMByteArray* ArrayStorage = nullptr;
-				if (NestedArrayStorage)
-				{
-					ArrayStorage = &(*NestedArrayStorage)[SliceIndex];
-				}
-
-				switch (Register.Type)
-				{
-					case ERigVMRegisterType::Plain:
-					{
-						FRigVMByteArray View;
-						Ar << View;
-						if (ArrayStorage)
-						{
-							*ArrayStorage = View;
-						}
-						break;
-					}
-					case ERigVMRegisterType::Name:
-					{
-						TArray<FName> View;
-						Ar << View;
-
-						if (ArrayStorage)
-						{
-							FRigVMDynamicArray<FName> ValueArray(*ArrayStorage);
-							ValueArray.CopyFrom(View);
-						}
-						break;
-					}
-					case ERigVMRegisterType::String:
-					{
-						TArray<FString> View;
-						Ar << View;
-
-						if (ArrayStorage)
-						{
-							FRigVMDynamicArray<FString> ValueArray(*ArrayStorage);
-							ValueArray.CopyFrom(View);
-						}
-						break;
-					}
-					case ERigVMRegisterType::Struct:
-					{
-						TArray<FString> View;
-						Ar << View;
-
+						uint8* DataPtr = &Data[Register.GetWorkByteIndex()];
 						UScriptStruct* ScriptStruct = GetScriptStruct(Register);
-						if (ScriptStruct && !bEncounteredErrorDuringLoad && ArrayStorage)
+						if (ScriptStruct && !bEncounteredErrorDuringLoad)
 						{
 							ensure(ScriptStruct->GetStructureSize() == Register.ElementSize);
 
-							ArrayStorage->SetNumZeroed(View.Num() * Register.ElementSize);
-							uint8* DataPtr = ArrayStorage->GetData();
-
-							for (uint16 ElementIndex = 0; ElementIndex < View.Num(); ElementIndex++)
+							for (uint16 ElementIndex = 0; ElementIndex < Register.GetTotalElementCount(); ElementIndex++)
 							{
 								FRigVMMemoryContainerImportErrorContext ErrorPipe;
 								ScriptStruct->ImportText(*View[ElementIndex], DataPtr, nullptr, PPF_None, &ErrorPipe, ScriptStruct->GetName());
@@ -1036,11 +1044,11 @@
 								DataPtr += Register.ElementSize;
 							}
 						}
-						break;
 					}
+					break;
 				}
 			}
-		}
+		} 
 	}
 
 	if (bEncounteredErrorDuringLoad)
@@ -1158,17 +1166,69 @@
 		return false;
 	}
 
+
 	const FRigVMRegister& Source = InSourceMemory->Registers[InSourceRegisterIndex];
 	const FRigVMRegister& Target = Registers[InTargetRegisterIndex];
 
-	// todo: copy between slices!
-	int32 SourceElementCount = Source.ElementCount;
-	const uint8* SourcePtr = InSourceMemory->GetDataPtr(Source, InSourceRegisterOffset, InSourceSliceIndex, false);
-	uint8* TargetPtr = GetDataPtr(Target, InTargetRegisterOffset, InTargetSliceIndex, false);
-	uint16 NumBytes = Target.GetNumBytesPerSlice();
+	if (Source.GetNumBytesPerSlice() == 0)
+	{
+		return true;
+	}
 
-	if (Source.IsNestedDynamic())
+	uint8* SourcePtr = nullptr;
+	uint8* TargetPtr = nullptr;
+	uint16 NumBytes = 0;
+	ERigVMRegisterType DataType = ERigVMRegisterType::Invalid;
+
+	if (InSourceRegisterOffset != INDEX_NONE)
 	{
+		DataType = InSourceMemory->RegisterOffsets[InSourceRegisterOffset].GetType();
+		NumBytes = InSourceMemory->RegisterOffsets[InSourceRegisterOffset].GetElementSize();
+	}
+	else
+	{
+		DataType = Source.Type;
+		NumBytes = Source.GetNumBytesPerSlice();
+	}
+
+	// determine the raw data pointer for source
+	ensure(Source.SliceCount > InSourceSliceIndex);
+
+	if (!Source.IsDynamic())
+	{
+		ensure(InSourceSliceIndex == 0);
+		ensure(Source.SliceCount == 1);
+
+		if (!Source.IsArray())
+		{ 
+			SourcePtr = InSourceMemory->GetDataPtr(Source, InSourceRegisterOffset, InSourceSliceIndex, false);
+		}
+		else
+		{ 
+			SourcePtr = (uint8*)&(InSourceMemory->Data[Source.GetWorkByteIndex()]);
+
+			FRigVMNestedByteArray* Storage = (FRigVMNestedByteArray*)SourcePtr;
+
+			ensure(Source.SliceCount == 1);
+			ensure(Storage->Num() == Source.SliceCount);
+
+			FRigVMByteArray* PerSliceStorage = &((*Storage)[InSourceSliceIndex]);
+
+			ensure(PerSliceStorage->Num() == Source.GetNumBytesPerSlice());
+
+			// Source ptr now points to the beginning of the slice
+			SourcePtr = PerSliceStorage->GetData();
+
+			if (InSourceRegisterOffset != INDEX_NONE)
+			{
+				SourcePtr = InSourceMemory->RegisterOffsets[InSourceRegisterOffset].GetData(SourcePtr);
+			}
+		}
+	}
+	else
+	{
+		// dynamic and nested dynamic, do not allow any register offsets
+		// because copy should only happen at slice level
 #if WITH_EDITOR
 		ensure(InSourceRegisterOffset == INDEX_NONE);
 #else
@@ -1177,24 +1237,93 @@
 			return true;
 		}
 #endif
-		FRigVMByteArray* ArrayStorage = (FRigVMByteArray*)SourcePtr;
-		SourceElementCount = ArrayStorage->Num() / Source.ElementSize;
-		SourcePtr = ArrayStorage->GetData();
-	}
-	else if (Source.IsDynamic())
-	{
-		FRigVMByteArray* ArrayStorage = (FRigVMByteArray*)SourcePtr;
-		SourcePtr = ArrayStorage->GetData();
+
+		SourcePtr = (uint8*)&(InSourceMemory->Data[Source.GetWorkByteIndex()]);
+
+		if (!Source.IsArray())
+		{
+			ensure(Source.ElementCount == 1);
+
+			FRigVMByteArray* Storage = (FRigVMByteArray*)SourcePtr;
+
+			ensure(Storage->Num() == Source.GetNumBytesAllSlices());
+
+			SourcePtr = (uint8*)Storage->GetData();
+
+			// Source ptr now points to the beginning of the slice
+			SourcePtr = SourcePtr + InSourceSliceIndex * Source.GetNumBytesPerSlice();
+		}
+		else
+		{
+			FRigVMNestedByteArray* Storage = (FRigVMNestedByteArray*)SourcePtr;
+
+			ensure(Storage->Num() == Source.SliceCount);
+
+			FRigVMByteArray* PerSliceStorage = &((*Storage)[InSourceSliceIndex]);
+
+			// Source ptr now points to the beginning of the slice
+			SourcePtr = PerSliceStorage->GetData();
+		} 
 	}
 
-	ERigVMRegisterType TargetType = Target.Type;
-	if (Target.IsDynamic() && !Target.bIsArray)
+	// determine the raw data pointer for target 
+
+	ensure(Target.SliceCount > InTargetSliceIndex); 
+
+	if (!Target.IsDynamic())
 	{
-		FRigVMByteArray* ArrayStorage = (FRigVMByteArray*)TargetPtr;
-		TargetPtr = ArrayStorage->GetData();
+		// plain and array
+
+		ensure(InTargetSliceIndex == 0);
+		ensure(Target.SliceCount == 1);
+
+		if (!Target.IsArray())
+		{ 
+			TargetPtr = GetDataPtr(Target, InTargetRegisterOffset, InTargetSliceIndex, false); 
+		}
+		else
+		{ 
+			TargetPtr = (uint8*)&Data[Target.GetWorkByteIndex()];
+
+			// fixed arrays are stored using a single slice nested byte array 
+			FRigVMNestedByteArray* Storage = (FRigVMNestedByteArray*)TargetPtr;
+
+			ensure(Target.SliceCount == 1);
+			ensure(Storage->Num() == Target.SliceCount);
+
+			FRigVMByteArray* PerSliceStorage = &((*Storage)[InTargetSliceIndex]);
+
+			// only array-array copy triggers reallocation of the array
+			if (InTargetRegisterOffset == INDEX_NONE
+				&& Source.IsArray()
+				&& InSourceRegisterOffset == INDEX_NONE)
+			{
+				// copying the entire slice, reallocate space
+				PerSliceStorage->Reset();
+				PerSliceStorage->AddZeroed(NumBytes);
+
+				TargetPtr = PerSliceStorage->GetData();
+			}
+			else
+			{ 
+				// target is somewhere inside the array
+
+				// target ptr now points to the beginning of the array, space should have been allocated
+				ensure(PerSliceStorage->Num() > 0);
+				TargetPtr = PerSliceStorage->GetData();
+ 
+				if (InTargetRegisterOffset != INDEX_NONE)
+				{ 
+					TargetPtr = RegisterOffsets[InTargetRegisterOffset].GetData(TargetPtr);
+				}
+
+			}
+		} 
 	}
-	else if (Target.IsNestedDynamic())
+	else
 	{
+		// dynamic and nested dynamic, do not allow any register offsets
+		// because copy should only happen at slice level
 #if WITH_EDITOR
 		ensure(InTargetRegisterOffset == INDEX_NONE);
 #else
@@ -1202,37 +1331,41 @@
 		{
 			return true;
 		}
-		return true;
 #endif
 
-		NumBytes = Source.ElementSize * SourceElementCount;
+		TargetPtr = (uint8*)&Data[Target.GetWorkByteIndex()];
 
-		FRigVMByteArray* ArrayStorage = (FRigVMByteArray*)TargetPtr;
-		if (ArrayStorage->Num() != NumBytes)
+		if (!Target.IsArray())
 		{
-			Destroy(InTargetRegisterIndex, INDEX_NONE, InTargetSliceIndex);
-			ArrayStorage->SetNumZeroed(NumBytes);
-			Construct(InTargetRegisterIndex, INDEX_NONE, InTargetSliceIndex);
+			ensure(Target.ElementCount == 1); 
 
-			TargetPtr = ArrayStorage->GetData();
+			FRigVMByteArray* Storage = (FRigVMByteArray*)TargetPtr;
+
+			ensure(Storage->Num() == Target.GetNumBytesAllSlices());
+
+			TargetPtr = (uint8*)Storage->GetData();
+ 
+			// target ptr now points to the beginning of the slice
+			TargetPtr = TargetPtr + InTargetSliceIndex * Target.GetNumBytesPerSlice();
 		}
-	}
-	else if (InTargetRegisterOffset != INDEX_NONE)
-	{
-		if (Target.GetNumBytesPerSlice() == 0)
+		else
 		{
-			return true;
-		}
+			FRigVMNestedByteArray* Storage = (FRigVMNestedByteArray*)TargetPtr;
 
-		TargetType = RegisterOffsets[InTargetRegisterOffset].GetType();
-		NumBytes = RegisterOffsets[InTargetRegisterOffset].GetElementSize();
-	}
-	else if (Target.GetNumBytesPerSlice() == 0)
-	{
-		return true;
-	}
+			ensure(Storage->Num() == Target.SliceCount); 
 
-	return Copy(InTargetRegisterIndex, InTargetRegisterOffset, TargetType, SourcePtr, TargetPtr, NumBytes);
+			FRigVMByteArray* PerSliceStorage = &((*Storage)[InTargetSliceIndex]);
+
+			// copying the entire slice, reallocate space
+			PerSliceStorage->Reset();
+			PerSliceStorage->AddZeroed(NumBytes); 
+
+			// target ptr now points to the beginning of the slice
+			TargetPtr = PerSliceStorage->GetData();
+		} 
+	} 
+
+	return Copy(InTargetRegisterIndex, InTargetRegisterOffset, DataType, SourcePtr, TargetPtr, NumBytes);
 }
 
 bool FRigVMMemoryContainer::Copy(
@@ -1348,28 +1481,26 @@
 	int32 ElementCount = InElementIndex == INDEX_NONE ? Register.GetTotalElementCount() : 1;
 	int32 NumSlices = 1;
 
-	uint8* TopDataPtr = &Data[Register.GetWorkByteIndex()];
-	if (Register.IsDynamic())
-	{
-		if (Register.IsNestedDynamic())
-		{
-			FRigVMNestedByteArray& Storage = *(FRigVMNestedByteArray*)TopDataPtr;
-			NumSlices = Storage.Num();
-			TopDataPtr = (uint8*)Storage.GetData();
-		}
-		else
-		{
-			FRigVMByteArray& Storage = *(FRigVMByteArray*)TopDataPtr;
-			ElementCount = Storage.Num() / Register.ElementSize;
-			if (ElementCount == 0)
-			{
-				Storage.AddZeroed(Register.ElementSize);
-				ElementCount = 1;
-			}
-			ensure(ElementCount * Register.ElementSize == Storage.Num());
-			TopDataPtr = Storage.GetData();
-		}
+	uint8* TopDataPtr = &Data[Register.GetWorkByteIndex()]; 
+	if (Register.IsArray())
+	{ 
+		FRigVMNestedByteArray& Storage = *(FRigVMNestedByteArray*)TopDataPtr;
+		NumSlices = Storage.Num();
+		TopDataPtr = (uint8*)Storage.GetData();
 	}
+	else if (Register.IsDynamic())
+	{
+		FRigVMByteArray& Storage = *(FRigVMByteArray*)TopDataPtr;
+		ElementCount = Storage.Num() / Register.ElementSize;
+		if (ElementCount == 0)
+		{
+			Storage.AddZeroed(Register.ElementSize);
+			ElementCount = 1;
+		}
+		ensure(ElementCount * Register.ElementSize == Storage.Num());
+		TopDataPtr = Storage.GetData();
 
+	} 
+
 	for (int32 SliceIndex = 0; SliceIndex < NumSlices; SliceIndex++)
 	{
@@ -1377,14 +1508,14 @@
 		{
 			SliceIndex = InSliceIndex;
 
-			if (Register.IsNestedDynamic())
+			if (Register.IsArray())
 			{
 				TopDataPtr = TopDataPtr + SliceIndex * sizeof(FRigVMByteArray);
 			}
 		}
 
 		uint8* NestedDataPtr = TopDataPtr;
-		if (Register.IsNestedDynamic())
+		if (Register.IsArray())
 		{
 			FRigVMByteArray& Storage = *(FRigVMByteArray*)TopDataPtr;
 			ElementCount = Storage.Num() / Register.ElementSize;
@@ -1435,7 +1566,7 @@
 			}
 		}
 
-		if (Register.IsNestedDynamic())
+		if (Register.IsArray())
 		{
 			TopDataPtr = TopDataPtr + sizeof(FRigVMByteArray);
 		}
@@ -1463,22 +1594,20 @@
 	int32 NumSlices = 1;
 
 	uint8* TopDataPtr = &Data[Register.GetWorkByteIndex()];
-	if (Register.IsDynamic())
+
+	if (Register.IsArray())
 	{
-		if (Register.IsNestedDynamic())
-		{
-			FRigVMNestedByteArray& Storage = *(FRigVMNestedByteArray*)TopDataPtr;
-			NumSlices = Storage.Num();
-			TopDataPtr = (uint8*)Storage.GetData();
-		}
-		else
-		{
-			FRigVMByteArray& Storage = *(FRigVMByteArray*)TopDataPtr;
-			ElementCount = Storage.Num() / Register.ElementSize;
-			ensure(ElementCount * Register.ElementSize == Storage.Num());
-			TopDataPtr = Storage.GetData();
-		}
+		FRigVMNestedByteArray& Storage = *(FRigVMNestedByteArray*)TopDataPtr;
+		NumSlices = Storage.Num();
+		TopDataPtr = (uint8*)Storage.GetData();
 	}
+	else if (Register.IsDynamic())
+	{
+		FRigVMByteArray& Storage = *(FRigVMByteArray*)TopDataPtr;
+		ElementCount = Storage.Num() / Register.ElementSize;
+		ensure(ElementCount * Register.ElementSize == Storage.Num());
+		TopDataPtr = Storage.GetData();
+	} 
 
 	for (int32 SliceIndex = 0; SliceIndex < NumSlices; SliceIndex++)
 	{
@@ -1486,14 +1615,14 @@
 		{
 			SliceIndex = InSliceIndex;
 
-			if (Register.IsNestedDynamic())
+			if (Register.IsArray())
 			{
 				TopDataPtr = TopDataPtr + SliceIndex * sizeof(FRigVMByteArray);
 			}
 		}
 
 		uint8* NestedDataPtr = TopDataPtr;
-		if (Register.IsNestedDynamic())
+		if (Register.IsArray())
 		{
 			FRigVMByteArray& Storage = *(FRigVMByteArray*)TopDataPtr;
 			ElementCount = Storage.Num() / Register.ElementSize;
@@ -1561,7 +1690,7 @@
 			}
 		}
 
-		if (Register.IsNestedDynamic())
+		if (Register.IsArray())
 		{
 			((FRigVMByteArray*)TopDataPtr)->Empty();
 			TopDataPtr = TopDataPtr + sizeof(FRigVMByteArray);
@@ -1573,21 +1702,18 @@
 		}
 	}
 
-	if (Register.IsDynamic())
+	
+	TopDataPtr = &Data[Register.GetWorkByteIndex()];
+	if (Register.IsArray())
 	{
-		TopDataPtr = &Data[Register.GetWorkByteIndex()];
-
-		if (Register.IsNestedDynamic())
+		if (InSliceIndex != INDEX_NONE)
 		{
-			if (InSliceIndex != INDEX_NONE)
-			{
-				((FRigVMNestedByteArray*)TopDataPtr)->Empty();
-			}
-		}
-		else
-		{
-			((FRigVMByteArray*)TopDataPtr)->Empty();
-		}
+			((FRigVMNestedByteArray*)TopDataPtr)->Empty();
+		} 
+	}
+	else if (Register.IsDynamic())
+	{
+		((FRigVMByteArray*)TopDataPtr)->Empty();
 	}
 
 	return true;
@@ -1742,7 +1868,7 @@
 
 	int32 SliceCount = Register.SliceCount;
 	
-	if (Register.IsNestedDynamic())
+	if (Register.IsArray())
 	{
 		FRigVMNestedByteArray* Storage = (FRigVMNestedByteArray*)&Data[Register.GetWorkByteIndex()];
 		SliceCount = Storage->Num();
@@ -1756,7 +1882,7 @@
 	for (int32 SliceIndex = 0; SliceIndex < SliceCount; SliceIndex++)
 	{
 		int32 ElementCount = Register.ElementCount;
-		if (Register.IsNestedDynamic())
+		if (Register.IsArray())
 		{
 			FRigVMNestedByteArray* Storage = (FRigVMNestedByteArray*)&Data[Register.GetWorkByteIndex()];
 			ElementCount = (*Storage)[SliceIndex].Num() / Register.ElementSize;
diff --git a/Engine/Source/Runtime/RigVM/Public/RigVMCore/RigVM.h b/Engine/Source/Runtime/RigVM/Public/RigVMCore/RigVM.h
--- a/Engine/Source/Runtime/RigVM/Public/RigVMCore/RigVM.h
+++ b/Engine/Source/Runtime/RigVM/Public/RigVMCore/RigVM.h
@@ -223,9 +223,9 @@ public:
 			RegisterIndex = WorkMemoryPtr->Add<T>(WorkMemoryPtr->SupportsNames() ? InName : NAME_None, DefaultValues[0], 1);
 		}
 		else
-	{
-			RegisterIndex = WorkMemoryPtr->AddFixedArray<T>(WorkMemoryPtr->SupportsNames() ? InName : NAME_None, FRigVMFixedArray<T>(DefaultValues), 1);
-	}
+		{
+			RegisterIndex = WorkMemoryPtr->AddStaticArray<T>(WorkMemoryPtr->SupportsNames() ? InName : NAME_None, FRigVMFixedArray<T>(DefaultValues), 1);
+		}
 
 		if (RegisterIndex == INDEX_NONE)
 	{
diff --git a/Engine/Source/Runtime/RigVM/Public/RigVMCore/RigVMArray.h b/Engine/Source/Runtime/RigVM/Public/RigVMCore/RigVMArray.h
--- a/Engine/Source/Runtime/RigVM/Public/RigVMCore/RigVMArray.h
+++ b/Engine/Source/Runtime/RigVM/Public/RigVMCore/RigVMArray.h
@@ -11,6 +11,24 @@ typedef TArray<FRigVMByteArray> FRigVMNestedByteArray;
 /**
  * The FRigVMDynamicArray is used as an array wrapping a generic TArray.
  * FRigVMDynamicArray is copied by reference.
+ * Note: currently there are two different use cases for FRigVMDynamicArray
+ * 
+ * 1. Dynamic && NonArray Register : 
+ * Example: RigUnit_SetControlTransform.generated.h -> FRigUnit_SetControlRotator -> RigVMExecute() -> CachedIndex_4_Array
+ * CachedIndex_4_Array->Num() == SliceCount, both the register and the array are dynamic in a sense that 
+ * they can both grow according to number of slices. In this case, the word "Dynamic" has a consistent meaning: multi-slice
+ * 
+ * 2. Dynamic && Array Register:
+ * Example: RigUnit_SetControlTransform.generated.h -> FRigUnit_SetMultiControlRotator -> RigVMExecute() -> CachedControlIndices
+ * CachedControlIndices->Num() is not related to SliceCount, its size/content are instead managed by the rig unit's execute function in RigUnit_SetControlTransform.cpp
+ *
+ * In this use case, though the register is dynamic(meaning multi-slice), the size of this array is not related to slicing, 
+ * it is only "dynamic" in a sense that it is resize-able(a different meaning of "dynamic"), which can cause confusion
+ *
+ * The current code use it in this use case only because it is copied by reference and resize-able.
+ * 
+ * Ideally, the size of the this data structure should be governed by a single rule, 
+ * mark as todo for future work to unify use cases above
  */
 template<class T>
 struct FRigVMDynamicArray
diff --git a/Engine/Source/Runtime/RigVM/Public/RigVMCore/RigVMMemory.h b/Engine/Source/Runtime/RigVM/Public/RigVMCore/RigVMMemory.h
--- a/Engine/Source/Runtime/RigVM/Public/RigVMCore/RigVMMemory.h
+++ b/Engine/Source/Runtime/RigVM/Public/RigVMCore/RigVMMemory.h
@@ -371,6 +371,7 @@ public:
 	enum FType
 	{
 		Plain,
+		StaticArray,
 		Dynamic,
 		NestedDynamic
 	};
@@ -403,6 +404,11 @@ public:
 		{
 			Type = FType::Dynamic;
 		}
+		else if (InRegister.IsArray())
+		{
+			Type = FType::StaticArray;
+			Size = (uint16)InRegister.GetNumBytesPerSlice();
+		}
 		else
 		{
 			Size = (uint16)InRegister.GetNumBytesPerSlice();
@@ -416,13 +422,6 @@ public:
 		, RegisterOffset(InRegisterOffset)
 	{}
 
-	FORCEINLINE_DEBUGGABLE FRigVMMemoryHandle(FRigVMNestedByteArray* InPtr, uint16 InSize = 1, const FRigVMRegisterOffset* InRegisterOffset = nullptr)
-		: Ptr((uint8*)InPtr)
-		, Type(FType::NestedDynamic)
-		, Size(InSize)
-		, RegisterOffset(InRegisterOffset)
-	{}
-
 	FORCEINLINE_DEBUGGABLE operator const uint8*() const
 	{
 		return GetData();
@@ -460,6 +459,17 @@ private:
 			{
 				return Ptr + SliceIndex * Size;
 			}
+			case FType::StaticArray: 
+			{
+				// array is implemented using a single slice nested byte array
+				// the cached pointer should be pointing to the beginning of that slice
+				ensure(SliceIndex == 0);
+
+				FRigVMByteArray* Storage = (FRigVMByteArray*)Ptr;
+
+				// always get array data for array because rig units consume these arrays as fixed-size arrays
+				return Storage->GetData();
+			}
 			case FType::Dynamic:
 			{
 				ensure(RegisterOffset == nullptr);
@@ -540,7 +550,7 @@ private:
  *
  * This can also be done with arrays:
  *      TArray<float> MyArray = {3.f, 4.f, 5.f};
- * 		int32 Index = Container.AddFixedArray<float>(MyArray);
+ * 		int32 Index = Container.AddStaticArray<float>(MyArray);
  *      FRigVMFixedArray<float> ArrayView = Container.GetFixedArray<float>(Index);
  */
 USTRUCT()
@@ -696,12 +706,38 @@ private:
 		else
 		{
 			Ptr = (uint8*)&Data[Register.GetWorkByteIndex(InSliceIndex)];
+
+			if (Register.IsArray())
+			{
+				FRigVMNestedByteArray* ArrayStorage = (FRigVMNestedByteArray*)Ptr;
+				Ptr = (uint8*)ArrayStorage->GetData();
+
+				if (Ptr)
+				{
+					Ptr = Ptr + InSliceIndex * sizeof(FRigVMByteArray);
+					if (Ptr && bArrayContent)
+					{
+						Ptr = ((FRigVMByteArray*)Ptr)->GetData();
+					}
+				}
+			}
 		}
 
-		if (InRegisterOffset != INDEX_NONE && Ptr != nullptr)
+		if ((Register.IsArray() && bArrayContent) || !Register.IsArray())
 		{
-			Ptr = RegisterOffsets[InRegisterOffset].GetData(Ptr);
+			if (InRegisterOffset != INDEX_NONE && Ptr != nullptr)
+			{
+				Ptr = RegisterOffsets[InRegisterOffset].GetData(Ptr);
+			}
 		}
+		else
+		{
+			// in case register is array and array content is false,
+			// the pointer to the array itself is returned, possibly for the purpose of reallocation
+			// in this case, register offset should not be supplied.
+			ensure(InRegisterOffset == INDEX_NONE);
+		}
+
 		return Ptr;
 	}
 
@@ -725,15 +761,25 @@ public:
 		{
 			NumBytes = RegisterOffset.GetElementSize();
 
+			// flatten the offset and cache the offset ptr directly
 			if (!RegisterOffset.ContainsArraySegment())
 			{
-				// flatten the offset and cache the offset ptr directly
+				// resolve the extra indirection that comes with Arrays
+				if (Register.IsArray())
+				{ 
+					FRigVMByteArray* Storage = (FRigVMByteArray*)Ptr;
+					Ptr = Storage->GetData();
+				}
+
 				Ptr = RegisterOffset.GetData(Ptr);
 				return FRigVMMemoryHandle(Ptr, NumBytes);
 			}
 		}
+ 
+		// array-array copy requires special logic, thus they need to be differentiated
+		FRigVMMemoryHandle::FType HandleType = Register.IsArray() ? FRigVMMemoryHandle::FType::StaticArray : FRigVMMemoryHandle::FType::Plain;
 
-		return FRigVMMemoryHandle(Ptr, NumBytes, FRigVMMemoryHandle::FType::Plain, RegisterOffset.IsValid() ? &RegisterOffset : nullptr);
+		return FRigVMMemoryHandle(Ptr, NumBytes, HandleType, RegisterOffset.IsValid() ? &RegisterOffset : nullptr);
 	}
 
 	// Returns a memory handle for a given register
@@ -904,19 +950,27 @@ public:
 	{
 		if (InRegisterOffset == INDEX_NONE)
 		{
-			uint8* Ptr = (uint8*)GetDataPtr(InRegister, InRegisterOffset, InSliceIndex, false);
-			if (InRegister.IsNestedDynamic())
-			{
-				FRigVMByteArray* Storage = (FRigVMByteArray*)Ptr;
-				return FRigVMFixedArray<T>(*Storage);
+			if (!InRegister.IsArray())
+			{ 
+				uint8* Ptr = (uint8*)GetDataPtr(InRegister, InRegisterOffset, InSliceIndex, false);
+				if (InRegister.IsNestedDynamic())
+				{
+					FRigVMByteArray* Storage = (FRigVMByteArray*)Ptr;
+					return FRigVMFixedArray<T>(*Storage);
+				}
+				else if (InRegister.IsDynamic())
+				{
+					FRigVMByteArray* Storage = (FRigVMByteArray*)Ptr;
+					Ptr = Storage->GetData() + InSliceIndex * InRegister.GetNumBytesPerSlice();
+				}
+
+				return FRigVMFixedArray<T>((T*)Ptr, InRegister.ElementCount);
 			}
-			else if (InRegister.IsDynamic())
+			else
 			{
-				FRigVMByteArray* Storage = (FRigVMByteArray*)Ptr;
-				Ptr = Storage->GetData() + InSliceIndex * InRegister.GetNumBytesPerSlice();
+				FRigVMByteArray* ArrayStorage = (FRigVMByteArray*)GetDataPtr(InRegister, INDEX_NONE, 0, false);
+				return FRigVMFixedArray<T>(*ArrayStorage);
 			}
-
-			return FRigVMFixedArray<T>((T*)Ptr, InRegister.ElementCount);
 		}
 
 		TArray<T>* StoredArray = (TArray<T>*)GetData(InRegister, InRegisterOffset, InSliceIndex);
@@ -1066,75 +1120,79 @@ public:
 		return GetIndex(InPotentialNewName) == INDEX_NONE;
 	}
 
-	// Adds a new named register for a typed array from an array view (used by compiler)
+	// Adds a new named register for a typed value from a value reference.
 	template<typename T>
-	FORCEINLINE_DEBUGGABLE int32 AddFixedArray(const FName& InNewName, const FRigVMFixedArray<T>& InArrayView, int32 InSliceCount = 1)
+	FORCEINLINE_DEBUGGABLE int32 Add(const FName& InNewName, const T& InValue, int32 InSliceCount = 1)
 	{
-		return AddRegisterArray<T>(true, InNewName, InArrayView.Num(), true, (const uint8*)InArrayView.GetData(), InSliceCount);
+		return AddRegisterArray<T>(true, InNewName, 1, false, (const uint8*)&InValue, InSliceCount);
 	}
 
-	// Adds a new unnamed register for a typed array from an array view.
+	// Adds a new unnamed register for a typed value from a value reference.
 	template<typename T>
-	FORCEINLINE_DEBUGGABLE int32 AddFixedArray(const FRigVMFixedArray<T>& InArrayView, int32 InSliceCount = 1)
+	FORCEINLINE_DEBUGGABLE int32 Add(const T& InValue, int32 InSliceCount = 1)
 	{
-		return AddFixedArray<T>(NAME_None, InArrayView, InSliceCount);
+		return Add<T>(NAME_None, InValue, InSliceCount);
 	}
 
+
 	// Adds a new named register for a typed array from an array view (used by compiler)
 	template<typename T>
-	FORCEINLINE_DEBUGGABLE int32 AddDynamicArray(const FName& InNewName, const FRigVMFixedArray<T>& InArrayView, int32 InSliceCount = 1)
+	FORCEINLINE_DEBUGGABLE int32 AddStaticArray(const FName& InNewName, const FRigVMFixedArray<T>& InArrayView, int32 InSliceCount = 1)
 	{
-		return AddRegisterArray<T>(false, InNewName, InArrayView.Num(), true, (const uint8*)InArrayView.GetData(), InSliceCount);
+		return AddRegisterArray<T>(true, InNewName, InArrayView.Num(), true, (const uint8*)InArrayView.GetData(), InSliceCount);
 	}
 
 	// Adds a new unnamed register for a typed array from an array view.
 	template<typename T>
-	FORCEINLINE_DEBUGGABLE int32 AddDynamicArray(const FRigVMFixedArray<T>& InArrayView, int32 InSliceCount = 1)
+	FORCEINLINE_DEBUGGABLE int32 AddStaticArray(const FRigVMFixedArray<T>& InArrayView, int32 InSliceCount = 1)
 	{
-		return AddDynamicArray<T>(NAME_None, InArrayView, InSliceCount);
-	}
+		return AddStaticArray<T>(NAME_None, InArrayView, InSliceCount);
+	} 
 
-	// Adds a new named register for a typed array from an array view (used by compiler)
+
+	// Adds a new named register for a typed value from a value reference.
 	template<typename T>
-	FORCEINLINE_DEBUGGABLE int32 AddDynamicArray(const FName& InNewName, const TArray<T>& InArray, int32 InSliceCount = 1)
+	FORCEINLINE_DEBUGGABLE int32 AddDynamicValue(const FName& InNewName, const T& InValue, int32 InSliceCount = 1)
 	{
-		return AddDynamicArray<T>(InNewName, FRigVMFixedArray<T>(InArray), InSliceCount);
+		return AddRegisterArray<T>(false, InNewName, 1, false, (const uint8*)&InValue, InSliceCount);
 	}
 
-	// Adds a new unnamed register for a typed array from an array view.
+	// Adds a new unnamed register for a typed value from a value reference.
 	template<typename T>
-	FORCEINLINE_DEBUGGABLE int32 AddDynamicArray(const TArray<T>& InArray, int32 InSliceCount = 1)
+	FORCEINLINE_DEBUGGABLE int32 AddDynamicValue(const T& InValue, int32 InSliceCount = 1)
 	{
-		return AddDynamicArray<T>(NAME_None, InArray, InSliceCount);
+		return AddDynamicValue<T>(NAME_None, InValue, InSliceCount);
 	}
 
-	// Adds a new named register for a typed value from a value reference.
+
+	// Adds a new named register for a typed array from an array view (used by compiler)
 	template<typename T>
-	FORCEINLINE_DEBUGGABLE int32 Add(const FName& InNewName, const T& InValue, int32 InSliceCount = 1)
+	FORCEINLINE_DEBUGGABLE int32 AddDynamicArray(const FName& InNewName, const FRigVMFixedArray<T>& InArrayView, int32 InSliceCount = 1)
 	{
-		return AddRegisterArray<T>(true, InNewName, 1, false, (const uint8*)&InValue, InSliceCount);
+		return AddRegisterArray<T>(false, InNewName, InArrayView.Num(), true, (const uint8*)InArrayView.GetData(), InSliceCount);
 	}
 
-	// Adds a new unnamed register for a typed value from a value reference.
+	// Adds a new unnamed register for a typed array from an array view.
 	template<typename T>
-	FORCEINLINE_DEBUGGABLE int32 Add(const T& InValue, int32 InSliceCount = 1)
+	FORCEINLINE_DEBUGGABLE int32 AddDynamicArray(const FRigVMFixedArray<T>& InArrayView, int32 InSliceCount = 1)
 	{
-		return Add<T>(NAME_None, InValue, InSliceCount);
+		return AddDynamicArray<T>(NAME_None, InArrayView, InSliceCount);
 	}
 
-	// Adds a new named register for a typed value from a value reference.
+	// Adds a new named register for a typed array from an array(used by compiler)
 	template<typename T>
-	FORCEINLINE_DEBUGGABLE int32 AddDynamicValue(const FName& InNewName, const T& InValue, int32 InSliceCount = 1)
+	FORCEINLINE_DEBUGGABLE int32 AddDynamicArray(const FName& InNewName, const TArray<T>& InArray, int32 InSliceCount = 1)
 	{
-		return AddRegisterArray<T>(false, InNewName, 1, false, (const uint8*)&InValue, InSliceCount);
+		return AddDynamicArray<T>(InNewName, FRigVMFixedArray<T>(InArray), InSliceCount);
 	}
 
-	// Adds a new unnamed register for a typed value from a value reference.
+	// Adds a new unnamed register for a typed array from an array.
 	template<typename T>
-	FORCEINLINE_DEBUGGABLE int32 AddDynamicValue(const T& InValue, int32 InSliceCount = 1)
+	FORCEINLINE_DEBUGGABLE int32 AddDynamicArray(const TArray<T>& InArray, int32 InSliceCount = 1)
 	{
-		return AddDynamicValue<T>(NAME_None, InValue, InSliceCount);
-	}
+		return AddDynamicArray<T>(NAME_None, InArray, InSliceCount);
+	} 
+
 
 	// Remove a register given its index
 	// Note: This only works if SupportsNames() == true
@@ -1187,84 +1245,123 @@ private:
 		typename T,
 		typename TEnableIf<TIsArithmetic<T>::Value>::Type* = nullptr
 	>
-	FORCEINLINE_DEBUGGABLE int32 AddRegisterArray(bool bFixed, const FName& InNewName, int32 InCount, bool bIsArrayPerSlice, const uint8* InDataPtr, int32 InSliceCount)
+	FORCEINLINE_DEBUGGABLE int32 AddRegisterArray(bool bStatic, const FName& InNewName, int32 InCount, bool bIsArrayPerSlice, const uint8* InDataPtr, int32 InSliceCount)
 	{
-		return AddRegisterArray<T>(bFixed, InNewName, InCount, bIsArrayPerSlice, InDataPtr, InSliceCount, ERigVMRegisterType::Plain, nullptr);
+		return AddRegisterArray<T>(bStatic, InNewName, InCount, bIsArrayPerSlice, InDataPtr, InSliceCount, ERigVMRegisterType::Plain, nullptr);
 	}
 
 	template<
 		typename T,
 		typename TEnableIf<TRigVMIsName<T>::Value>::Type* = nullptr
 	>
-	FORCEINLINE_DEBUGGABLE int32 AddRegisterArray(bool bFixed, const FName& InNewName, int32 InCount, bool bIsArrayPerSlice, const uint8* InDataPtr, int32 InSliceCount)
+	FORCEINLINE_DEBUGGABLE int32 AddRegisterArray(bool bStatic, const FName& InNewName, int32 InCount, bool bIsArrayPerSlice, const uint8* InDataPtr, int32 InSliceCount)
 	{
-		return AddRegisterArray<T>(bFixed, InNewName, InCount, bIsArrayPerSlice, InDataPtr, InSliceCount, ERigVMRegisterType::Name, nullptr);
+		return AddRegisterArray<T>(bStatic, InNewName, InCount, bIsArrayPerSlice, InDataPtr, InSliceCount, ERigVMRegisterType::Name, nullptr);
 	}
 
 	template<
 		typename T,
 		typename TEnableIf<TRigVMIsString<T>::Value>::Type* = nullptr
 	>
-	FORCEINLINE_DEBUGGABLE int32 AddRegisterArray(bool bFixed, const FName& InNewName, int32 InCount, bool bIsArrayPerSlice, const uint8* InDataPtr, int32 InSliceCount)
+	FORCEINLINE_DEBUGGABLE int32 AddRegisterArray(bool bStatic, const FName& InNewName, int32 InCount, bool bIsArrayPerSlice, const uint8* InDataPtr, int32 InSliceCount)
 	{
-		return AddRegisterArray<T>(bFixed, InNewName, InCount, bIsArrayPerSlice, InDataPtr, InSliceCount, ERigVMRegisterType::String, nullptr);
+		return AddRegisterArray<T>(bStatic, InNewName, InCount, bIsArrayPerSlice, InDataPtr, InSliceCount, ERigVMRegisterType::String, nullptr);
 	}
 
 	template<
 		typename T,
 		typename TEnableIf<TModels<CRigVMUStruct, T>::Value>::Type* = nullptr
 	>
-	FORCEINLINE_DEBUGGABLE int32 AddRegisterArray(bool bFixed, const FName& InNewName, int32 InCount, bool bIsArrayPerSlice, const uint8* InDataPtr, int32 InSliceCount)
+	FORCEINLINE_DEBUGGABLE int32 AddRegisterArray(bool bStatic, const FName& InNewName, int32 InCount, bool bIsArrayPerSlice, const uint8* InDataPtr, int32 InSliceCount)
 	{
-		return AddRegisterArray<T>(bFixed, InNewName, InCount, bIsArrayPerSlice, InDataPtr, InSliceCount, ERigVMRegisterType::Struct, T::StaticStruct());
+		return AddRegisterArray<T>(bStatic, InNewName, InCount, bIsArrayPerSlice, InDataPtr, InSliceCount, ERigVMRegisterType::Struct, T::StaticStruct());
 	}
 
 	template<
 		typename T,
 		typename TEnableIf<TIsEnum<T>::Value>::Type* = nullptr
 	>
-	FORCEINLINE_DEBUGGABLE int32 AddRegisterArray(bool bFixed, const FName& InNewName, int32 InCount, bool bIsArrayPerSlice, const uint8* InDataPtr, int32 InSliceCount)
+	FORCEINLINE_DEBUGGABLE int32 AddRegisterArray(bool bStatic, const FName& InNewName, int32 InCount, bool bIsArrayPerSlice, const uint8* InDataPtr, int32 InSliceCount)
 	{
-		return AddRegisterArray<T>(bFixed, InNewName, InCount, bIsArrayPerSlice, InDataPtr, InSliceCount, ERigVMRegisterType::Plain, nullptr);
+		return AddRegisterArray<T>(bStatic, InNewName, InCount, bIsArrayPerSlice, InDataPtr, InSliceCount, ERigVMRegisterType::Plain, nullptr);
 	}
 
 	template <
 		typename T,
 		typename TEnableIf<TRigVMIsBaseStructure<T>::Value, T>::Type* = nullptr
 	>
-	FORCEINLINE_DEBUGGABLE int32 AddRegisterArray(bool bFixed, const FName& InNewName, int32 InCount, bool bIsArrayPerSlice, const uint8* InDataPtr, int32 InSliceCount)
+	FORCEINLINE_DEBUGGABLE int32 AddRegisterArray(bool bStatic, const FName& InNewName, int32 InCount, bool bIsArrayPerSlice, const uint8* InDataPtr, int32 InSliceCount)
 	{
-		return AddRegisterArray<T>(bFixed, InNewName, InCount, bIsArrayPerSlice, InDataPtr, InSliceCount, ERigVMRegisterType::Plain, TBaseStructure<T>::Get());
+		return AddRegisterArray<T>(bStatic, InNewName, InCount, bIsArrayPerSlice, InDataPtr, InSliceCount, ERigVMRegisterType::Plain, TBaseStructure<T>::Get());
 	}
 
 	template<typename T>
-	FORCEINLINE_DEBUGGABLE int32 AddRegisterArray(bool bFixed, const FName& InNewName, int32 InCount, bool bIsArrayPerSlice, const uint8* InDataPtr, int32 InSliceCount, ERigVMRegisterType InType, UScriptStruct* InScriptStruct)
+	FORCEINLINE_DEBUGGABLE int32 AddRegisterArray(bool bStatic, const FName& InNewName, int32 InCount, bool bIsArrayPerSlice, const uint8* InDataPtr, int32 InSliceCount, ERigVMRegisterType InType, UScriptStruct* InScriptStruct)
+	{ 
+		int32 InElementSize = sizeof(T);
+		return AddRegisterArray(bStatic, InNewName, InElementSize, InCount, bIsArrayPerSlice, InDataPtr, InSliceCount, InType, InScriptStruct);
+	}
+
+	FORCEINLINE_DEBUGGABLE int32 AddRegisterArray(bool bStatic, const FName& InNewName, int32 InElementSize, int32 InCount, bool bIsArrayPerSlice, const uint8* InDataPtr, int32 InSliceCount, ERigVMRegisterType InType, UScriptStruct* InScriptStruct)
 	{
-		if(bFixed)
+		if (bStatic)
 		{
-			return AddFixedArray(InNewName, sizeof(T), InCount, bIsArrayPerSlice, InDataPtr, InSliceCount, InType, InScriptStruct);
+			ensure(InSliceCount == 1);
+			if (!bIsArrayPerSlice)
+			{ 
+				ensure(InCount == 1);
+				return AddStaticValue(InNewName, InElementSize, InDataPtr, InType, InScriptStruct);
+			}
+			else
+			{ 
+				return AddStaticArray(InNewName, InElementSize, InCount, InDataPtr, InType, InScriptStruct);
+			}
 		}
 		else
 		{
-			return AddDynamicArray(InNewName, sizeof(T), InCount, bIsArrayPerSlice, InDataPtr, InSliceCount, InType, InScriptStruct);
+			if (!bIsArrayPerSlice)
+			{ 
+				ensure(InCount == 1);
+				return AddDynamicValue(InNewName, InElementSize, InDataPtr, InSliceCount, InType, InScriptStruct);
+			}
+			else
+			{
+
+				return AddDynamicArray(InNewName, InElementSize, InCount, InDataPtr, InSliceCount, InType, InScriptStruct);
+			}
 		}
 	}
 
-	FORCEINLINE_DEBUGGABLE int32 AddRegisterArray(bool bFixed, const FName& InNewName, int32 InElementSize, int32 InCount, bool bIsArrayPerSlice, const uint8* InDataPtr, int32 InSliceCount, ERigVMRegisterType InType, UScriptStruct* InScriptStruct)
-	{
-		if (bFixed)
+	FORCEINLINE_DEBUGGABLE int32 AddStaticValue(const FName& InNewName, int32 InElementSize, const uint8* InDataPtr, ERigVMRegisterType InType, UScriptStruct* InScriptStruct)
+	{ 
+		int32 Register = Allocate(InNewName, InElementSize, 1, 1, nullptr, false);
+		if (!Registers.IsValidIndex(Register))
 		{
-			return AddFixedArray(InNewName, InElementSize, InCount, bIsArrayPerSlice, InDataPtr, InSliceCount, InType, InScriptStruct);
+			return Register;
 		}
-		else
+
+		Registers[Register].Type = InType;
+		Registers[Register].ScriptStructIndex = FindOrAddScriptStruct(InScriptStruct);
+		Registers[Register].bIsDynamic = false;
+		Registers[Register].bIsArray = false;
+		Registers[Register].ElementSize = InElementSize;
+		Registers[Register].ElementCount = 1;
+		Registers[Register].SliceCount = 1;
+
+		UpdateRegisters();
+		Construct(Register);
+
+		if (InDataPtr)
 		{
-			return AddDynamicArray(InNewName, InElementSize, InCount, bIsArrayPerSlice, InDataPtr, InSliceCount, InType, InScriptStruct);
+			Copy(Register, INDEX_NONE, InType, InDataPtr, GetData(Registers[Register], INDEX_NONE, 0), InElementSize);
 		}
+
+		return Register;
 	}
 
-	// Adds a new named register for a fixed array from a data pointer (used by compiler)
-	FORCEINLINE_DEBUGGABLE int32 AddFixedArray(const FName& InNewName, int32 InElementSize, int32 InCount, bool bIsArrayPerSlice, const uint8* InDataPtr, int32 InSliceCount, ERigVMRegisterType InType, UScriptStruct* InScriptStruct)	{
-		int32 Register = Allocate(InNewName, InElementSize, InCount, InSliceCount, nullptr, false);
+	FORCEINLINE_DEBUGGABLE int32 AddStaticArray(const FName& InNewName, int32 InElementSize, int32 InCount, const uint8* InDataPtr, ERigVMRegisterType InType, UScriptStruct* InScriptStruct)
+	{ 
+		int32 Register = Allocate(InNewName, sizeof(FRigVMNestedByteArray), 1, 1, nullptr);
 		if (!Registers.IsValidIndex(Register))
 		{
 			return Register;
@@ -1272,96 +1369,103 @@ private:
 
 		Registers[Register].Type = InType;
 		Registers[Register].ScriptStructIndex = FindOrAddScriptStruct(InScriptStruct);
-		Registers[Register].bIsArray = bIsArrayPerSlice;
+		Registers[Register].bIsDynamic = false;
+		Registers[Register].bIsArray = true;
+		Registers[Register].ElementSize = InElementSize;
+		Registers[Register].ElementCount = InCount;
+		Registers[Register].SliceCount = 1;
 
-		UpdateRegisters();
-		Construct(Register);
+		uint8* Ptr = (uint8*)&Data[Registers[Register].GetWorkByteIndex()];
+		FRigVMNestedByteArray& Storage = *(FRigVMNestedByteArray*)Ptr;
+		Storage.SetNum(1);
+		Storage[0].SetNumZeroed(InCount * InElementSize);
 
 		if (InDataPtr)
 		{
-			for (int32 SliceIndex = 0; SliceIndex < InSliceCount; SliceIndex++)
-			{
-				Copy(Register, INDEX_NONE, InType, InDataPtr, GetData(Registers[Register], INDEX_NONE, SliceIndex), InElementSize * InCount);
-			}
+			Copy(Register, INDEX_NONE, InType, InDataPtr, Storage[0].GetData(), InElementSize * InCount);
+		}
+		else
+		{
+			Construct(Register);
 		}
+
 		return Register;
 	}
-
-	// Adds a new named register for a dynamic array from a data pointer (used by compiler)
-	FORCEINLINE_DEBUGGABLE int32 AddDynamicArray(const FName& InNewName, int32 InElementSize, int32 InCount, bool bIsArrayPerSlice, const uint8* InDataPtr, int32 InSliceCount, ERigVMRegisterType InType, UScriptStruct* InScriptStruct)
-	{
-		int32 Register = INDEX_NONE;
-		
-		if (bIsArrayPerSlice)
+ 
+	FORCEINLINE_DEBUGGABLE int32 AddDynamicValue(const FName& InNewName, int32 InElementSize, const uint8* InDataPtr, int32 InSliceCount, ERigVMRegisterType InType, UScriptStruct* InScriptStruct) 
+	{ 
+		int32 Register = Allocate(InNewName, sizeof(FRigVMByteArray), 1, 1, nullptr);
+		if (!Registers.IsValidIndex(Register))
 		{
-			Register = Allocate(InNewName, sizeof(FRigVMNestedByteArray), 1, 1, nullptr);
-			if (!Registers.IsValidIndex(Register))
-			{
-				return Register;
-			}
+			return Register;
+		}
+
+		Registers[Register].Type = InType;
+		Registers[Register].ScriptStructIndex = FindOrAddScriptStruct(InScriptStruct);
+		Registers[Register].bIsDynamic = true;
+		Registers[Register].bIsArray = false;
+		Registers[Register].ElementSize = InElementSize;
+		Registers[Register].ElementCount = 1;
+		Registers[Register].SliceCount = InSliceCount;
 
-			Registers[Register].Type = InType;
-			Registers[Register].ScriptStructIndex = FindOrAddScriptStruct(InScriptStruct);
-			Registers[Register].bIsDynamic = true;
-			Registers[Register].bIsArray = bIsArrayPerSlice;
-			Registers[Register].ElementSize = InElementSize;
-			Registers[Register].ElementCount = InCount;
-			Registers[Register].SliceCount = InSliceCount;
+		uint8* Ptr = (uint8*)&Data[Registers[Register].GetWorkByteIndex()];
+		FRigVMByteArray& Storage = *(FRigVMByteArray*)Ptr;
 
-			uint8* Ptr = (uint8*)&Data[Registers[Register].GetWorkByteIndex()];
-			FRigVMNestedByteArray& Storage = *(FRigVMNestedByteArray*)Ptr;
-			Storage.SetNum(InSliceCount);
+		Storage.SetNumZeroed(InSliceCount * InElementSize);
 
+		if (InDataPtr)
+		{
 			for (int32 SliceIndex = 0; SliceIndex < InSliceCount; SliceIndex++)
 			{
-				Storage[SliceIndex].SetNumZeroed(InCount * InElementSize);
-				if (InDataPtr)
-				{
-					Copy(Register, INDEX_NONE, InType, InDataPtr, Storage[SliceIndex].GetData(), InElementSize * InCount);
-				}
-			}
-
-			if (InDataPtr == nullptr)
-			{
-				Construct(Register);
+				Copy(Register, INDEX_NONE, InType, InDataPtr, &Storage[SliceIndex * InElementSize], InElementSize);
 			}
 		}
 		else
 		{
-			Register = Allocate(InNewName, sizeof(FRigVMByteArray), 1, 1, nullptr);
-			if (!Registers.IsValidIndex(Register))
-			{
-				return Register;
-			}
+			Construct(Register);
+		}
+
+		return Register;
+	}
 
-			Registers[Register].Type = InType;
-			Registers[Register].ScriptStructIndex = FindOrAddScriptStruct(InScriptStruct);
-			Registers[Register].bIsDynamic = true;
-			Registers[Register].bIsArray = bIsArrayPerSlice;
-			Registers[Register].ElementSize = InElementSize;
-			Registers[Register].ElementCount = InCount;
-			Registers[Register].SliceCount = InSliceCount;
+	FORCEINLINE_DEBUGGABLE int32 AddDynamicArray(const FName& InNewName, int32 InElementSize, int32 InCount, const uint8* InDataPtr, int32 InSliceCount, ERigVMRegisterType InType, UScriptStruct* InScriptStruct)
+	{
+		int32 Register = Allocate(InNewName, sizeof(FRigVMNestedByteArray), 1, 1, nullptr);
+		if (!Registers.IsValidIndex(Register))
+		{
+			return Register;
+		}
+
+		Registers[Register].Type = InType;
+		Registers[Register].ScriptStructIndex = FindOrAddScriptStruct(InScriptStruct);
+		Registers[Register].bIsDynamic = true;
+		Registers[Register].bIsArray = true;
+		Registers[Register].ElementSize = InElementSize;
+		Registers[Register].ElementCount = InCount;
+		Registers[Register].SliceCount = InSliceCount;
 
-			uint8* Ptr = (uint8*)&Data[Registers[Register].GetWorkByteIndex()];
-			FRigVMByteArray& Storage = *(FRigVMByteArray*)Ptr;
+		uint8* Ptr = (uint8*)&Data[Registers[Register].GetWorkByteIndex()];
+		FRigVMNestedByteArray& Storage = *(FRigVMNestedByteArray*)Ptr;
+		Storage.SetNum(InSliceCount);
 
-			Storage.SetNumZeroed(InSliceCount * InElementSize);
+		for (int32 SliceIndex = 0; SliceIndex < InSliceCount; SliceIndex++)
+		{
+			Storage[SliceIndex].SetNumZeroed(InCount * InElementSize);
 			if (InDataPtr)
 			{
-				for (int32 SliceIndex = 0; SliceIndex < InSliceCount; SliceIndex++)
-				{
-					Copy(Register, INDEX_NONE, InType, InDataPtr, &Storage[SliceIndex * InElementSize], InElementSize);
-				}
-			}
-			else
-			{
-				Construct(Register);
+				Copy(Register, INDEX_NONE, InType, InDataPtr, Storage[SliceIndex].GetData(), InElementSize * InCount);
 			}
 		}
 
+		if (InDataPtr == nullptr)
+		{
+			Construct(Register);
+		}
+
 		return Register;
 	}
 
+
 	// Updates internal data for topological changes
 	void UpdateRegisters();
 
