Control Rig: Crash in FRigVMMemoryContainer on Save

To fix this we'll make sure to copy a deferred vm if needed prior to save.

diff --git a/Engine/Source/Runtime/RigVM/Private/RigVMCore/RigVM.cpp b/Engine/Source/Runtime/RigVM/Private/RigVMCore/RigVM.cpp
--- a/Engine/Source/Runtime/RigVM/Private/RigVMCore/RigVM.cpp
+++ b/Engine/Source/Runtime/RigVM/Private/RigVMCore/RigVM.cpp
@@ -106,6 +106,8 @@ void URigVM::Serialize(FArchive& Ar)
 
 void URigVM::Save(FArchive& Ar)
 {
+	CopyDeferredVMIfRequired();
+
 	Ar << WorkMemoryStorage;
 	Ar << LiteralMemoryStorage;
 	Ar << FunctionNamesStorage;
