Control Rig: Running backwards and forwards solve now turns the path on both sets of nodes opaque

diff --git a/Engine/Plugins/Experimental/ControlRig/Source/ControlRig/Private/ControlRig.cpp b/Engine/Plugins/Experimental/ControlRig/Source/ControlRig/Private/ControlRig.cpp
--- a/Engine/Plugins/Experimental/ControlRig/Source/ControlRig/Private/ControlRig.cpp
+++ b/Engine/Plugins/Experimental/ControlRig/Source/ControlRig/Private/ControlRig.cpp
@@ -445,6 +445,12 @@ void UControlRig::Execute(const EControlRigState InState, const FName& InEventNa
 				VM->AddExternalVariable(ExternalVariable);
 			}
 		}
+
+		// default to always clear data after each execution
+		// only set a valid first entry event later when execution
+		// has passed the initialization stage and there are multiple events present in one evaluation
+		// first entry event is used to determined when to clear data during an evaluation
+		VM->SetFirstEntryEventInEventQueue(NAME_None);
 	}
 
 	bool bJustRanInit = false;
@@ -598,8 +604,19 @@ void UControlRig::Execute(const EControlRigState InState, const FName& InEventNa
 
 	if (!bSetupModeEnabled)
 	{
+
 		if(!IsExecuting())
-		{
+		{ 
+
+#if WITH_EDITOR
+			// only set a valid first entry event when execution
+			// has passed the initialization stage and there are multiple events present
+			if (EventQueue.Num() >= 2 && VM && InState != EControlRigState::Init)
+			{
+				VM->SetFirstEntryEventInEventQueue(EventQueue[0]);
+			}
+#endif
+
 			ExecuteUnits(Context, InEventName);
 
 			if (InState == EControlRigState::Init)
diff --git a/Engine/Source/Runtime/RigVM/Private/RigVMCore/RigVM.cpp b/Engine/Source/Runtime/RigVM/Private/RigVMCore/RigVM.cpp
--- a/Engine/Source/Runtime/RigVM/Private/RigVMCore/RigVM.cpp
+++ b/Engine/Source/Runtime/RigVM/Private/RigVMCore/RigVM.cpp
@@ -67,6 +67,9 @@ URigVM::URigVM()
     , ByteCodePtr(&ByteCodeStorage)
     , FunctionNamesPtr(&FunctionNamesStorage)
     , FunctionsPtr(&FunctionsStorage)
+#if WITH_EDITOR
+	, FirstEntryEventInQueue(NAME_None)
+#endif
 	, ExecutingThreadId(INDEX_NONE)
 	, DeferredVMToCopy(nullptr)
 {
@@ -965,9 +968,13 @@ bool URigVM::Execute(FRigVMMemoryContainerPtrArray Memory, FRigVMFixedArray<void
 
 #if WITH_EDITOR
 	TArray<FName>& FunctionNames = GetFunctionNames();
-	InstructionVisitedDuringLastRun.Reset();
-	InstructionVisitOrder.Reset();
-	InstructionVisitedDuringLastRun.SetNumZeroed(Instructions.Num());
+
+	if (FirstEntryEventInQueue == NAME_None || FirstEntryEventInQueue == InEntryName)
+	{
+		InstructionVisitedDuringLastRun.Reset();
+		InstructionVisitOrder.Reset();
+		InstructionVisitedDuringLastRun.SetNumZeroed(Instructions.Num());
+	}
 #endif
 
 	Context.Reset();
@@ -988,7 +995,7 @@ bool URigVM::Execute(FRigVMMemoryContainerPtrArray Memory, FRigVMFixedArray<void
 	while (Instructions.IsValidIndex(Context.InstructionIndex))
 	{
 #if WITH_EDITOR
-		InstructionVisitedDuringLastRun[Context.InstructionIndex] = true;
+		InstructionVisitedDuringLastRun[Context.InstructionIndex]++;
 		InstructionVisitOrder.Add(Context.InstructionIndex);
 #endif
 
diff --git a/Engine/Source/Runtime/RigVM/Public/RigVMCore/RigVM.h b/Engine/Source/Runtime/RigVM/Public/RigVMCore/RigVM.h
--- a/Engine/Source/Runtime/RigVM/Public/RigVMCore/RigVM.h
+++ b/Engine/Source/Runtime/RigVM/Public/RigVMCore/RigVM.h
@@ -194,13 +194,15 @@ public:
 	{
 		if (InstructionVisitedDuringLastRun.IsValidIndex(InIndex))
 		{
-			return InstructionVisitedDuringLastRun[InIndex];
+			return InstructionVisitedDuringLastRun[InIndex] > 0;
 		}
 		return false;
 	}
 
 	// Returns the order of all instructions during the last run
 	FORCEINLINE const TArray<int32> GetInstructionVisitOrder() const { return InstructionVisitOrder; }
+
+	FORCEINLINE const void SetFirstEntryEventInEventQueue(const FName& InFirstEventName) { FirstEntryEventInQueue = InFirstEventName; }
 #endif
 
 	// Returns the parameters of the VM
@@ -528,8 +530,13 @@ private:
 	TArray<FRigVMExternalVariable> ExternalVariables;
 
 #if WITH_EDITOR
-	TArray<bool> InstructionVisitedDuringLastRun;
+	// stores the number of times each instruction was visited
+	TArray<int32> InstructionVisitedDuringLastRun;
 	TArray<int32> InstructionVisitOrder;
+	
+	// Control Rig can run multiple events per evaluation, such as the Backward&Forward Solve Mode,
+	// store the first event such that we know when to reset data for a new round of rig evaluation
+	FName FirstEntryEventInQueue;
 #endif
 
 	void CacheSingleMemoryHandle(const FRigVMOperand& InArg, bool bForExecute = false);
