Control Rig VM: Adding test for guarding against leaks when using sharedptr

The unit test uses a sharedptr of FString across multiple VMs

diff --git a/Engine/Source/Runtime/RigVM/Private/RigVMCore/RigVM.cpp b/Engine/Source/Runtime/RigVM/Private/RigVMCore/RigVM.cpp
--- a/Engine/Source/Runtime/RigVM/Private/RigVMCore/RigVM.cpp
+++ b/Engine/Source/Runtime/RigVM/Private/RigVMCore/RigVM.cpp
@@ -189,7 +189,7 @@ void URigVM::Empty()
 	CachedMemoryHandles.Empty();
 }
 
-void URigVM::CopyFrom(URigVM* InVM, bool bDeferCopy, bool bReferenceLiteralMemory, bool bReferenceByteCode)
+void URigVM::CopyFrom(URigVM* InVM, bool bDeferCopy, bool bReferenceLiteralMemory, bool bReferenceByteCode, bool bCopyExternalVariables)
 {
 	check(InVM);
 
@@ -257,6 +257,11 @@ void URigVM::CopyFrom(URigVM* InVM, bool bDeferCopy, bool bReferenceLiteralMemor
 	Instructions = InVM->Instructions;
 	Parameters = InVM->Parameters;
 	ParametersNameMap = InVM->ParametersNameMap;
+
+	if (bCopyExternalVariables)
+	{
+		ExternalVariables = InVM->ExternalVariables;
+	}
 }
 
 int32 URigVM::AddRigVMFunction(UScriptStruct* InRigVMStruct, const FName& InMethodName)
diff --git a/Engine/Source/Runtime/RigVM/Public/RigVMCore/RigVM.h b/Engine/Source/Runtime/RigVM/Public/RigVMCore/RigVM.h
--- a/Engine/Source/Runtime/RigVM/Public/RigVMCore/RigVM.h
+++ b/Engine/Source/Runtime/RigVM/Public/RigVMCore/RigVM.h
@@ -133,7 +133,7 @@ public:
 	void Empty();
 
 	// resets the container and clones the input VM
-	void CopyFrom(URigVM* InVM, bool bDeferCopy = false, bool bReferenceLiteralMemory = false, bool bReferenceByteCode = false);
+	void CopyFrom(URigVM* InVM, bool bDeferCopy = false, bool bReferenceLiteralMemory = false, bool bReferenceByteCode = false, bool bCopyExternalVariables = false);
 
 	// Initializes all execute ops and their memory.
 	bool Initialize(FRigVMMemoryContainerPtrArray Memory, FRigVMFixedArray<void*> AdditionalArguments);
diff --git a/Engine/Source/Runtime/RigVM/Public/RigVMCore/RigVMExternalVariable.h b/Engine/Source/Runtime/RigVM/Public/RigVMCore/RigVMExternalVariable.h
--- a/Engine/Source/Runtime/RigVM/Public/RigVMCore/RigVMExternalVariable.h
+++ b/Engine/Source/Runtime/RigVM/Public/RigVMCore/RigVMExternalVariable.h
@@ -372,7 +372,7 @@ struct RIGVM_API FRigVMExternalVariable
 	}
 
 	template<typename T>
-	FORCEINLINE T GetValue()
+	FORCEINLINE T GetValue() const
 	{
 		ensure(IsValid() && !bIsArray);
 		return *(T*)Memory;
@@ -385,6 +385,13 @@ struct RIGVM_API FRigVMExternalVariable
 		return *(T*)Memory;
 	}
 
+	template<typename T>
+	FORCEINLINE const T& GetRef() const
+	{
+		ensure(IsValid() && !bIsArray);
+		return *(T*)Memory;
+	}
+
 	template<typename T>
 	FORCEINLINE void SetValue(const T& InValue)
 	{
