Control Rig: Fixed VM crash when duplicating Control Rig Track in Sequencer
Root cause: when calculating alignment, we were using a pointer to the TArray Data instead of the TArray it self that is actually living in the RigVM memory. Thus wrong alignment is exported to text for copy.

diff --git a/Engine/Source/Runtime/RigVM/Private/RigVMCore/RigVMMemory.cpp b/Engine/Source/Runtime/RigVM/Private/RigVMCore/RigVMMemory.cpp
--- a/Engine/Source/Runtime/RigVM/Private/RigVMCore/RigVMMemory.cpp
+++ b/Engine/Source/Runtime/RigVM/Private/RigVMCore/RigVMMemory.cpp
@@ -1867,7 +1867,7 @@ void FRigVMMemoryContainer::UpdateRegisters()
 
 		if (Alignment != 0)
 		{
-			uint8* Pointer = GetData(RegisterIndex);
+			uint8* Pointer = (uint8*)(&(Data[Register.GetWorkByteIndex()]));
 
 			if (Register.AlignmentBytes > 0)
 			{
@@ -1877,7 +1877,7 @@ void FRigVMMemoryContainer::UpdateRegisters()
 					AlignmentShift -= Register.AlignmentBytes;
 					Register.ByteIndex -= Register.AlignmentBytes;
 					Register.AlignmentBytes = 0;
-					Pointer = GetData(RegisterIndex);
+					Pointer = (uint8*)(&(Data[Register.GetWorkByteIndex()]));
 				}
 			}
 
@@ -1887,7 +1887,7 @@ void FRigVMMemoryContainer::UpdateRegisters()
 				Register.AlignmentBytes++;
 				Register.ByteIndex++;
 				AlignmentShift++;
-				Pointer = GetData(RegisterIndex);
+				Pointer = (uint8*)(&(Data[Register.GetWorkByteIndex()]));
 			}
 		}
 	}
