Control Rig: backing out CL15358231, it is causing other crashes

diff --git a/Engine/Source/Runtime/RigVM/Private/RigVMCore/RigVMMemory.cpp b/Engine/Source/Runtime/RigVM/Private/RigVMCore/RigVMMemory.cpp
--- a/Engine/Source/Runtime/RigVM/Private/RigVMCore/RigVMMemory.cpp
+++ b/Engine/Source/Runtime/RigVM/Private/RigVMCore/RigVMMemory.cpp
@@ -1867,7 +1867,7 @@ void FRigVMMemoryContainer::UpdateRegisters()
 
 		if (Alignment != 0)
 		{
-			uint8* Pointer = (uint8*)(&(Data[Register.GetWorkByteIndex()]));
+			uint8* Pointer = GetData(RegisterIndex);
 
 			if (Register.AlignmentBytes > 0)
 			{
@@ -1877,7 +1877,7 @@ void FRigVMMemoryContainer::UpdateRegisters()
 					AlignmentShift -= Register.AlignmentBytes;
 					Register.ByteIndex -= Register.AlignmentBytes;
 					Register.AlignmentBytes = 0;
-					Pointer = (uint8*)(&(Data[Register.GetWorkByteIndex()]));
+					Pointer = GetData(RegisterIndex);
 				}
 			}
 
@@ -1887,7 +1887,7 @@ void FRigVMMemoryContainer::UpdateRegisters()
 				Register.AlignmentBytes++;
 				Register.ByteIndex++;
 				AlignmentShift++;
-				Pointer = (uint8*)(&(Data[Register.GetWorkByteIndex()]));
+				Pointer = GetData(RegisterIndex);
 			}
 		}
 	}
