Control Rig: now uses the right pointer to check for alignment to avoid crash

diff --git a/Engine/Source/Runtime/RigVM/Private/RigVMCore/RigVMMemory.cpp b/Engine/Source/Runtime/RigVM/Private/RigVMCore/RigVMMemory.cpp
--- a/Engine/Source/Runtime/RigVM/Private/RigVMCore/RigVMMemory.cpp
+++ b/Engine/Source/Runtime/RigVM/Private/RigVMCore/RigVMMemory.cpp
@@ -1867,7 +1867,15 @@ void FRigVMMemoryContainer::UpdateRegisters()
 
 		if (Alignment != 0)
 		{
-			uint8* Pointer = GetData(RegisterIndex);
+			// no need to adjust for alignment if nothing is allocated
+			if (!Register.IsDynamic() && Register.ElementCount == 0)
+			{
+				continue;
+			}
+
+			ensure(Data.IsValidIndex(Register.GetWorkByteIndex()));
+
+			uint8* Pointer = (uint8*)(&(Data[Register.GetWorkByteIndex()]));
 
 			if (Register.AlignmentBytes > 0)
 			{
@@ -1877,7 +1885,7 @@ void FRigVMMemoryContainer::UpdateRegisters()
 					AlignmentShift -= Register.AlignmentBytes;
 					Register.ByteIndex -= Register.AlignmentBytes;
 					Register.AlignmentBytes = 0;
-					Pointer = GetData(RegisterIndex);
+					Pointer = (uint8*)(&(Data[Register.GetWorkByteIndex()]));
 				}
 			}
 
@@ -1887,7 +1895,7 @@ void FRigVMMemoryContainer::UpdateRegisters()
 				Register.AlignmentBytes++;
 				Register.ByteIndex++;
 				AlignmentShift++;
-				Pointer = GetData(RegisterIndex);
+				Pointer = (uint8*)(&(Data[Register.GetWorkByteIndex()]));
 			}
 		}
 	}
