RigVM: Serialization code not dealing with Ar.IsCountingMemory()

diff --git a/Engine/Source/Runtime/RigVM/Private/RigVMCore/RigVM.cpp b/Engine/Source/Runtime/RigVM/Private/RigVMCore/RigVM.cpp
--- a/Engine/Source/Runtime/RigVM/Private/RigVMCore/RigVM.cpp
+++ b/Engine/Source/Runtime/RigVM/Private/RigVMCore/RigVM.cpp
@@ -14,7 +14,7 @@ void FRigVMParameter::Serialize(FArchive& Ar)
 		return;
 	}
 
-	if (Ar.IsSaving() || Ar.IsObjectReferenceCollector())
+	if (Ar.IsSaving() || Ar.IsObjectReferenceCollector() || Ar.IsCountingMemory())
 	{
 		Save(Ar);
 	}
@@ -93,7 +93,7 @@ void URigVM::Serialize(FArchive& Ar)
 
 	ensure(ExecutingThreadId == INDEX_NONE);
 
-	if (Ar.IsSaving() || Ar.IsObjectReferenceCollector())
+	if (Ar.IsSaving() || Ar.IsObjectReferenceCollector() || Ar.IsCountingMemory())
 	{
 		Save(Ar);
 	}
diff --git a/Engine/Source/Runtime/RigVM/Private/RigVMCore/RigVMByteCode.cpp b/Engine/Source/Runtime/RigVM/Private/RigVMCore/RigVMByteCode.cpp
--- a/Engine/Source/Runtime/RigVM/Private/RigVMCore/RigVMByteCode.cpp
+++ b/Engine/Source/Runtime/RigVM/Private/RigVMCore/RigVMByteCode.cpp
@@ -135,7 +135,7 @@ void FRigVMByteCode::Serialize(FArchive& Ar)
 		return;
 	}
 
-	if (Ar.IsSaving() || Ar.IsObjectReferenceCollector())
+	if (Ar.IsSaving() || Ar.IsObjectReferenceCollector() || Ar.IsCountingMemory())
 	{
 		Save(Ar);
 	}
diff --git a/Engine/Source/Runtime/RigVM/Private/RigVMCore/RigVMMemory.cpp b/Engine/Source/Runtime/RigVM/Private/RigVMCore/RigVMMemory.cpp
--- a/Engine/Source/Runtime/RigVM/Private/RigVMCore/RigVMMemory.cpp
+++ b/Engine/Source/Runtime/RigVM/Private/RigVMCore/RigVMMemory.cpp
@@ -42,7 +42,7 @@ void FRigVMRegister::Serialize(FArchive& Ar)
 	Ar << Name;
 	Ar << ScriptStructIndex;
 
-	if (Ar.IsSaving() || Ar.IsObjectReferenceCollector())
+	if (Ar.IsSaving() || Ar.IsObjectReferenceCollector() || Ar.IsCountingMemory())
 	{
 		Save(Ar);
 	}
@@ -95,7 +95,7 @@ void FRigVMRegisterOffset::Serialize(FArchive& Ar)
 		return;
 	}
 
-	if (Ar.IsSaving() || Ar.IsObjectReferenceCollector())
+	if (Ar.IsSaving() || Ar.IsObjectReferenceCollector() || Ar.IsCountingMemory())
 	{
 		Save(Ar);
 	}
@@ -513,7 +513,7 @@ void FRigVMMemoryContainer::Serialize(FArchive& Ar)
  		return;
  	}
 
-	if (Ar.IsSaving() || Ar.IsObjectReferenceCollector())
+	if (Ar.IsSaving() || Ar.IsObjectReferenceCollector() || Ar.IsCountingMemory())
 	{
 		Save(Ar);
 	}
