ControlRig: Copy values from RigVMExecuteContext to ControlRigExecuteContext

diff --git a/Engine/Plugins/Experimental/ControlRig/Source/ControlRig/Private/Units/Execution/RigUnit_BeginExecution.cpp b/Engine/Plugins/Experimental/ControlRig/Source/ControlRig/Private/Units/Execution/RigUnit_BeginExecution.cpp
--- a/Engine/Plugins/Experimental/ControlRig/Source/ControlRig/Private/Units/Execution/RigUnit_BeginExecution.cpp
+++ b/Engine/Plugins/Experimental/ControlRig/Source/ControlRig/Private/Units/Execution/RigUnit_BeginExecution.cpp
@@ -8,6 +8,7 @@ FName FRigUnit_BeginExecution::EventName = TEXT("Update");
 FRigUnit_BeginExecution_Execute()
 {
     DECLARE_SCOPE_HIERARCHICAL_COUNTER_RIGUNIT()
+	ExecuteContext.CopyFrom(RigVMExecuteContext);
 	ExecuteContext.Hierarchy = Context.Hierarchy;
 	ExecuteContext.EventName = FRigUnit_BeginExecution::EventName;
 }
diff --git a/Engine/Plugins/Experimental/ControlRig/Source/ControlRig/Private/Units/Execution/RigUnit_InverseExecution.cpp b/Engine/Plugins/Experimental/ControlRig/Source/ControlRig/Private/Units/Execution/RigUnit_InverseExecution.cpp
--- a/Engine/Plugins/Experimental/ControlRig/Source/ControlRig/Private/Units/Execution/RigUnit_InverseExecution.cpp
+++ b/Engine/Plugins/Experimental/ControlRig/Source/ControlRig/Private/Units/Execution/RigUnit_InverseExecution.cpp
@@ -8,6 +8,7 @@ FName FRigUnit_InverseExecution::EventName = TEXT("Inverse");
 FRigUnit_InverseExecution_Execute()
 {
     DECLARE_SCOPE_HIERARCHICAL_COUNTER_RIGUNIT()
+	ExecuteContext.CopyFrom(RigVMExecuteContext);
 	ExecuteContext.Hierarchy = Context.Hierarchy;
 	ExecuteContext.EventName = FRigUnit_InverseExecution::EventName;
 }
diff --git a/Engine/Plugins/Experimental/ControlRig/Source/ControlRig/Private/Units/Execution/RigUnit_PrepareForExecution.cpp b/Engine/Plugins/Experimental/ControlRig/Source/ControlRig/Private/Units/Execution/RigUnit_PrepareForExecution.cpp
--- a/Engine/Plugins/Experimental/ControlRig/Source/ControlRig/Private/Units/Execution/RigUnit_PrepareForExecution.cpp
+++ b/Engine/Plugins/Experimental/ControlRig/Source/ControlRig/Private/Units/Execution/RigUnit_PrepareForExecution.cpp
@@ -8,6 +8,7 @@ FName FRigUnit_PrepareForExecution::EventName = TEXT("Setup");
 FRigUnit_PrepareForExecution_Execute()
 {
     DECLARE_SCOPE_HIERARCHICAL_COUNTER_RIGUNIT()
+	ExecuteContext.CopyFrom(RigVMExecuteContext);
 	ExecuteContext.Hierarchy = Context.Hierarchy;
 	ExecuteContext.EventName = FRigUnit_PrepareForExecution::EventName;
 }
diff --git a/Engine/Source/Runtime/RigVM/Public/RigVMCore/RigVMExecuteContext.h b/Engine/Source/Runtime/RigVM/Public/RigVMCore/RigVMExecuteContext.h
--- a/Engine/Source/Runtime/RigVM/Public/RigVMCore/RigVMExecuteContext.h
+++ b/Engine/Source/Runtime/RigVM/Public/RigVMCore/RigVMExecuteContext.h
@@ -172,6 +172,17 @@ struct RIGVM_API FRigVMExecuteContext
 		ExternalVariables.Reset();
 	}
 
+	FORCEINLINE void CopyFrom(const FRigVMExecuteContext& Other)
+	{
+		EventName = Other.EventName;
+		FunctionName = Other.FunctionName;
+		InstructionIndex = Other.InstructionIndex;
+		ExternalVariables = Other.ExternalVariables;
+		OpaqueArguments = Other.OpaqueArguments;
+		Slices = Other.Slices;
+		SliceOffsets = Other.SliceOffsets;
+	}
+
 	FORCEINLINE const FRigVMSlice& GetSlice() const
 	{
 		const int32 SliceOffset = (int32)SliceOffsets[InstructionIndex];
