ControlRig Debugging: Add copy RigVM with register values

diff --git a/Engine/Source/Runtime/RigVM/Private/RigVMCore/RigVM.cpp b/Engine/Source/Runtime/RigVM/Private/RigVMCore/RigVM.cpp
--- a/Engine/Source/Runtime/RigVM/Private/RigVMCore/RigVM.cpp
+++ b/Engine/Source/Runtime/RigVM/Private/RigVMCore/RigVM.cpp
@@ -189,7 +189,7 @@ void URigVM::Empty()
 	CachedMemoryHandles.Empty();
 }
 
-void URigVM::CopyFrom(URigVM* InVM, bool bDeferCopy, bool bReferenceLiteralMemory, bool bReferenceByteCode, bool bCopyExternalVariables)
+void URigVM::CopyFrom(URigVM* InVM, bool bDeferCopy, bool bReferenceLiteralMemory, bool bReferenceByteCode, bool bCopyExternalVariables, bool bCopyDynamicRegisters)
 {
 	check(InVM);
 
@@ -205,7 +205,14 @@ void URigVM::CopyFrom(URigVM* InVM, bool bDeferCopy, bool bReferenceLiteralMemor
 
 	if(InVM->WorkMemoryPtr == &InVM->WorkMemoryStorage)
 	{
-		WorkMemoryStorage = InVM->WorkMemoryStorage;
+		if (bCopyDynamicRegisters)
+		{
+			WorkMemoryStorage.CopyRegisters(InVM->WorkMemoryStorage);
+		}
+		else
+		{
+			WorkMemoryStorage = InVM->WorkMemoryStorage;			
+		}
 		WorkMemoryPtr = &WorkMemoryStorage;
 	}
 	else
diff --git a/Engine/Source/Runtime/RigVM/Private/RigVMCore/RigVMMemory.cpp b/Engine/Source/Runtime/RigVM/Private/RigVMCore/RigVMMemory.cpp
--- a/Engine/Source/Runtime/RigVM/Private/RigVMCore/RigVMMemory.cpp
+++ b/Engine/Source/Runtime/RigVM/Private/RigVMCore/RigVMMemory.cpp
@@ -465,6 +465,19 @@ FRigVMMemoryContainer::~FRigVMMemoryContainer()
 	Reset();
 }
 
+bool FRigVMMemoryContainer::CopyRegisters(const FRigVMMemoryContainer &InOther)
+{
+	ensure(Registers.Num() == InOther.Registers.Num());
+	for (int32 Index = 0; Index < Registers.Num(); ++Index)
+	{
+		if (!Copy(Index, Index, &InOther))
+		{
+			return false;
+		}
+	}
+	return true;
+}
+
 FRigVMMemoryContainer& FRigVMMemoryContainer::operator= (const FRigVMMemoryContainer &InOther)
 {
 	Reset();
@@ -1191,6 +1204,17 @@ bool FRigVMMemoryContainer::Copy(
 		SourcePtr = ArrayStorage->GetData();
 	}
 
+	// TODO: preventing crash when TargetPtr has not been allocated 
+	if (TargetPtr == nullptr && Target.IsNestedDynamic())
+	{
+		uint8* Ptr = (uint8*)&Data[Target.GetWorkByteIndex()];
+		FRigVMNestedByteArray& Storage = *(FRigVMNestedByteArray*)Ptr;
+		Storage.SetNum(Target.SliceCount);
+		Storage[InTargetSliceIndex].SetNumZeroed(SourceElementCount * Source.ElementSize);
+		Copy(InTargetRegisterIndex, INDEX_NONE, Source.Type, SourcePtr, Storage[InTargetSliceIndex].GetData(), Source.ElementSize * SourceElementCount);
+		return true;
+	}
+
 	ERigVMRegisterType TargetType = Target.Type;
 	if (Target.IsDynamic() && !Target.bIsArray)
 	{
diff --git a/Engine/Source/Runtime/RigVM/Public/RigVMCore/RigVM.h b/Engine/Source/Runtime/RigVM/Public/RigVMCore/RigVM.h
--- a/Engine/Source/Runtime/RigVM/Public/RigVMCore/RigVM.h
+++ b/Engine/Source/Runtime/RigVM/Public/RigVMCore/RigVM.h
@@ -136,7 +136,7 @@ public:
 	void Empty();
 
 	// resets the container and clones the input VM
-	void CopyFrom(URigVM* InVM, bool bDeferCopy = false, bool bReferenceLiteralMemory = false, bool bReferenceByteCode = false, bool bCopyExternalVariables = false);
+	void CopyFrom(URigVM* InVM, bool bDeferCopy = false, bool bReferenceLiteralMemory = false, bool bReferenceByteCode = false, bool bCopyExternalVariables = false, bool bCopyDynamicRegisters = false);
 
 	// Initializes all execute ops and their memory.
 	bool Initialize(FRigVMMemoryContainerPtrArray Memory, FRigVMFixedArray<void*> AdditionalArguments);
diff --git a/Engine/Source/Runtime/RigVM/Public/RigVMCore/RigVMMemory.h b/Engine/Source/Runtime/RigVM/Public/RigVMCore/RigVMMemory.h
--- a/Engine/Source/Runtime/RigVM/Public/RigVMCore/RigVMMemory.h
+++ b/Engine/Source/Runtime/RigVM/Public/RigVMCore/RigVMMemory.h
@@ -554,6 +554,7 @@ public:
 	FRigVMMemoryContainer(const FRigVMMemoryContainer& Other);
 	~FRigVMMemoryContainer();
 
+	bool CopyRegisters(const FRigVMMemoryContainer &InOther);
 	FRigVMMemoryContainer& operator= (const FRigVMMemoryContainer &InOther);
 
 	// returns the memory type of this container
