ControlRig: Fix warning in RigVMMemory copy registers

diff --git a/Engine/Source/Runtime/RigVM/Private/RigVMCore/RigVMMemory.cpp b/Engine/Source/Runtime/RigVM/Private/RigVMCore/RigVMMemory.cpp
--- a/Engine/Source/Runtime/RigVM/Private/RigVMCore/RigVMMemory.cpp
+++ b/Engine/Source/Runtime/RigVM/Private/RigVMCore/RigVMMemory.cpp
@@ -1221,7 +1221,7 @@ bool FRigVMMemoryContainer::Copy(
 		FRigVMByteArray* ArrayStorage = (FRigVMByteArray*)TargetPtr;
 		TargetPtr = ArrayStorage->GetData();
 	}
-	else if (Target.IsNestedDynamic())
+	else if (TargetPtr && Target.IsNestedDynamic())
 	{
 #if WITH_EDITOR
 		ensure(InTargetRegisterOffset == INDEX_NONE);
