Control Rig: Added comments RigVMMemory as per review with Sara

diff --git a/Engine/Source/Runtime/RigVM/Private/RigVMCore/RigVMMemory.cpp b/Engine/Source/Runtime/RigVM/Private/RigVMCore/RigVMMemory.cpp
--- a/Engine/Source/Runtime/RigVM/Private/RigVMCore/RigVMMemory.cpp
+++ b/Engine/Source/Runtime/RigVM/Private/RigVMCore/RigVMMemory.cpp
@@ -368,7 +368,7 @@ bool FRigVMRegisterOffset::ContainsArraySegment() const
 {
 	for (int32 SegmentIndex : Segments)
 	{
-		if (SegmentIndex < 0)
+		if (SegmentIndex < 0) // Segement indices less than zero represent array element offsets.
 		{
 			return true;
 		}
@@ -508,6 +508,8 @@ FRigVMMemoryContainer& FRigVMMemoryContainer::operator= (const FRigVMMemoryConta
 
 		Construct(Index);
 
+		// Why was this introduced? Are we trying to save time here?
+		// An assignment should be copying everything.
 		if (MemoryType == ERigVMMemoryType::Literal)
 		{
 			Copy(Index, Index, &InOther);
diff --git a/Engine/Source/Runtime/RigVM/Public/RigVMCore/RigVMMemory.h b/Engine/Source/Runtime/RigVM/Public/RigVMCore/RigVMMemory.h
--- a/Engine/Source/Runtime/RigVM/Public/RigVMCore/RigVMMemory.h
+++ b/Engine/Source/Runtime/RigVM/Public/RigVMCore/RigVMMemory.h
@@ -43,6 +43,8 @@ enum class ERigVMMemoryType: uint8
  * memory all the way to the actual byte address in memory.
  * The FRigVMOperand is a light weight address for a register in
  * a FRigVMMemoryContainer.
+ * For external variables the register index represents the
+ * index of the external variable within the running VM.
  */
 USTRUCT()
 struct RIGVM_API FRigVMOperand
@@ -106,6 +108,10 @@ private:
 	UPROPERTY()
 	ERigVMMemoryType MemoryType;
 
+	/**
+	 * The index of the register inside of the specific type of memory (work, literal etc).
+	 * For external variables the register index represents the index of the external variable within the running VM.
+	 */
 	UPROPERTY()
 	uint16 RegisterIndex;
 	
@@ -123,10 +129,10 @@ typedef FRigVMFixedArray<FRigVMOperand> FRigVMOperandArray;
 UENUM()
 enum class ERigVMRegisterType : uint8
 {
-	Plain, // bool, int32, float, FVector etc.
+	Plain, // bool, int32, float, FVector etc. (also structs that do NOT require a constructor / destructor to be valid)
 	String, // FString
 	Name, // FName
-	Struct, // Any USTRUCT
+	Struct, // Any USTRUCT (structs which require a constructor / destructor to be valid. structs which have indirection (arrays or strings in them, for example).
 	Invalid
 };
 
@@ -172,24 +178,28 @@ struct RIGVM_API FRigVMRegister
 	UPROPERTY()
 	uint32 ByteIndex;
 
-	// The size of each store element
+	// The size of each store element (for example 4 for a float)
 	UPROPERTY()
 	uint16 ElementSize;
 
-	// The number of elements in this register
+	// The number of elements in this register (for example the number of elements in an array)
 	UPROPERTY()
 	uint16 ElementCount;
 
-	// The number of slices (complete copies)
+	// The number of slices (complete copies) (for example the number of iterations on a fixed loop)
+	// Potentially redundant state - can be removed.
 	UPROPERTY()
 	uint16 SliceCount;
 
-	// The number of leading bytes for alignment
+	// The number of leading bytes for alignment.
+	// Bytes that had to be introduced before the register
+	// to align the registers memory as per platform specification.
 	UPROPERTY()
 	uint8 AlignmentBytes;
 
 	// The number of trailing bytes.
 	// These originate after shrinking a register.
+	// Potentially unused state - can be removed.
 	UPROPERTY()
 	uint16 TrailingBytes;
 
@@ -206,16 +216,21 @@ struct RIGVM_API FRigVMRegister
 	UPROPERTY()
 	bool bIsArray;
 
-	// If true defines this register to use dynamic storage
+	// If true defines this register to use dynamic storage.
+	// Is this an array or singleton value with multiple slices
+	// which potentially requires changing count at runtime. 
 	UPROPERTY()
 	bool bIsDynamic;
 
 #if WITH_EDITORONLY_DATA
 
 	// Defines the CPP type used for the register
+	// This is only used for debugging purposes in editor.
 	UPROPERTY(transient)
 	FName BaseCPPType;
 
+	// The resolved UScriptStruct / UObject (in the future)
+	// used for debugging.
 	UPROPERTY(transient)
 	UObject* BaseCPPTypeObject;
 
@@ -258,6 +273,7 @@ struct RIGVM_API FRigVMRegister
 	FORCEINLINE_DEBUGGABLE bool IsArray() const { return bIsArray || (ElementCount > 1); }
 
 	// Returns true if the register stores shallow memory
+	// Potentially unused (to be removed)
 	FORCEINLINE_DEBUGGABLE bool IsShallow() const { return ScriptStructIndex == INDEX_NONE; }
 
 	// Returns the number of allocated bytes (including alignment + trailing bytes)
@@ -279,6 +295,7 @@ typedef FRigVMFixedArray<FRigVMRegister> FRigVMRegisterArray;
 // The register offset represents a memory offset within a register's memory.
 // This can be used to represent memory addresses of array elements within
 // a struct, for example.
+// A register offset's path can look like MyTransformStruct.Transforms[3].Translation.X
 USTRUCT()
 struct RIGVM_API FRigVMRegisterOffset
 {
@@ -334,27 +351,40 @@ public:
 
 private:
 
+	// Segments represent the memory offset(s) to use when accessing the target memory.
+	// In case of indirection in the source memory each memory offset / jump is stored.
+	// So for example: If you are accessing the third Rotation.X of an array of transforms,
+	// the segments would read as: [-2, 12] (second array element, and the fourth float)
+	// Segment indices less than zero represent array element offsets, while positive numbers
+	// represents jumps within a struct.
 	UPROPERTY()
 	TArray<int32> Segments;
 
+	// Type of resulting register (for example Plain for Transform.Translation.X)
 	UPROPERTY()
 	ERigVMRegisterType Type;
 
+	// The c++ type of the resulting memory (for example float for Transform.Translation.X)
 	UPROPERTY()
 	FName CPPType;
 
+	// The c++ script struct of the resulting memory (for example FVector for Transform.Translation)
 	UPROPERTY()
 	UScriptStruct* ScriptStruct;
 
+	// The c++ script struct of the source memory (for example FTransform for Transform.Translation)
 	UPROPERTY()
 	UScriptStruct* ParentScriptStruct;
 
+	// The index of the element within an array (for example 3 for Transform[3])
 	UPROPERTY()
 	int32 ArrayIndex;
 
+	// The number of bytes of the resulting memory (for example 4 (float) for Transform.Translation.X)
 	UPROPERTY()
 	uint16 ElementSize;
 
+	// The cached path of the segments within this register, for example FTransform.Translation.X
 	UPROPERTY()
 	FString CachedSegmentPath;
 
@@ -368,11 +398,13 @@ private:
 struct FRigVMMemoryHandle
 {
 public:
+
+	// The type of memory being accessed.
 	enum FType
 	{
-		Plain,
-		Dynamic,
-		NestedDynamic
+		Plain, // Shallow memory with no indirection (for example FTransform.Translation.X)
+		Dynamic, // Memory representing an array
+		NestedDynamic // Memory representing an array or single element with slices (for example within loop)
 	};
 
 	FORCEINLINE_DEBUGGABLE FRigVMMemoryHandle()
@@ -423,21 +455,25 @@ public:
 		, RegisterOffset(InRegisterOffset)
 	{}
 
+	// getter for the contained const memory
 	FORCEINLINE_DEBUGGABLE operator const uint8*() const
 	{
 		return GetData();
 	}
 
+	// getter for the contained mutable memory
 	FORCEINLINE_DEBUGGABLE operator uint8*()
 	{
 		return GetData();
 	}
 
+	// returns the contained const memory for a given slice index
 	FORCEINLINE_DEBUGGABLE const uint8* GetData(int32 SliceIndex = 0, bool bGetArrayData = false) const
 	{
 		return GetData_Internal(SliceIndex, bGetArrayData);
 	}
 
+	// returns the contained mutable memory for a given slice index
 	FORCEINLINE_DEBUGGABLE uint8* GetData(int32 SliceIndex = 0, bool bGetArrayData = false)
 	{
 		return GetData_Internal(SliceIndex, bGetArrayData);
@@ -542,6 +578,13 @@ private:
  *      TArray<float> MyArray = {3.f, 4.f, 5.f};
  * 		int32 Index = Container.AddFixedArray<float>(MyArray);
  *      FRigVMFixedArray<float> ArrayView = Container.GetFixedArray<float>(Index);
+ *
+ *
+ * Registers can store dynamically resizable memory as well by reyling on indirection.
+ * 
+ * Arrays with a single slice are going to be stored as a FRigVMByteArray, [unit8], so for example for a float array in C++ it would be (TArray<float> Param).
+ * Single values with multiple slices are also going to be stored as a FRigVMByteArray, [unit8], so for example for a float array in C++ it would be (float Param). 
+ * Arrays with multiple slices are going to be stored as a FRigVMNestedByteArray, [[unit8]], so for example for a float array in C++ it would be (TArray<float> Param).
  */
 USTRUCT()
 struct RIGVM_API FRigVMMemoryContainer
