Control Rig RigVMMemory: Cleanup unused methods

diff --git a/Engine/Source/Runtime/RigVM/Private/RigVMCore/RigVMMemory.cpp b/Engine/Source/Runtime/RigVM/Private/RigVMCore/RigVMMemory.cpp
--- a/Engine/Source/Runtime/RigVM/Private/RigVMCore/RigVMMemory.cpp
+++ b/Engine/Source/Runtime/RigVM/Private/RigVMCore/RigVMMemory.cpp
@@ -1623,37 +1623,6 @@ bool FRigVMMemoryContainer::Destroy(int32 InRegisterIndex, int32 InElementIndex,
 	return true;
 }
 
-FName FRigVMMemoryContainer::Rename(int32 InRegisterIndex, const FName& InNewName)
-{
-	if (Registers[InRegisterIndex].Name == InNewName)
-	{
-		return Registers[InRegisterIndex].Name;
-	}
-
-	if (!IsNameAvailable(InNewName))
-	{
-		return Registers[InRegisterIndex].Name;
-	}
-
-	Registers[InRegisterIndex].Name = InNewName;
-	UpdateRegisters();
-
-	return InNewName;
-}
-
-FName FRigVMMemoryContainer::Rename(const FName& InOldName, const FName& InNewName)
-{
-	ensure(bUseNameMap);
-
-	int32 RegisterIndex = GetIndex(InOldName);
-	if (RegisterIndex == INDEX_NONE)
-	{
-		return NAME_None;
-	}
-
-	return Rename(RegisterIndex, InNewName);
-}
-
 int32 FRigVMMemoryContainer::GetOrAddRegisterOffset(int32 InRegisterIndex, int32 InArrayElement)
 {
 	return GetOrAddRegisterOffset(InRegisterIndex, FString(), InArrayElement);
diff --git a/Engine/Source/Runtime/RigVM/Public/RigVMCore/RigVMMemory.h b/Engine/Source/Runtime/RigVM/Public/RigVMCore/RigVMMemory.h
--- a/Engine/Source/Runtime/RigVM/Public/RigVMCore/RigVMMemory.h
+++ b/Engine/Source/Runtime/RigVM/Public/RigVMCore/RigVMMemory.h
@@ -272,10 +272,6 @@ struct RIGVM_API FRigVMRegister
 	// Returns true if the register stores more than one element
 	FORCEINLINE_DEBUGGABLE bool IsArray() const { return bIsArray || (ElementCount > 1); }
 
-	// Returns true if the register stores shallow memory
-	// Potentially unused (to be removed)
-	FORCEINLINE_DEBUGGABLE bool IsShallow() const { return ScriptStructIndex == INDEX_NONE; }
-
 	// Returns the number of allocated bytes (including alignment + trailing bytes)
 	FORCEINLINE_DEBUGGABLE uint16 GetAllocatedBytes() const { return ElementCount * ElementSize * SliceCount + (uint16)AlignmentBytes + TrailingBytes; }
 
@@ -1180,14 +1176,6 @@ public:
 		return AddDynamicValue<T>(NAME_None, InValue, InSliceCount);
 	}
 
-	// Remove a register given its index
-	// Note: This only works if SupportsNames() == true
-	FName Rename(int32 InRegisterIndex, const FName& InNewName);
-
-	// Remove a register given its old name
-	// Note: This only works if SupportsNames() == true
-	FName Rename(const FName& InOldName, const FName& InNewName);
-
 	// Adds a register path and returns its index
 	int32 GetOrAddRegisterOffset(int32 InRegisterIndex, int32 InArrayElement = 0);
 
