Control Rig: Avoid serializing deep into memory container during Memory Counting

diff --git a/Engine/Source/Runtime/RigVM/Private/RigVMCore/RigVMMemory.cpp b/Engine/Source/Runtime/RigVM/Private/RigVMCore/RigVMMemory.cpp
--- a/Engine/Source/Runtime/RigVM/Private/RigVMCore/RigVMMemory.cpp
+++ b/Engine/Source/Runtime/RigVM/Private/RigVMCore/RigVMMemory.cpp
@@ -562,7 +562,6 @@ void FRigVMMemoryContainer::Save(FArchive& Ar)
 	Ar << MemoryType;
 	Ar << Registers;
 	Ar << RegisterOffsets;
-
 	
 	TArray<FString> ScriptStructPaths;
 	for (UScriptStruct* ScriptStruct : ScriptStructs)
@@ -574,6 +573,15 @@ void FRigVMMemoryContainer::Save(FArchive& Ar)
 	uint64 TotalBytes = Data.Num();
 	Ar << TotalBytes;
 
+	// if we are running a memory count
+	// there's no need to do a deep serialize
+	// since we know the memory needed
+	if(Ar.IsCountingMemory())
+	{
+		Ar << Data;
+		return;
+	}
+
 	for (FRigVMRegister& Register : Registers)
 	{
 		if (Register.ElementCount == 0 && !Register.IsDynamic())
