Control Rig: Moving debug memory to its own container

diff --git a/Engine/Plugins/Experimental/ControlRig/Source/ControlRig/Private/ControlRig.cpp b/Engine/Plugins/Experimental/ControlRig/Source/ControlRig/Private/ControlRig.cpp
--- a/Engine/Plugins/Experimental/ControlRig/Source/ControlRig/Private/ControlRig.cpp
+++ b/Engine/Plugins/Experimental/ControlRig/Source/ControlRig/Private/ControlRig.cpp
@@ -958,13 +958,13 @@ void UControlRig::ExecuteUnits(FRigUnitContext& InOutContext, const FName& InEve
 {
 	if (VM)
 	{
-		FRigVMMemoryContainer* LocalMemory[] = { VM->WorkMemoryPtr, VM->LiteralMemoryPtr };
+		FRigVMMemoryContainer* LocalMemory[] = { VM->WorkMemoryPtr, VM->LiteralMemoryPtr, VM->DebugMemoryPtr };
 		TArray<void*> AdditionalArguments;
 		AdditionalArguments.Add(&InOutContext);
 
 		if (InOutContext.State == EControlRigState::Init)
 		{
-			VM->Initialize(FRigVMMemoryContainerPtrArray(LocalMemory, 2), AdditionalArguments);
+			VM->Initialize(FRigVMMemoryContainerPtrArray(LocalMemory, 3), AdditionalArguments);
 		}
 		else
 		{
@@ -981,7 +981,7 @@ void UControlRig::ExecuteUnits(FRigUnitContext& InOutContext, const FName& InEve
 				}
 			}
 #endif
-			VM->Execute(FRigVMMemoryContainerPtrArray(LocalMemory, 2), AdditionalArguments, InEventName);
+			VM->Execute(FRigVMMemoryContainerPtrArray(LocalMemory, 3), AdditionalArguments, InEventName);
 		}
 	}
 }
diff --git a/Engine/Plugins/Experimental/ControlRig/Source/ControlRigDeveloper/Private/ControlRigBlueprint.cpp b/Engine/Plugins/Experimental/ControlRig/Source/ControlRigDeveloper/Private/ControlRigBlueprint.cpp
--- a/Engine/Plugins/Experimental/ControlRig/Source/ControlRigDeveloper/Private/ControlRigBlueprint.cpp
+++ b/Engine/Plugins/Experimental/ControlRig/Source/ControlRigDeveloper/Private/ControlRigBlueprint.cpp
@@ -2563,6 +2563,25 @@ void UControlRigBlueprint::HandleModifiedEvent(ERigVMGraphNotifType InNotifType,
 				break;
 			}
 			case ERigVMGraphNotifType::PinWatchedChanged:
+			{
+				if (UControlRig* CR = Cast<UControlRig>(GetObjectBeingDebugged()))
+				{
+					URigVMPin* Pin = CastChecked<URigVMPin>(InSubject)->GetRootPin(); 
+					URigVMCompiler* Compiler = URigVMCompiler::StaticClass()->GetDefaultObject<URigVMCompiler>();
+					Compiler->Settings = VMCompileSettings;
+					TSharedPtr<FRigVMParserAST> RuntimeAST = Model->GetRuntimeAST();
+					
+					if(Pin->RequiresWatch())
+					{
+						Compiler->CreateDebugRegister(Pin, CR->GetVM(), &PinToOperandMap, RuntimeAST);
+					}
+					else
+					{
+						Compiler->RemoveDebugRegister(Pin, CR->GetVM(), &PinToOperandMap, RuntimeAST);
+					}
+				}
+				// break; fall through
+			}
 			case ERigVMGraphNotifType::PinTypeChanged:
 			case ERigVMGraphNotifType::PinIndexChanged:
 			{
@@ -2584,6 +2603,11 @@ void UControlRigBlueprint::HandleModifiedEvent(ERigVMGraphNotifType InNotifType,
 									{
 										WatchedPins.Remove(EdPin);
 									}
+
+									if(InNotifType == ERigVMGraphNotifType::PinWatchedChanged)
+									{
+										return;
+									}
 									RequestAutoVMRecompilation();
 									MarkPackageDirty();
 								}
diff --git a/Engine/Plugins/Experimental/ControlRig/Source/ControlRigEditor/Private/Graph/SControlRigGraphNode.cpp b/Engine/Plugins/Experimental/ControlRig/Source/ControlRigEditor/Private/Graph/SControlRigGraphNode.cpp
--- a/Engine/Plugins/Experimental/ControlRig/Source/ControlRigEditor/Private/Graph/SControlRigGraphNode.cpp
+++ b/Engine/Plugins/Experimental/ControlRig/Source/ControlRigEditor/Private/Graph/SControlRigGraphNode.cpp
@@ -1132,9 +1132,7 @@ void SControlRigGraphNode::GetNodeInfoPopups(FNodeInfoContext* Context, TArray<F
 						FString PinHash = URigVMCompiler::GetPinHash(ModelPin, nullptr, true);
 						if (const FRigVMOperand* WatchOperand = RigBlueprint->PinToOperandMap.Find(PinHash))
 						{
-							FRigVMMemoryContainer& Memory = 
-								WatchOperand->GetMemoryType() == ERigVMMemoryType::Literal ?
-								ActiveObject->GetVM()->GetLiteralMemory() : ActiveObject->GetVM()->GetWorkMemory();
+							FRigVMMemoryContainer& Memory = ActiveObject->GetVM()->GetDebugMemory();
 
 							TArray<FString> DefaultValues = Memory.GetRegisterValueAsString(*WatchOperand, ModelPin->GetCPPType(), ModelPin->GetCPPTypeObject());
 							if (DefaultValues.Num() == 1)
diff --git a/Engine/Source/Developer/RigVMDeveloper/Private/RigVMCompiler/RigVMAST.cpp b/Engine/Source/Developer/RigVMDeveloper/Private/RigVMCompiler/RigVMAST.cpp
--- a/Engine/Source/Developer/RigVMDeveloper/Private/RigVMCompiler/RigVMAST.cpp
+++ b/Engine/Source/Developer/RigVMDeveloper/Private/RigVMCompiler/RigVMAST.cpp
@@ -1842,11 +1842,11 @@ bool FRigVMParserAST::FoldConstantValuesToLiterals(URigVMGraph* InGraph, URigVMC
 
 	TempCompiler->Compile(InGraph, InController, TempVM, InExternalVariables, InRigVMUserData, &Operands, TempAST);
 
-	FRigVMMemoryContainer* Memory[] = { TempVM->WorkMemoryPtr, TempVM->LiteralMemoryPtr };
+	FRigVMMemoryContainer* Memory[] = { TempVM->WorkMemoryPtr, TempVM->LiteralMemoryPtr, TempVM->DebugMemoryPtr };
 
 	for (const FRigVMUserDataArray& RigVMUserData : InRigVMUserData)
 	{
-		TempVM->Execute(FRigVMMemoryContainerPtrArray(Memory, 2), RigVMUserData);
+		TempVM->Execute(FRigVMMemoryContainerPtrArray(Memory, 3), RigVMUserData);
 	}
 
 	// copy the values out of the temp VM and set them on the cached value
@@ -2169,6 +2169,21 @@ const FRigVMExprAST* FRigVMParserAST::GetExprForSubject(const FRigVMASTProxy& In
 	return nullptr;
 }
 
+TArray<const FRigVMExprAST*> FRigVMParserAST::GetExpressionsForSubject(UObject* InSubject) const
+{
+	TArray<const FRigVMExprAST*> ExpressionsForSubject;
+	
+	for (TPair<FRigVMASTProxy, FRigVMExprAST*> Pair : SubjectToExpression)
+	{
+		if(Pair.Key.GetCallstack().Last() == InSubject)
+		{
+			ExpressionsForSubject.Add(Pair.Value);
+		}
+	}
+
+	return ExpressionsForSubject;
+}
+
 void FRigVMParserAST::PrepareCycleChecking(URigVMPin* InPin)
 {
 	if (InPin == nullptr)
diff --git a/Engine/Source/Developer/RigVMDeveloper/Private/RigVMCompiler/RigVMASTProxy.cpp b/Engine/Source/Developer/RigVMDeveloper/Private/RigVMCompiler/RigVMASTProxy.cpp
--- a/Engine/Source/Developer/RigVMDeveloper/Private/RigVMCompiler/RigVMASTProxy.cpp
+++ b/Engine/Source/Developer/RigVMDeveloper/Private/RigVMCompiler/RigVMASTProxy.cpp
@@ -45,6 +45,15 @@ int32 FRigVMCallstack::Num() const
 	return Stack.Num();
 }
 
+const UObject* FRigVMCallstack::Last() const
+{
+	if(Stack.IsEmpty())
+	{
+		return nullptr;
+	}
+	return Stack.Last();
+}
+
 const UObject* FRigVMCallstack::operator[](int32 InIndex) const
 {
 	return Stack[InIndex];
diff --git a/Engine/Source/Developer/RigVMDeveloper/Private/RigVMCompiler/RigVMCompiler.cpp b/Engine/Source/Developer/RigVMDeveloper/Private/RigVMCompiler/RigVMCompiler.cpp
--- a/Engine/Source/Developer/RigVMDeveloper/Private/RigVMCompiler/RigVMCompiler.cpp
+++ b/Engine/Source/Developer/RigVMDeveloper/Private/RigVMCompiler/RigVMCompiler.cpp
@@ -264,6 +264,22 @@ bool URigVMCompiler::Compile(URigVMGraph* InGraph, URigVMController* InControlle
 
 	WorkData.VM->GetByteCode().AlignByteCode();
 
+	// setup debug registers
+	for(int32 GraphIndex=0; GraphIndex<VisitedGraphs.Num(); GraphIndex++)
+	{
+		URigVMGraph* VisitedGraph = VisitedGraphs[GraphIndex];
+		for(URigVMNode* ModelNode : VisitedGraph->GetNodes())
+		{
+			for(URigVMPin* ModelPin : ModelNode->GetPins())
+			{
+				if(ModelPin->RequiresWatch(false))
+				{
+					CreateDebugRegister(ModelPin, WorkData.VM, WorkData.PinPathToOperand, AST);
+				}
+			}
+		}
+	}
+
 	UE_LOG_RIGVMMEMORY(TEXT("RigVMCompiler: Finished '%s'."), *InGraph->GetPathName());
 
 	return true;
@@ -520,39 +536,6 @@ int32 URigVMCompiler::TraverseCallExtern(const FRigVMCallExternExprAST* InExpr,
 			const FString CallPath = InExpr->GetProxy().GetCallstack().GetCallPath();
 			WorkData.VM->GetByteCode().SetSubject(InstructionIndex, UnitNode, CallPath);
 		}
-
-		ensure(InExpr->NumChildren() == UnitNode->Pins.Num());
-		for (int32 PinIndex = 0; PinIndex < UnitNode->Pins.Num(); PinIndex++)
-		{
-			URigVMPin* Pin = UnitNode->Pins[PinIndex];
-
-			// ensure to copy the debug values
-			if (Pin->RequiresWatch(true) && Settings.EnablePinWatches)
-			{
-				const FRigVMExprAST* PinExpr = InExpr->ChildAt(PinIndex);
-				if (PinExpr->IsA(FRigVMExprAST::EType::CachedValue))
-				{
-					PinExpr = PinExpr->To<FRigVMCachedValueExprAST>()->GetVarExpr();
-				}
-				ensure(PinExpr->IsA(FRigVMExprAST::EType::Var));
-				FRigVMOperand Source = WorkData.ExprToOperand.FindChecked(PinExpr->To<FRigVMVarExprAST>());
-				FString PinHash = GetPinHash(Pin, PinExpr->To<FRigVMVarExprAST>(), true);
-				if (const FRigVMOperand* Target = WorkData.PinPathToOperand->Find(PinHash))
-				{
-					FRigVMOperand SourceRootOperand(Source.GetMemoryType(), Source.GetRegisterIndex(), INDEX_NONE);
-					WorkData.VM->GetByteCode().AddCopyOp(SourceRootOperand, *Target);
-					const int32 CopyOpInstructionIndex = WorkData.VM->GetByteCode().GetNumInstructions() - 1;
-
-					if (Settings.SetupNodeInstructionIndex)
-					{
-						const FRigVMVarExprAST* VarExpr = PinExpr->To<FRigVMVarExprAST>();
-						const FRigVMASTProxy Proxy = VarExpr->GetProxy().GetSibling(VarExpr->GetPin()->GetNode());
-						const FString CallPath = Proxy.GetCallstack().GetCallPath();
-						WorkData.VM->GetByteCode().SetSubject(CopyOpInstructionIndex, Proxy.GetSubject(), CallPath);
-					}
-				}
-			}
-		}
 	}
 
 	return InstructionIndex;
@@ -789,7 +772,8 @@ void URigVMCompiler::TraverseAssign(const FRigVMAssignExprAST* InExpr, FRigVMCom
 					}
 					else
 					{
-						FRigVMMemoryContainer& Memory = Operand.GetMemoryType() == ERigVMMemoryType::Literal ? VM->GetLiteralMemory() : VM->GetWorkMemory();
+						FRigVMMemoryContainer& Memory = Operand.GetMemoryType() == ERigVMMemoryType::Literal ? VM->GetLiteralMemory() :
+							(Operand.GetMemoryType() == ERigVMMemoryType::Work ? VM->GetWorkMemory() : VM->GetDebugMemory());
 						Operand = Memory.GetOperand(Operand.GetRegisterIndex(), SegmentPath, ArrayIndex);
 					}
 				}
@@ -827,48 +811,6 @@ void URigVMCompiler::TraverseAssign(const FRigVMAssignExprAST* InExpr, FRigVMCom
 				}
 			}
 		}
-
-		if (Cast<URigVMVariableNode>(SourceExpr->GetPin()->GetNode()))
-		{
-			if (SourceExpr->GetPin()->RequiresWatch(true) && Settings.EnablePinWatches)
-			{
-				FRigVMOperand WatchOperand = FindOrAddRegister(SourceExpr, WorkData, true /*debug watch*/);
-				if (WatchOperand.IsValid())
-				{
-					FRigVMOperand SourceRootOperand(Source.GetMemoryType(), Source.GetRegisterIndex(), INDEX_NONE);
-					WorkData.VM->GetByteCode().AddCopyOp(SourceRootOperand, WatchOperand);
-					int32 CopyInstructionIndex = WorkData.VM->GetByteCode().GetNumInstructions() - 1;
-
-					if (Settings.SetupNodeInstructionIndex)
-					{
-						URigVMNode* Node = SourceExpr->GetPin()->GetNode();
-						const FString CallPath = SourceExpr->GetProxy().GetSibling(Node).GetCallstack().GetCallPath();
-						WorkData.VM->GetByteCode().SetSubject(CopyInstructionIndex, Node, CallPath);
-					}
-				}
-			}
-		}
-
-		if (Cast<URigVMVariableNode>(TargetExpr->GetPin()->GetNode()))
-		{
-			if (TargetExpr->GetPin()->RequiresWatch(true) && Settings.EnablePinWatches)
-			{
-				FRigVMOperand WatchOperand = FindOrAddRegister(TargetExpr, WorkData, true /*debug watch*/);
-				if (WatchOperand.IsValid())
-				{
-					FRigVMOperand TargetRootOperand(Target.GetMemoryType(), Target.GetRegisterIndex(), INDEX_NONE);
-					WorkData.VM->GetByteCode().AddCopyOp(TargetRootOperand, WatchOperand);
-					int32 CopyInstructionIndex = WorkData.VM->GetByteCode().GetNumInstructions() - 1;
-
-					if (Settings.SetupNodeInstructionIndex)
-					{
-						URigVMNode* Node = TargetExpr->GetPin()->GetNode();
-						const FString CallPath = TargetExpr->GetProxy().GetSibling(Node).GetCallstack().GetCallPath();
-						WorkData.VM->GetByteCode().SetSubject(CopyInstructionIndex, Node, CallPath);
-					}
-				}
-			}
-		}
 	}
 }
 
@@ -1271,6 +1213,83 @@ FString URigVMCompiler::GetPinHash(const URigVMPin* InPin, const FRigVMVarExprAS
 	return FString::Printf(TEXT("%s%s%s"), *Prefix, *InPin->GetPinPath(true /*  full node path */), *Suffix);
 }
 
+void URigVMCompiler::CreateDebugRegister(URigVMPin* InPin, URigVM* OutVM,
+	TMap<FString, FRigVMOperand>* OutOperands, TSharedPtr<FRigVMParserAST> InRuntimeAST)
+{
+	check(InPin);
+	check(OutVM);
+	check(OutOperands);
+	check(InRuntimeAST.IsValid());
+	
+	FRigVMCompilerWorkData WorkData;
+	WorkData.VM = OutVM;
+	WorkData.PinPathToOperand = OutOperands;
+
+	URigVMPin* Pin = InPin->GetRootPin();
+	URigVMPin* SourcePin = Pin;
+	if(Settings.ASTSettings.bFoldAssignments)
+	{
+		while(SourcePin->GetSourceLinks().Num() > 0)
+		{
+			SourcePin = SourcePin->GetSourceLinks()[0]->GetSourcePin();
+		}
+	}
+	
+	TArray<const FRigVMExprAST*> Expressions = InRuntimeAST->GetExpressionsForSubject(SourcePin);
+	for(const FRigVMExprAST* Expression : Expressions)
+	{
+		check(Expression->IsA(FRigVMExprAST::EType::Var));
+		const FRigVMVarExprAST* VarExpression = Expression->To<FRigVMVarExprAST>();
+		const FString PinHash = GetPinHash(SourcePin, VarExpression, false);
+		if(const FRigVMOperand* Operand = OutOperands->Find(PinHash))
+		{
+			const FRigVMASTProxy PinProxy = FRigVMASTProxy::MakeFromUObject(Pin);
+			FRigVMVarExprAST TempVarExpr(FRigVMExprAST::EType::Var, PinProxy);
+			TempVarExpr.ParserPtr = InRuntimeAST.Get();
+	
+			const FRigVMOperand DebugOperand = FindOrAddRegister(&TempVarExpr, WorkData, true);
+			if(DebugOperand.IsValid())
+			{
+				FRigVMOperand KeyOperand(Operand->GetMemoryType(), Operand->GetRegisterIndex()); // no register offset
+				OutVM->OperandToDebugRegisters.FindOrAdd(KeyOperand).AddUnique(DebugOperand);
+			}
+		}
+	}
+}
+
+void URigVMCompiler::RemoveDebugRegister(URigVMPin* InPin, URigVM* OutVM,
+	TMap<FString, FRigVMOperand>* OutOperands, TSharedPtr<FRigVMParserAST> InRuntimeAST)
+{
+	check(OutVM);
+	check(OutOperands);
+
+	URigVMPin* Pin = InPin->GetRootPin();
+	const FRigVMASTProxy PinProxy = FRigVMASTProxy::MakeFromUObject(Pin);
+	FRigVMVarExprAST TempVarExpr(FRigVMExprAST::EType::Var, PinProxy);
+	TempVarExpr.ParserPtr = InRuntimeAST.Get();
+
+	const FString PinHash = GetPinHash(Pin, &TempVarExpr, true);
+	if(const FRigVMOperand* DebugOperand = OutOperands->Find(PinHash))
+	{
+		TArray<FRigVMOperand> KeysToRemove;
+		for(TPair<FRigVMOperand, TArray<FRigVMOperand>>& Pair : OutVM->OperandToDebugRegisters)
+		{
+			if(Pair.Value.Remove(*DebugOperand) > 0)
+			{
+				if(Pair.Value.IsEmpty())
+				{
+					KeysToRemove.AddUnique(Pair.Key);
+				}
+			}
+		}
+		for(const FRigVMOperand& KeyToRemove : KeysToRemove)
+		{
+			OutVM->OperandToDebugRegisters.Remove(KeyToRemove);
+		}
+		OutOperands->Remove(PinHash);
+	}
+}
+
 UScriptStruct* URigVMCompiler::GetScriptStructForCPPType(const FString& InCPPType)
 {
 	if (InCPPType == TEXT("FRotator"))
@@ -1397,15 +1416,22 @@ FRigVMOperand URigVMCompiler::FindOrAddRegister(const FRigVMVarExprAST* InVarExp
 	bool bIsLiteral = InVarExpr->GetType() == FRigVMExprAST::EType::Literal && !bIsDebugValue;
 	bool bIsParameter = InVarExpr->IsGraphParameter();
 
+	FRigVMMemoryContainer& Memory = 
+	    bIsLiteral ? WorkData.VM->GetLiteralMemory() :
+			(bIsDebugValue ? WorkData.VM->GetDebugMemory() : WorkData.VM->GetWorkMemory());
+
 	FRigVMOperand const* ExistingOperand = WorkData.PinPathToOperand->Find(Hash);
 	if (ExistingOperand)
 	{
-		if (!bIsDebugValue)
+		if(ExistingOperand->GetMemoryType() == Memory.GetMemoryType())
 		{
-			check(!WorkData.ExprToOperand.Contains(InVarExpr));
-			WorkData.ExprToOperand.Add(InVarExpr, *ExistingOperand);
+			if (!bIsDebugValue)
+			{
+				check(!WorkData.ExprToOperand.Contains(InVarExpr));
+				WorkData.ExprToOperand.Add(InVarExpr, *ExistingOperand);
+			}
+			return *ExistingOperand;
 		}
-		return *ExistingOperand;
 	}
 
 	// check if this is a variable with a segment path
@@ -1498,8 +1524,6 @@ FRigVMOperand URigVMCompiler::FindOrAddRegister(const FRigVMVarExprAST* InVarExp
 			ScriptStruct = GetScriptStructForCPPType(BaseCPPType);
 		}
 
-		FRigVMMemoryContainer& Memory = bIsLiteral ? WorkData.VM->GetLiteralMemory() : WorkData.VM->GetWorkMemory();
-
 		bool bRegisterIsMultiUse = Settings.ConsolidateWorkRegisters;
 		if (bRegisterIsMultiUse)
 		{
@@ -1629,6 +1653,49 @@ FRigVMOperand URigVMCompiler::FindOrAddRegister(const FRigVMVarExprAST* InVarExp
 				ensure(NumSlices >= 1);
 			}
 
+			int32 Register = INDEX_NONE;
+
+			// debug watch register might already exists - look for them by name
+			if(bIsDebugValue && Memory.SupportsNames())
+			{
+				Register = Memory.GetIndex(RegisterName);
+				if(Register != INDEX_NONE)
+				{
+					bool bFoundValidRegister = false;
+					if(ScriptStruct)
+					{
+						if(Memory[Register].BaseCPPTypeObject == ScriptStruct)
+						{
+							bFoundValidRegister = true;
+						}
+					}
+					else if(UEnum* Enum = Pin->GetEnum())
+					{
+						if(Memory[Register].BaseCPPTypeObject == Enum)
+						{
+							bFoundValidRegister = true;
+						}
+					}
+					else
+					{
+						if(Memory[Register].BaseCPPType == *BaseCPPType)
+						{
+							bFoundValidRegister = true;
+						}
+					}
+
+					if(bFoundValidRegister)
+					{
+						Operand = Memory.GetOperand(Register);
+						if(ExistingOperand == nullptr)
+						{
+							WorkData.PinPathToOperand->Add(Hash, Operand);
+						}
+						return Operand;
+					}
+				}
+			}
+
 			if (ScriptStruct)
 			{
 				TArray<uint8, TAlignedHeapAllocator<16>> Data;
@@ -1646,7 +1713,9 @@ FRigVMOperand URigVMCompiler::FindOrAddRegister(const FRigVMVarExprAST* InVarExp
 					Ptr = Ptr + ScriptStruct->GetStructureSize();
 				}
 
-				int32 Register = Memory.AddRegisterArray(!Pin->IsDynamicArray() && !bIsDebugValue, RegisterName, ScriptStruct->GetStructureSize(), DefaultValues.Num(), Pin->IsArray(), Data.GetData(), NumSlices, ERigVMRegisterType::Struct, ScriptStruct);
+				Register = Memory.AddRegisterArray(!Pin->IsDynamicArray() && !bIsDebugValue, RegisterName, ScriptStruct->GetStructureSize(), DefaultValues.Num(), Pin->IsArray(), Data.GetData(), NumSlices, ERigVMRegisterType::Struct, ScriptStruct);
+				check(Register != INDEX_NONE);
+				
 #if WITH_EDITORONLY_DATA
 				Memory.Registers[Register].BaseCPPType = *ScriptStruct->GetStructCPPName();
 				Memory.Registers[Register].BaseCPPTypeObject = ScriptStruct;
@@ -1670,7 +1739,9 @@ FRigVMOperand URigVMCompiler::FindOrAddRegister(const FRigVMVarExprAST* InVarExp
 					}
 				}
 
-				int32 Register = Memory.AddRegisterArray<uint8>(!Pin->IsDynamicArray() && !bIsDebugValue, RegisterName, Values.Num(), Pin->IsArray(), Values.GetData(), NumSlices, ERigVMRegisterType::Plain, nullptr);
+				Register = Memory.AddRegisterArray<uint8>(!Pin->IsDynamicArray() && !bIsDebugValue, RegisterName, Values.Num(), Pin->IsArray(), Values.GetData(), NumSlices, ERigVMRegisterType::Plain, nullptr);
+				check(Register != INDEX_NONE);
+				
 #if WITH_EDITORONLY_DATA
 				Memory.Registers[Register].BaseCPPType = *Enum->CppType;
 				Memory.Registers[Register].BaseCPPTypeObject = Enum;
@@ -1684,7 +1755,10 @@ FRigVMOperand URigVMCompiler::FindOrAddRegister(const FRigVMVarExprAST* InVarExp
 				{
 					Values.Add((DefaultValue == TEXT("True")) || (DefaultValue == TEXT("true")) || (DefaultValue == TEXT("1")));
 				}
-				int32 Register = Memory.AddRegisterArray<bool>(!Pin->IsDynamicArray() && !bIsDebugValue, RegisterName, Values.Num(), Pin->IsArray(), (uint8*)Values.GetData(), NumSlices, ERigVMRegisterType::Plain, nullptr);
+
+				Register = Memory.AddRegisterArray<bool>(!Pin->IsDynamicArray() && !bIsDebugValue, RegisterName, Values.Num(), Pin->IsArray(), (uint8*)Values.GetData(), NumSlices, ERigVMRegisterType::Plain, nullptr);
+				check(Register != INDEX_NONE);
+				
 #if WITH_EDITORONLY_DATA
 				Memory.Registers[Register].BaseCPPType = *BaseCPPType;
 				Memory.Registers[Register].BaseCPPTypeObject = nullptr;
@@ -1705,7 +1779,9 @@ FRigVMOperand URigVMCompiler::FindOrAddRegister(const FRigVMVarExprAST* InVarExp
 						Values.Add(FCString::Atoi(*DefaultValue));
 					}
 				}
-				int32 Register = Memory.AddRegisterArray<int32>(!Pin->IsDynamicArray() && !bIsDebugValue, RegisterName, Values.Num(), Pin->IsArray(), (uint8*)Values.GetData(), NumSlices, ERigVMRegisterType::Plain, nullptr);
+				Register = Memory.AddRegisterArray<int32>(!Pin->IsDynamicArray() && !bIsDebugValue, RegisterName, Values.Num(), Pin->IsArray(), (uint8*)Values.GetData(), NumSlices, ERigVMRegisterType::Plain, nullptr);
+				check(Register != INDEX_NONE);
+
 #if WITH_EDITORONLY_DATA
 				Memory.Registers[Register].BaseCPPType = *BaseCPPType;
 				Memory.Registers[Register].BaseCPPTypeObject = nullptr;
@@ -1726,7 +1802,9 @@ FRigVMOperand URigVMCompiler::FindOrAddRegister(const FRigVMVarExprAST* InVarExp
 						Values.Add(FCString::Atof(*DefaultValue));
 					}
 				}
-				int32 Register = Memory.AddRegisterArray<float>(!Pin->IsDynamicArray() && !bIsDebugValue, RegisterName, Values.Num(), Pin->IsArray(), (uint8*)Values.GetData(), NumSlices, ERigVMRegisterType::Plain, nullptr);
+				Register = Memory.AddRegisterArray<float>(!Pin->IsDynamicArray() && !bIsDebugValue, RegisterName, Values.Num(), Pin->IsArray(), (uint8*)Values.GetData(), NumSlices, ERigVMRegisterType::Plain, nullptr);
+				check(Register != INDEX_NONE);
+				
 #if WITH_EDITORONLY_DATA
 				Memory.Registers[Register].BaseCPPType = *BaseCPPType;
 				Memory.Registers[Register].BaseCPPTypeObject = nullptr;
@@ -1740,7 +1818,10 @@ FRigVMOperand URigVMCompiler::FindOrAddRegister(const FRigVMVarExprAST* InVarExp
 				{
 					Values.Add(*DefaultValue);
 				}
-				int32 Register = Memory.AddRegisterArray<FName>(!Pin->IsDynamicArray() && !bIsDebugValue, RegisterName, Values.Num(), Pin->IsArray(), (uint8*)Values.GetData(), NumSlices, ERigVMRegisterType::Name, nullptr);
+
+				Register = Memory.AddRegisterArray<FName>(!Pin->IsDynamicArray() && !bIsDebugValue, RegisterName, Values.Num(), Pin->IsArray(), (uint8*)Values.GetData(), NumSlices, ERigVMRegisterType::Name, nullptr);
+				check(Register != INDEX_NONE);
+				
 #if WITH_EDITORONLY_DATA
 				Memory.Registers[Register].BaseCPPType = *BaseCPPType;
 				Memory.Registers[Register].BaseCPPTypeObject = nullptr;
@@ -1749,7 +1830,9 @@ FRigVMOperand URigVMCompiler::FindOrAddRegister(const FRigVMVarExprAST* InVarExp
 			}
 			else if (BaseCPPType == TEXT("FString"))
 			{
-				int32 Register = Memory.AddRegisterArray<FString>(!Pin->IsDynamicArray() && !bIsDebugValue, RegisterName, DefaultValues.Num(), Pin->IsArray(), (uint8*)DefaultValues.GetData(), NumSlices, ERigVMRegisterType::String, nullptr);
+				Register = Memory.AddRegisterArray<FString>(!Pin->IsDynamicArray() && !bIsDebugValue, RegisterName, DefaultValues.Num(), Pin->IsArray(), (uint8*)DefaultValues.GetData(), NumSlices, ERigVMRegisterType::String, nullptr);
+				check(Register != INDEX_NONE);
+				
 #if WITH_EDITORONLY_DATA
 				Memory.Registers[Register].BaseCPPType = *BaseCPPType;
 				Memory.Registers[Register].BaseCPPTypeObject = nullptr;
@@ -1783,7 +1866,7 @@ FRigVMOperand URigVMCompiler::FindOrAddRegister(const FRigVMVarExprAST* InVarExp
 
 	// Get all possible pins that lead to the same operand
 
-	if(Settings.ASTSettings.bFoldAssignments)
+	if(Settings.ASTSettings.bFoldAssignments && !bIsDebugValue)
 	{
 		TArray<const URigVMPin*> VirtualPins;
 		Pin->GetExposedPinChain(VirtualPins);
@@ -1795,7 +1878,10 @@ FRigVMOperand URigVMCompiler::FindOrAddRegister(const FRigVMVarExprAST* InVarExp
 	}
 	else
 	{
-		WorkData.PinPathToOperand->Add(Hash, Operand);
+		if(ExistingOperand == nullptr)
+		{
+			WorkData.PinPathToOperand->Add(Hash, Operand);
+		}
 	}
 	
 	if (!bIsDebugValue)
diff --git a/Engine/Source/Developer/RigVMDeveloper/Public/RigVMCompiler/RigVMAST.h b/Engine/Source/Developer/RigVMDeveloper/Public/RigVMCompiler/RigVMAST.h
--- a/Engine/Source/Developer/RigVMDeveloper/Public/RigVMCompiler/RigVMAST.h
+++ b/Engine/Source/Developer/RigVMDeveloper/Public/RigVMCompiler/RigVMAST.h
@@ -1192,6 +1192,11 @@ public:
 	// @return the expressoin for the given subject (or nullptr)
 	const FRigVMExprAST* GetExprForSubject(const FRigVMASTProxy& InProxy) const;
 
+	// returns all expressions for a given subject. subjects include nodes and pins.
+	// @param InSubject the subject to retrieve the expression for (node or pin)
+	// @return all expressions for the given subject
+	TArray<const FRigVMExprAST*> GetExpressionsForSubject(UObject* InSubject) const;
+
 	// Prepares the parser for cycle checking on a given pin.
 	// This marks up the parents and childen of the corresponding expression in the graph,
 	// to allow the client to determine if a new parent / child relationship could cause a cycle.
diff --git a/Engine/Source/Developer/RigVMDeveloper/Public/RigVMCompiler/RigVMASTProxy.h b/Engine/Source/Developer/RigVMDeveloper/Public/RigVMCompiler/RigVMASTProxy.h
--- a/Engine/Source/Developer/RigVMDeveloper/Public/RigVMCompiler/RigVMASTProxy.h
+++ b/Engine/Source/Developer/RigVMDeveloper/Public/RigVMCompiler/RigVMASTProxy.h
@@ -20,6 +20,9 @@ public:
 	// Returns the number of entries on this callstack
 	int32 Num() const;
 
+	// Returns the last object in the callstack
+	const UObject* Last() const;
+
 	// Returns the element at a given index
 	const UObject* operator[](int32 InIndex) const;
 
diff --git a/Engine/Source/Developer/RigVMDeveloper/Public/RigVMCompiler/RigVMCompiler.h b/Engine/Source/Developer/RigVMDeveloper/Public/RigVMCompiler/RigVMCompiler.h
--- a/Engine/Source/Developer/RigVMDeveloper/Public/RigVMCompiler/RigVMCompiler.h
+++ b/Engine/Source/Developer/RigVMDeveloper/Public/RigVMCompiler/RigVMCompiler.h
@@ -121,6 +121,8 @@ public:
 
 	static UScriptStruct* GetScriptStructForCPPType(const FString& InCPPType);
 	static FString GetPinHash(const URigVMPin* InPin, const FRigVMVarExprAST* InVarExpr, bool bIsDebugValue = false);
+	void CreateDebugRegister(URigVMPin* InPin, URigVM* OutVM, TMap<FString, FRigVMOperand>* OutOperands, TSharedPtr<FRigVMParserAST> InRuntimeAST);
+	void RemoveDebugRegister(URigVMPin* InPin, URigVM* OutVM, TMap<FString, FRigVMOperand>* OutOperands, TSharedPtr<FRigVMParserAST> InRuntimeAST);
 
 private:
 
diff --git a/Engine/Source/Runtime/RigVM/Private/RigVMCore/RigVM.cpp b/Engine/Source/Runtime/RigVM/Private/RigVMCore/RigVM.cpp
--- a/Engine/Source/Runtime/RigVM/Private/RigVMCore/RigVM.cpp
+++ b/Engine/Source/Runtime/RigVM/Private/RigVMCore/RigVM.cpp
@@ -64,6 +64,7 @@ UScriptStruct* FRigVMParameter::GetScriptStruct() const
 URigVM::URigVM()
     : WorkMemoryPtr(&WorkMemoryStorage)
     , LiteralMemoryPtr(&LiteralMemoryStorage)
+    , DebugMemoryPtr(&DebugMemoryStorage)
     , ByteCodePtr(&ByteCodeStorage)
 #if WITH_EDITOR
 	, DebugInfo(nullptr)
@@ -79,6 +80,7 @@ URigVM::URigVM()
 {
 	GetWorkMemory().SetMemoryType(ERigVMMemoryType::Work);
 	GetLiteralMemory().SetMemoryType(ERigVMMemoryType::Literal);
+	GetDebugMemory().SetMemoryType(ERigVMMemoryType::Debug);
 }
 
 URigVM::~URigVM()
@@ -161,6 +163,7 @@ void URigVM::Reset()
 {
 	WorkMemoryStorage.Reset();
 	LiteralMemoryStorage.Reset();
+	DebugMemoryStorage.Reset();
 	FunctionNamesStorage.Reset();
 	FunctionsStorage.Reset();
 	ByteCodeStorage.Reset();
@@ -171,19 +174,21 @@ void URigVM::Reset()
 
 	WorkMemoryPtr = &WorkMemoryStorage;
 	LiteralMemoryPtr = &LiteralMemoryStorage;
+	DebugMemoryPtr = &DebugMemoryStorage;
 	FunctionNamesPtr = &FunctionNamesStorage;
 	FunctionsPtr = &FunctionsStorage;
 	ByteCodePtr = &ByteCodeStorage;
 
 	InvalidateCachedMemory();
 	
-	CachedDebugWatchRegisters.Reset();
+	OperandToDebugRegisters.Reset();
 }
 
 void URigVM::Empty()
 {
 	WorkMemoryStorage.Empty();
 	LiteralMemoryStorage.Empty();
+	DebugMemoryStorage.Empty();
 	FunctionNamesStorage.Empty();
 	FunctionsStorage.Empty();
 	ByteCodeStorage.Empty();
@@ -199,7 +204,7 @@ void URigVM::Empty()
 	FirstHandleForInstruction.Empty();
 	CachedMemoryHandles.Empty();
 
-	CachedDebugWatchRegisters.Empty();
+	OperandToDebugRegisters.Empty();
 }
 
 void URigVM::CopyFrom(URigVM* InVM, bool bDeferCopy, bool bReferenceLiteralMemory, bool bReferenceByteCode, bool bCopyExternalVariables, bool bCopyDynamicRegisters)
@@ -240,6 +245,16 @@ void URigVM::CopyFrom(URigVM* InVM, bool bDeferCopy, bool bReferenceLiteralMemor
 		LiteralMemoryPtr = InVM->LiteralMemoryPtr;
 	}
 
+	if(InVM->DebugMemoryPtr == &InVM->DebugMemoryStorage)
+	{
+		DebugMemoryStorage = InVM->DebugMemoryStorage;
+		DebugMemoryPtr = &DebugMemoryStorage;
+	}
+	else
+	{
+		DebugMemoryPtr = InVM->DebugMemoryPtr;
+	}
+
 	if(InVM->FunctionNamesPtr == &InVM->FunctionNamesStorage && !bReferenceByteCode)
 	{
 		FunctionNamesStorage = InVM->FunctionNamesStorage;
@@ -737,10 +752,10 @@ bool URigVM::Initialize(FRigVMMemoryContainerPtrArray Memory, FRigVMFixedArray<v
 		return true;
 	}
 
-	FRigVMMemoryContainer* LocalMemory[] = { WorkMemoryPtr, LiteralMemoryPtr };
+	FRigVMMemoryContainer* LocalMemory[] = { WorkMemoryPtr, LiteralMemoryPtr, DebugMemoryPtr };
 	if (Memory.Num() == 0)
 	{
-		Memory = FRigVMMemoryContainerPtrArray(LocalMemory, 2);
+		Memory = FRigVMMemoryContainerPtrArray(LocalMemory, 3);
 	}
 
 	CacheMemoryHandlesIfRequired(Memory);
@@ -755,8 +770,6 @@ bool URigVM::Initialize(FRigVMMemoryContainerPtrArray Memory, FRigVMFixedArray<v
 	Context.SliceOffsets.AddZeroed(Instructions.Num());
 	Context.OpaqueArguments = AdditionalArguments;
 	Context.ExternalVariables = ExternalVariables;
-
-	CacheDebugWatchRegisters();
 	
 #if WITH_EDITOR
 	if (DebugInfo)
@@ -1039,10 +1052,10 @@ bool URigVM::Execute(FRigVMMemoryContainerPtrArray Memory, FRigVMFixedArray<void
 		return true;
 	}
 
-	FRigVMMemoryContainer* LocalMemory[] = { WorkMemoryPtr, LiteralMemoryPtr }; 
+	FRigVMMemoryContainer* LocalMemory[] = { WorkMemoryPtr, LiteralMemoryPtr, DebugMemoryPtr }; 
 	if (Memory.Num() == 0)
 	{
-		Memory = FRigVMMemoryContainerPtrArray(LocalMemory, 2);
+		Memory = FRigVMMemoryContainerPtrArray(LocalMemory, 3);
 	}
 
 	CacheMemoryHandlesIfRequired(Memory);
@@ -1065,7 +1078,7 @@ bool URigVM::Execute(FRigVMMemoryContainerPtrArray Memory, FRigVMFixedArray<void
 	Context.OpaqueArguments = AdditionalArguments;
 	Context.ExternalVariables = ExternalVariables;
 
-	ClearDebugWatchRegisterMemory();
+	ClearDebugMemory();
 
 	if (!InEntryName.IsNone())
 	{
@@ -1182,18 +1195,40 @@ bool URigVM::Execute(FRigVMMemoryContainerPtrArray Memory, FRigVMFixedArray<void
 			case ERigVMOpCode::Execute_64_Operands:
 			{
 				const FRigVMExecuteOp& Op = ByteCode.GetOpAt<FRigVMExecuteOp>(Instruction);
-				int32 OperandCount = FirstHandleForInstruction[Context.InstructionIndex + 1] - FirstHandleForInstruction[Context.InstructionIndex];
+				const int32 OperandCount = FirstHandleForInstruction[Context.InstructionIndex + 1] - FirstHandleForInstruction[Context.InstructionIndex];
 				FRigVMMemoryHandleArray OpHandles(&CachedMemoryHandles[FirstHandleForInstruction[Context.InstructionIndex]], OperandCount);
 #if WITH_EDITOR
 				Context.FunctionName = FunctionNames[Op.FunctionIndex];
 #endif
 				(*Functions[Op.FunctionIndex])(Context, OpHandles);
+
+#if WITH_EDITOR
+				if(DebugMemoryPtr->Num() > 0)
+				{
+					const FRigVMOperandArray Operands = ByteCode.GetOperandsForExecuteOp(Instruction);
+					check(Operands.Num() == OperandCount);
+					
+					for(int32 OperandIndex=0;OperandIndex<OperandCount;OperandIndex++)
+					{
+						CopyOperandForDebuggingIfNeeded(Operands[OperandIndex], OpHandles[OperandIndex]);
+					}
+				}
+#endif
+
 				Context.InstructionIndex++;
 				break;
 			}
 			case ERigVMOpCode::Zero:
 			{
 				*((int32*)CachedMemoryHandles[FirstHandleForInstruction[Context.InstructionIndex]].GetData()) = 0;
+#if WITH_EDITOR
+				if(DebugMemoryPtr->Num() > 0)
+				{
+					const FRigVMUnaryOp& Op = ByteCode.GetOpAt<FRigVMUnaryOp>(Instruction);
+					CopyOperandForDebuggingIfNeeded(Op.Arg, CachedMemoryHandles[FirstHandleForInstruction[Context.InstructionIndex]]);
+				}
+#endif
+
 				Context.InstructionIndex++;
 				break;
 			}
@@ -1218,8 +1253,8 @@ bool URigVM::Execute(FRigVMMemoryContainerPtrArray Memory, FRigVMFixedArray<void
 				void* SourcePtr = SourceHandle;
 				void* TargetPtr = TargetHandle;
 
-				uint64 NumBytes = reinterpret_cast<uint64>(CachedMemoryHandles[FirstHandleForInstruction[Context.InstructionIndex] + 2].GetData());
-				ERigVMRegisterType MemoryType = (ERigVMRegisterType)reinterpret_cast<uint64>(CachedMemoryHandles[FirstHandleForInstruction[Context.InstructionIndex] + 3].GetData());
+				const uint64 NumBytes = reinterpret_cast<uint64>(CachedMemoryHandles[FirstHandleForInstruction[Context.InstructionIndex] + 2].GetData());
+				const ERigVMRegisterType MemoryType = (ERigVMRegisterType)reinterpret_cast<uint64>(CachedMemoryHandles[FirstHandleForInstruction[Context.InstructionIndex] + 3].GetData());
 
 				if (TargetHandle.Type == FRigVMMemoryHandle::Dynamic)
 				{
@@ -1228,7 +1263,7 @@ bool URigVM::Execute(FRigVMMemoryContainerPtrArray Memory, FRigVMFixedArray<void
 					{
 						Storage->Reset();
 					}
-					int32 ByteIndex = Storage->AddZeroed(NumBytes);
+					const int32 ByteIndex = Storage->AddZeroed(NumBytes);
 					TargetPtr = Storage->GetData() + ByteIndex;
 				}
 				else if (TargetHandle.Type == FRigVMMemoryHandle::NestedDynamic)
@@ -1238,7 +1273,7 @@ bool URigVM::Execute(FRigVMMemoryContainerPtrArray Memory, FRigVMFixedArray<void
 					{
 						Storage->Reset();
 					}
-					int32 ArrayIndex = Storage->Add(FRigVMByteArray());
+					const int32 ArrayIndex = Storage->Add(FRigVMByteArray());
 					(*Storage)[ArrayIndex].AddZeroed(NumBytes);
 					TargetPtr = (*Storage)[ArrayIndex].GetData();
 				}
@@ -1252,7 +1287,7 @@ bool URigVM::Execute(FRigVMMemoryContainerPtrArray Memory, FRigVMFixedArray<void
 					}
 					case ERigVMRegisterType::Name:
 					{
-						int32 NumNames = NumBytes / sizeof(FName);
+						const int32 NumNames = NumBytes / sizeof(FName);
 						FRigVMFixedArray<FName> TargetNames((FName*)TargetPtr, NumNames);
 						FRigVMFixedArray<FName> SourceNames((FName*)SourcePtr, NumNames);
 						for (int32 Index = 0; Index < NumNames; Index++)
@@ -1263,7 +1298,7 @@ bool URigVM::Execute(FRigVMMemoryContainerPtrArray Memory, FRigVMFixedArray<void
 					}
 					case ERigVMRegisterType::String:
 					{
-						int32 NumStrings = NumBytes / sizeof(FString);
+						const int32 NumStrings = NumBytes / sizeof(FString);
 						FRigVMFixedArray<FString> TargetStrings((FString*)TargetPtr, NumStrings);
 						FRigVMFixedArray<FString> SourceStrings((FString*)SourcePtr, NumStrings);
 						for (int32 Index = 0; Index < NumStrings; Index++)
@@ -1275,7 +1310,7 @@ bool URigVM::Execute(FRigVMMemoryContainerPtrArray Memory, FRigVMFixedArray<void
 					case ERigVMRegisterType::Struct:
 					{
 						UScriptStruct* ScriptStruct = (UScriptStruct*)CachedMemoryHandles[FirstHandleForInstruction[Context.InstructionIndex] + 4].GetData();
-						int32 NumStructs = NumBytes / ScriptStruct->GetStructureSize();
+						const int32 NumStructs = NumBytes / ScriptStruct->GetStructureSize();
 						if (NumStructs > 0 && TargetPtr)
 						{
 							ScriptStruct->CopyScriptStruct(TargetPtr, SourcePtr, NumStructs);
@@ -1296,12 +1331,26 @@ bool URigVM::Execute(FRigVMMemoryContainerPtrArray Memory, FRigVMFixedArray<void
 			case ERigVMOpCode::Increment:
 			{
 				(*((int32*)CachedMemoryHandles[FirstHandleForInstruction[Context.InstructionIndex]].GetData()))++;
+#if WITH_EDITOR
+				if(DebugMemoryPtr->Num() > 0)
+				{
+					const FRigVMUnaryOp& Op = ByteCode.GetOpAt<FRigVMUnaryOp>(Instruction);
+					CopyOperandForDebuggingIfNeeded(Op.Arg, CachedMemoryHandles[FirstHandleForInstruction[Context.InstructionIndex]]);
+				}
+#endif
 				Context.InstructionIndex++;
 				break;
 			}
 			case ERigVMOpCode::Decrement:
 			{
 				(*((int32*)CachedMemoryHandles[FirstHandleForInstruction[Context.InstructionIndex]].GetData()))--;
+#if WITH_EDITOR
+				if(DebugMemoryPtr->Num() > 0)
+				{
+					const FRigVMUnaryOp& Op = ByteCode.GetOpAt<FRigVMUnaryOp>(Instruction);
+					CopyOperandForDebuggingIfNeeded(Op.Arg, CachedMemoryHandles[FirstHandleForInstruction[Context.InstructionIndex]]);
+				}
+#endif
 				Context.InstructionIndex++;
 				break;
 			}
@@ -1311,8 +1360,8 @@ bool URigVM::Execute(FRigVMMemoryContainerPtrArray Memory, FRigVMFixedArray<void
 				const FRigVMComparisonOp& Op = ByteCode.GetOpAt<FRigVMComparisonOp>(Instruction);
 				const FRigVMRegister& RegisterA = (*Memory[Op.A.GetContainerIndex()])[Op.A.GetRegisterIndex()];
 				const FRigVMRegister& RegisterB = (*Memory[Op.B.GetContainerIndex()])[Op.B.GetRegisterIndex()];
-				uint16 BytesA = RegisterA.GetNumBytesPerSlice();
-				uint16 BytesB = RegisterB.GetNumBytesPerSlice();
+				const uint16 BytesA = RegisterA.GetNumBytesPerSlice();
+				const uint16 BytesB = RegisterB.GetNumBytesPerSlice();
 				
 				bool Result = false;
 				if (BytesA == BytesB && RegisterA.Type == RegisterB.Type && RegisterA.ScriptStructIndex == RegisterB.ScriptStructIndex)
@@ -1400,7 +1449,7 @@ bool URigVM::Execute(FRigVMMemoryContainerPtrArray Memory, FRigVMFixedArray<void
 			case ERigVMOpCode::JumpAbsoluteIf:
 			{
 				const FRigVMJumpIfOp& Op = ByteCode.GetOpAt<FRigVMJumpIfOp>(Instruction);
-				bool Condition = *(bool*)CachedMemoryHandles[FirstHandleForInstruction[Context.InstructionIndex]].GetData();
+				const bool Condition = *(bool*)CachedMemoryHandles[FirstHandleForInstruction[Context.InstructionIndex]].GetData();
 				if (Condition == Op.Condition)
 				{
 					Context.InstructionIndex = Op.InstructionIndex;
@@ -1414,7 +1463,7 @@ bool URigVM::Execute(FRigVMMemoryContainerPtrArray Memory, FRigVMFixedArray<void
 			case ERigVMOpCode::JumpForwardIf:
 			{
 				const FRigVMJumpIfOp& Op = ByteCode.GetOpAt<FRigVMJumpIfOp>(Instruction);
-				bool Condition = *(bool*)CachedMemoryHandles[FirstHandleForInstruction[Context.InstructionIndex]].GetData();
+				const bool Condition = *(bool*)CachedMemoryHandles[FirstHandleForInstruction[Context.InstructionIndex]].GetData();
 				if (Condition == Op.Condition)
 				{
 					Context.InstructionIndex += Op.InstructionIndex;
@@ -1428,7 +1477,7 @@ bool URigVM::Execute(FRigVMMemoryContainerPtrArray Memory, FRigVMFixedArray<void
 			case ERigVMOpCode::JumpBackwardIf:
 			{
 				const FRigVMJumpIfOp& Op = ByteCode.GetOpAt<FRigVMJumpIfOp>(Instruction);
-				bool Condition = *(bool*)CachedMemoryHandles[FirstHandleForInstruction[Context.InstructionIndex]].GetData();
+				const bool Condition = *(bool*)CachedMemoryHandles[FirstHandleForInstruction[Context.InstructionIndex]].GetData();
 				if (Condition == Op.Condition)
 				{
 					Context.InstructionIndex -= Op.InstructionIndex;
@@ -1458,8 +1507,8 @@ bool URigVM::Execute(FRigVMMemoryContainerPtrArray Memory, FRigVMFixedArray<void
 			}
 			case ERigVMOpCode::BeginBlock:
 			{
-				int32 Count = (*((int32*)CachedMemoryHandles[FirstHandleForInstruction[Context.InstructionIndex]].GetData()));
-				int32 Index = (*((int32*)CachedMemoryHandles[FirstHandleForInstruction[Context.InstructionIndex] + 1].GetData()));
+				const int32 Count = (*((int32*)CachedMemoryHandles[FirstHandleForInstruction[Context.InstructionIndex]].GetData()));
+				const int32 Index = (*((int32*)CachedMemoryHandles[FirstHandleForInstruction[Context.InstructionIndex] + 1].GetData()));
 				Context.BeginSlice(Count, Index);
 				Context.InstructionIndex++;
 				break;
@@ -1516,6 +1565,10 @@ void URigVM::SetRegisterValueFromString(const FRigVMOperand& InOperand, const FS
 	{
 		GetWorkMemory().SetRegisterValueFromString(InOperand, InCPPType, InCPPTypeObject, InDefaultValues);
 	}
+	else if (InOperand.GetMemoryType() == ERigVMMemoryType::Debug)
+	{
+		GetDebugMemory().SetRegisterValueFromString(InOperand, InCPPType, InCPPTypeObject, InDefaultValues);
+	}
 }
 
 
@@ -1788,6 +1841,10 @@ FString URigVM::GetOperandLabel(const FRigVMOperand& InOperand, TFunction<FStrin
 	{
 		MemoryPtr = LiteralMemoryPtr;
 	}
+	else if (InOperand.GetMemoryType() == ERigVMMemoryType::Debug)
+	{
+		MemoryPtr = DebugMemoryPtr;
+	}
 	else
 	{
 		MemoryPtr = WorkMemoryPtr;
@@ -1829,48 +1886,19 @@ FString URigVM::GetOperandLabel(const FRigVMOperand& InOperand, TFunction<FStrin
 
 #endif
 
-void URigVM::CacheDebugWatchRegisters()
+void URigVM::ClearDebugMemory()
 {
-	FRigVMMemoryContainer* WorkMemory = CachedMemory[(int32)ERigVMMemoryType::Work];
-
-	if (WorkMemory)
-	{
-		CachedDebugWatchRegisters.Reset();
-		for (int32 RegisterIndex = 0; RegisterIndex < WorkMemoryPtr->Registers.Num(); RegisterIndex++)
-		{
-			FRigVMRegister& Register = WorkMemory->Registers[RegisterIndex];
-			if (Register.Name.ToString().StartsWith("DebugWatch:"))
-			{
-				ensure(Register.IsDynamic());
-
-				CachedDebugWatchRegisters.Add(&Register); 
-			}
-		}
-	} 
-}
-
-void URigVM::ClearDebugWatchRegisterMemory()
-{
-	FRigVMMemoryContainer* WorkMemory = CachedMemory[(int32)ERigVMMemoryType::Work];
-
-	if (WorkMemory)
+#if WITH_EDITOR
+	FRigVMMemoryContainer* DebugMemory = CachedMemory[(int32)ERigVMMemoryType::Debug];
+	if (DebugMemory)
 	{ 
-		for (const FRigVMRegister* Register : CachedDebugWatchRegisters)
-		{ 
-			ensure(Register->IsDynamic());
-			
-			if (!Register->IsArray())
-			{
-				FRigVMByteArray* VMMemoryPtr = (FRigVMByteArray*)(&(WorkMemory->Data[Register->GetWorkByteIndex()]));
-				VMMemoryPtr->Reset();
-			}
-			else
-			{ 
-				FRigVMNestedByteArray* VMMemoryPtr = (FRigVMNestedByteArray*)(&(WorkMemory->Data[Register->GetWorkByteIndex()]));
-				VMMemoryPtr->Reset();
-			}
+		for (int32 RegisterIndex = 0; RegisterIndex < DebugMemory->Num(); RegisterIndex++)
+		{
+			ensure(DebugMemory->GetRegister(RegisterIndex).IsDynamic());
+			DebugMemory->Destroy(RegisterIndex);
 		}	
 	} 
+#endif
 }
 
 void URigVM::CacheSingleMemoryHandle(const FRigVMOperand& InArg, bool bForExecute)
@@ -1905,3 +1933,117 @@ void URigVM::CacheSingleMemoryHandle(const FRigVMOperand& InArg, bool bForExecut
 		}
 	}
 }
+
+void URigVM::CopyOperandForDebuggingImpl(const FRigVMOperand& InArg, const FRigVMMemoryHandle& InHandle, const FRigVMOperand& InDebugOperand)
+{
+#if WITH_EDITOR
+	check(InArg.IsValid());
+	check(InArg.GetRegisterOffset() == INDEX_NONE);
+	check(InDebugOperand.IsValid());
+	check(InDebugOperand.GetRegisterOffset() == INDEX_NONE);
+
+	const FRigVMRegister& DebugRegister = DebugMemoryPtr->GetRegister(InDebugOperand.GetRegisterIndex());
+	check(DebugRegister.IsDynamic());
+
+	if (Context.GetSlice().GetIndex() == 0)
+	{
+		DebugMemoryPtr->Destroy(InDebugOperand.GetRegisterIndex());
+	}
+
+	// the source pointer is not going to be sliced since we only allow
+	// watches on things exposed from a node (so no hidden pins)
+	const uint8* SourcePtr = InHandle.GetData(0, true);
+	uint8* TargetPtr = nullptr;
+
+	int32 NumBytes = (int32)DebugRegister.ElementSize;
+	if(InHandle.GetType() == FRigVMMemoryHandle::FType::Dynamic)
+	{
+		FRigVMByteArray* Storage = (FRigVMByteArray*)InHandle.Ptr;
+		NumBytes = Storage->Num();
+		TargetPtr = Storage->GetData();
+	}
+	else if(InHandle.GetType() == FRigVMMemoryHandle::FType::NestedDynamic)
+	{
+		FRigVMNestedByteArray* Storage = (FRigVMNestedByteArray*)InHandle.Ptr;
+		NumBytes = (*Storage)[Context.GetSlice().GetIndex()].Num(); 
+		TargetPtr = (*Storage)[Context.GetSlice().GetIndex()].GetData();
+	}
+
+	const FRigVMMemoryHandle DebugHandle = DebugMemoryPtr->GetHandle(InDebugOperand.GetRegisterIndex());
+	if (DebugRegister.IsNestedDynamic())
+	{
+		FRigVMNestedByteArray* Storage = (FRigVMNestedByteArray*)DebugHandle.Ptr;
+		while(Storage->Num() < Context.GetSlice().TotalNum())
+		{
+			Storage->Add(FRigVMByteArray());
+		}
+		(*Storage)[Context.GetSlice().GetIndex()].AddZeroed(NumBytes);
+		TargetPtr = (*Storage)[Context.GetSlice().GetIndex()].GetData();
+	}
+	else
+	{
+		const int32 TotalBytes = Context.GetSlice().TotalNum() * NumBytes;
+		FRigVMByteArray* Storage = (FRigVMByteArray*)DebugHandle.Ptr;
+		while(Storage->Num() < TotalBytes)
+		{
+			Storage->AddZeroed(NumBytes);
+		}
+		TargetPtr = &(*Storage)[Context.GetSlice().GetIndex() * NumBytes];
+	}
+
+	if((SourcePtr == nullptr) || (TargetPtr == nullptr))
+	{
+		return;
+	}
+
+	switch (DebugRegister.Type)
+	{
+		case ERigVMRegisterType::Plain:
+		{
+			FMemory::Memcpy(TargetPtr, SourcePtr, NumBytes);
+			break;
+		}
+		case ERigVMRegisterType::Name:
+		{
+			const int32 NumNames = NumBytes / sizeof(FName);
+			FRigVMFixedArray<FName> TargetNames((FName*)TargetPtr, NumNames);
+			FRigVMFixedArray<FName> SourceNames((FName*)SourcePtr, NumNames);
+			for (int32 Index = 0; Index < NumNames; Index++)
+			{
+				TargetNames[Index] = SourceNames[Index];
+			}
+			break;
+		}
+		case ERigVMRegisterType::String:
+		{
+			const int32 NumStrings = NumBytes / sizeof(FString);
+			FRigVMFixedArray<FString> TargetStrings((FString*)TargetPtr, NumStrings);
+			FRigVMFixedArray<FString> SourceStrings((FString*)SourcePtr, NumStrings);
+			for (int32 Index = 0; Index < NumStrings; Index++)
+			{
+				TargetStrings[Index] = SourceStrings[Index];
+			}
+			break;
+		}
+		case ERigVMRegisterType::Struct:
+		{
+			UScriptStruct* ScriptStruct = DebugMemoryPtr->GetScriptStruct(DebugRegister);
+			const int32 NumStructs = NumBytes / ScriptStruct->GetStructureSize();
+			if (NumStructs > 0 && TargetPtr)
+			{
+				ScriptStruct->CopyScriptStruct(TargetPtr, SourcePtr, NumStructs);
+			}
+			break;
+		}
+		default:
+		{
+			// the default pass for any complex memory
+			FRigVMMemoryContainer* LocalMemory[] = { WorkMemoryPtr, LiteralMemoryPtr, DebugMemoryPtr };
+			DebugMemoryPtr->Copy(InArg, InDebugOperand, LocalMemory[InArg.GetContainerIndex()]);
+			break;
+		}
+	}
+
+	
+#endif
+}
diff --git a/Engine/Source/Runtime/RigVM/Public/RigVMCore/RigVM.h b/Engine/Source/Runtime/RigVM/Public/RigVMCore/RigVM.h
--- a/Engine/Source/Runtime/RigVM/Public/RigVMCore/RigVM.h
+++ b/Engine/Source/Runtime/RigVM/Public/RigVMCore/RigVM.h
@@ -181,6 +181,13 @@ public:
 	FORCEINLINE FRigVMMemoryContainer& GetLiteralMemory() { return *LiteralMemoryPtr; }
 	FORCEINLINE const FRigVMMemoryContainer& GetLiteralMemory() const { return *LiteralMemoryPtr; }
 
+	// The default debug watch memory
+	UPROPERTY()
+	FRigVMMemoryContainer DebugMemoryStorage;
+	FRigVMMemoryContainer* DebugMemoryPtr;
+	FORCEINLINE FRigVMMemoryContainer& GetDebugMemory() { return *DebugMemoryPtr; }
+	FORCEINLINE const FRigVMMemoryContainer& GetDebugMemory() const { return *DebugMemoryPtr; }
+
 	// The byte code of the VM
 	UPROPERTY()
 	FRigVMByteCode ByteCodeStorage;
@@ -488,6 +495,7 @@ public:
 		FRigVMStatistics Statistics;
 		Statistics.LiteralMemory = LiteralMemoryPtr->GetStatistics();
 		Statistics.WorkMemory = WorkMemoryPtr->GetStatistics();
+		Statistics.DebugMemory = DebugMemoryPtr->GetStatistics();
 		Statistics.ByteCode = ByteCodePtr->GetStatistics();
 		Statistics.BytesForCaching = FirstHandleForInstruction.GetAllocatedSize() + CachedMemoryHandles.GetAllocatedSize();
 		Statistics.BytesForCDO =
@@ -572,18 +580,31 @@ private:
 	// store the first event such that we know when to reset data for a new round of rig evaluation
 	FName FirstEntryEventInQueue;
 #endif
-
-	TArray<FRigVMRegister*> CachedDebugWatchRegisters;
 	
 	// debug watch register memory needs to be cleared for each execution
-	// cache the registers so that we can directly use them during Execute()
-	void CacheDebugWatchRegisters();
-	
-	// debug watch register memory needs to be cleared for each execution
-	void ClearDebugWatchRegisterMemory();
+	void ClearDebugMemory();
 	
 	void CacheSingleMemoryHandle(const FRigVMOperand& InArg, bool bForExecute = false);
 
+	FORCEINLINE void CopyOperandForDebuggingIfNeeded(const FRigVMOperand& InArg, const FRigVMMemoryHandle& InHandle)
+	{
+#if WITH_EDITOR
+		const FRigVMOperand KeyOperand(InArg.GetMemoryType(), InArg.GetRegisterIndex()); // no register offset
+		if(const TArray<FRigVMOperand>* DebugOperandsPtr = OperandToDebugRegisters.Find(KeyOperand))
+		{
+			const TArray<FRigVMOperand>& DebugOperands = *DebugOperandsPtr;
+			for(const FRigVMOperand& DebugOperand : DebugOperands)
+			{
+				CopyOperandForDebuggingImpl(InArg, InHandle, DebugOperand);
+			}
+		}
+#endif
+	}
+
+	void CopyOperandForDebuggingImpl(const FRigVMOperand& InArg, const FRigVMMemoryHandle& InHandle, const FRigVMOperand& InDebugOperand);
+
+	TMap<FRigVMOperand, TArray<FRigVMOperand>> OperandToDebugRegisters;
+
 	int32 ExecutingThreadId;
 
 	UPROPERTY(transient)
diff --git a/Engine/Source/Runtime/RigVM/Public/RigVMCore/RigVMMemory.h b/Engine/Source/Runtime/RigVM/Public/RigVMCore/RigVMMemory.h
--- a/Engine/Source/Runtime/RigVM/Public/RigVMCore/RigVMMemory.h
+++ b/Engine/Source/Runtime/RigVM/Public/RigVMCore/RigVMMemory.h
@@ -32,6 +32,7 @@ enum class ERigVMMemoryType: uint8
 {
 	Work, // Mutable state
 	Literal, // Const / fixed state
+	Debug, // Owned memory used for debug watches
 	External, // Unowned external memory
 	Invalid
 };
@@ -94,6 +95,11 @@ public:
 	// returns the register offset of this argument
 	FORCEINLINE_DEBUGGABLE int32 GetRegisterOffset() const { return RegisterOffset == UINT16_MAX ? INDEX_NONE : (int32)RegisterOffset; }
 
+	friend FORCEINLINE uint32 GetTypeHash(const FRigVMOperand& Operand)
+	{
+  		return HashCombine(HashCombine((uint32)Operand.MemoryType, (uint32)Operand.RegisterIndex), (uint32)Operand.RegisterOffset);
+  	}
+	
 	void Serialize(FArchive& Ar);
 	void Save(FArchive& Ar) const;
 	void Load(FArchive& Ar);
@@ -475,6 +481,11 @@ public:
 		return GetData_Internal(SliceIndex, bGetArrayData);
 	}
 
+	FORCEINLINE FType GetType() const
+	{
+		return Type;
+	}
+
 private:
 
 	FORCEINLINE_DEBUGGABLE uint8* GetData_Internal_NoOffset(int32 SliceIndex, bool bGetArrayData = false) const
diff --git a/Engine/Source/Runtime/RigVM/Public/RigVMCore/RigVMStatistics.h b/Engine/Source/Runtime/RigVM/Public/RigVMCore/RigVMStatistics.h
--- a/Engine/Source/Runtime/RigVM/Public/RigVMCore/RigVMStatistics.h
+++ b/Engine/Source/Runtime/RigVM/Public/RigVMCore/RigVMStatistics.h
@@ -55,6 +55,9 @@ public:
 	UPROPERTY(VisibleAnywhere, Category = "Statistics")
 	FRigVMMemoryStatistics WorkMemory;
 
+	UPROPERTY(VisibleAnywhere, Category = "Statistics")
+	FRigVMMemoryStatistics DebugMemory;
+
 	UPROPERTY(VisibleAnywhere, Category = "Statistics")
 	uint32 BytesForCaching = 0;
 
