Control Rig: Force GerRegisterValueAsString to print complete string to avoid infinite loop in FoldConstantValuesToLiterals

diff --git a/Engine/Source/Runtime/RigVM/Private/RigVMCore/RigVMMemory.cpp b/Engine/Source/Runtime/RigVM/Private/RigVMCore/RigVMMemory.cpp
--- a/Engine/Source/Runtime/RigVM/Private/RigVMCore/RigVMMemory.cpp
+++ b/Engine/Source/Runtime/RigVM/Private/RigVMCore/RigVMMemory.cpp
@@ -1812,15 +1812,11 @@ TArray<FString> FRigVMMemoryContainer::GetRegisterValueAsString(const FRigVMOper
 				UScriptStruct* ScriptStruct = GetScriptStruct(Register);
 				if (ScriptStruct == InCPPTypeObject)
 				{
-					TArray<uint8, TAlignedHeapAllocator<16>> DefaultStructData;
-					DefaultStructData.AddZeroed(ScriptStruct->GetStructureSize());
-					ScriptStruct->InitializeDefaultValue(DefaultStructData.GetData());
-
 					uint8* DataPtr = (uint8*)GetData(Register, INDEX_NONE, SliceIndex);
-					DataPtr += ElementIndex * ScriptStruct->GetStructureSize();
-					ScriptStruct->ExportText(DefaultValue, DataPtr, DefaultStructData.GetData(), nullptr, PPF_None, nullptr);
+					DataPtr += ElementIndex * ScriptStruct->GetStructureSize(); 
 
-					ScriptStruct->DestroyStruct(DefaultStructData.GetData(), 1);
+					//use DataPtr for both value and default value to force exporting complete string
+					ScriptStruct->ExportText(DefaultValue, DataPtr, DataPtr, nullptr, PPF_None, nullptr);
 				}
 			}
 			else if (const UEnum* Enum = Cast<const UEnum>(InCPPTypeObject))
