ControlRig: StepIn/StepOver/StepOut

diff --git a/Engine/Plugins/Experimental/ControlRig/Source/ControlRig/Private/ControlRig.cpp b/Engine/Plugins/Experimental/ControlRig/Source/ControlRig/Private/ControlRig.cpp
--- a/Engine/Plugins/Experimental/ControlRig/Source/ControlRig/Private/ControlRig.cpp
+++ b/Engine/Plugins/Experimental/ControlRig/Source/ControlRig/Private/ControlRig.cpp
@@ -1235,6 +1235,7 @@ void UControlRig::HandleExecutionReachedExit()
 		SnapShotVM->CopyFrom(VM, false, false, false, true, true);
 	}
 	DebugInfo.ResetState();
+	VM->SetBreakpointAction(ERigVMBreakpointAction::None);
 #endif
 	
 	if (LatestExecutedState != EControlRigState::Init && bAccumulateTime)
@@ -2415,14 +2416,14 @@ void UControlRig::AddBreakpoint(int32 InstructionIndex, URigVMNode* InNode)
 	DebugInfo.AddBreakpoint(InstructionIndex, InNode);
 }
 
-void UControlRig::ResumeExecution()
+bool UControlRig::ExecuteBreakpointAction(const ERigVMBreakpointAction BreakpointAction)
 {
-	// this makes sure that the snapshot exists
-	if(URigVM* SnapShotVM = GetSnapshotVM())
+	if (VM->GetHaltedAtInstruction() != INDEX_NONE)
 	{
-		VM->CopyFrom(SnapShotVM, false, false, false, true, true);
+		VM->SetBreakpointAction(BreakpointAction);
+		return true;
 	}
-	VM->ResumeExecution();
+	return false;
 }
 
 #endif // WITH_EDITOR
diff --git a/Engine/Plugins/Experimental/ControlRig/Source/ControlRig/Public/ControlRig.h b/Engine/Plugins/Experimental/ControlRig/Source/ControlRig/Public/ControlRig.h
--- a/Engine/Plugins/Experimental/ControlRig/Source/ControlRig/Public/ControlRig.h
+++ b/Engine/Plugins/Experimental/ControlRig/Source/ControlRig/Public/ControlRig.h
@@ -662,8 +662,9 @@ public:
 	/** Adds a breakpoint in the VM at the InstructionIndex */
 	void AddBreakpoint(int32 InstructionIndex, URigVMNode* InNode);
 
-	/** Reloads the VM from the SnapshotVM, and increments activation hit on the halted breakpoint */
-	void ResumeExecution();
+	/** If the VM is halted at a breakpoint, it sets a breakpoint action so that
+	 *  it is applied on the next VM execution */
+	bool ExecuteBreakpointAction(const ERigVMBreakpointAction BreakpointAction);
 	
 	FRigVMDebugInfo& GetDebugInfo() { return DebugInfo; }
 
diff --git a/Engine/Plugins/Experimental/ControlRig/Source/ControlRigEditor/Private/ControlRigBlueprintCommands.cpp b/Engine/Plugins/Experimental/ControlRig/Source/ControlRigEditor/Private/ControlRigBlueprintCommands.cpp
--- a/Engine/Plugins/Experimental/ControlRig/Source/ControlRigEditor/Private/ControlRigBlueprintCommands.cpp
+++ b/Engine/Plugins/Experimental/ControlRig/Source/ControlRigEditor/Private/ControlRigBlueprintCommands.cpp
@@ -18,6 +18,9 @@ void FControlRigBlueprintCommands::RegisterCommands()
 	UI_COMMAND(ReleaseMode, "Release Mode", "Compiles and Executes the rig, ignoring debug data.", EUserInterfaceActionType::Button, FInputChord());
 	UI_COMMAND(DebugMode, "Debug Mode", "Compiles and Executes the unoptimized rig, stopping at breakpoints.", EUserInterfaceActionType::Button, FInputChord());
 	UI_COMMAND(ResumeExecution, "Resume", "Resumes execution after being halted at a breakpoint.", EUserInterfaceActionType::Button, FInputChord(EKeys::F5));
+	UI_COMMAND(StepInto, "Step Into", "Steps into the collapsed/function node, when halted at a breakpoint.", EUserInterfaceActionType::Button, FInputChord(EKeys::F11));
+	UI_COMMAND(StepOut, "Step Out", "Steps out of the collapsed/function node, when halted at a breakpoint.", EUserInterfaceActionType::Button, FInputChord(EKeys::F11, EModifierKey::Shift));
+	UI_COMMAND(StepOver, "Step Over", "Steps over the node, when halted at a breakpoint.", EUserInterfaceActionType::Button, FInputChord(EKeys::F10));
 }
 
 #undef LOCTEXT_NAMESPACE
diff --git a/Engine/Plugins/Experimental/ControlRig/Source/ControlRigEditor/Private/ControlRigBlueprintCommands.h b/Engine/Plugins/Experimental/ControlRig/Source/ControlRigEditor/Private/ControlRigBlueprintCommands.h
--- a/Engine/Plugins/Experimental/ControlRig/Source/ControlRigEditor/Private/ControlRigBlueprintCommands.h
+++ b/Engine/Plugins/Experimental/ControlRig/Source/ControlRigEditor/Private/ControlRigBlueprintCommands.h
@@ -54,6 +54,15 @@ public:
 	/** Resume the execution of the graph when halted at a breakpoint */
 	TSharedPtr< FUICommandInfo > ResumeExecution;
 
+	/** Steps to the next node in the execution graph (at the same graph level) when halted at a breakpoint */
+	TSharedPtr< FUICommandInfo > StepOver;
+
+	/** Steps into the collapsed/function node, when halted at a breakpoint */
+	TSharedPtr< FUICommandInfo > StepInto;
+
+	/** Steps out of the collapsed/function node, when halted at a breakpoint */
+	TSharedPtr< FUICommandInfo > StepOut;
+
 	/**
 	 * Initialize commands
 	 */
diff --git a/Engine/Plugins/Experimental/ControlRig/Source/ControlRigEditor/Private/ControlRigEditorStyle.h b/Engine/Plugins/Experimental/ControlRig/Source/ControlRigEditor/Private/ControlRigEditorStyle.h
--- a/Engine/Plugins/Experimental/ControlRig/Source/ControlRigEditor/Private/ControlRigEditorStyle.h
+++ b/Engine/Plugins/Experimental/ControlRig/Source/ControlRigEditor/Private/ControlRigEditorStyle.h
@@ -93,6 +93,10 @@ public:
 			Set("ControlRig.ResumeExecution", new IMAGE_BRUSH_SVG("Slate/simulate", Icon40x40));
 			Set("ControlRig.ReleaseMode", new IMAGE_BRUSH_SVG("Slate/animation", Icon40x40));
 			Set("ControlRig.DebugMode", new IMAGE_BRUSH_SVG("Slate/Bug", Icon40x40));
+			
+			Set( "ControlRig.StepOut", new IMAGE_BRUSH("Slate/icon_DebugStepOut_40x", Icon40x40));
+			Set( "ControlRig.StepInto", new IMAGE_BRUSH( "Slate/icon_DebugStepIn_40x", Icon40x40 ) );
+			Set( "ControlRig.StepOver", new IMAGE_BRUSH( "Slate/icon_DebugStepOver_40x", Icon40x40 ) );
 		}
 
 		// Graph styles
diff --git a/Engine/Plugins/Experimental/ControlRig/Source/ControlRigEditor/Private/Editor/ControlRigEditor.cpp b/Engine/Plugins/Experimental/ControlRig/Source/ControlRigEditor/Private/Editor/ControlRigEditor.cpp
--- a/Engine/Plugins/Experimental/ControlRig/Source/ControlRigEditor/Private/Editor/ControlRigEditor.cpp
+++ b/Engine/Plugins/Experimental/ControlRig/Source/ControlRigEditor/Private/Editor/ControlRigEditor.cpp
@@ -457,21 +457,23 @@ void FControlRigEditor::BindCommands()
 
 	GetToolkitCommands()->MapAction(
 		FControlRigBlueprintCommands::Get().ResumeExecution,
-		FExecuteAction::CreateLambda([this]()
-        {
-            if (ControlRig)
-            {
-	            ControlRig->ResumeExecution();
-            }
-        }),			
-		FIsActionChecked::CreateLambda([this]()
-		{
-			if (ControlRig && ControlRig->GetVM())
-			{
-				return ControlRig->GetVM()->GetHaltedAtInstruction() != INDEX_NONE;
-			}
-			return false;
-	}));
+		FExecuteAction::CreateSP(this, &FControlRigEditor::HandleBreakpointActionRequested, ERigVMBreakpointAction::Resume),			
+		FIsActionChecked::CreateSP(this, &FControlRigEditor::IsHaltedAtBreakpoint));
+
+	GetToolkitCommands()->MapAction(
+		FControlRigBlueprintCommands::Get().StepOver,
+		FExecuteAction::CreateSP(this, &FControlRigEditor::HandleBreakpointActionRequested, ERigVMBreakpointAction::StepOver),
+		FIsActionChecked::CreateSP(this, &FControlRigEditor::IsHaltedAtBreakpoint));
+
+	GetToolkitCommands()->MapAction(
+		FControlRigBlueprintCommands::Get().StepInto,
+		FExecuteAction::CreateSP(this, &FControlRigEditor::HandleBreakpointActionRequested, ERigVMBreakpointAction::StepInto),
+		FIsActionChecked::CreateSP(this, &FControlRigEditor::IsHaltedAtBreakpoint));
+
+	GetToolkitCommands()->MapAction(
+		FControlRigBlueprintCommands::Get().StepOut,
+		FExecuteAction::CreateSP(this, &FControlRigEditor::HandleBreakpointActionRequested, ERigVMBreakpointAction::StepOut),
+		FIsActionChecked::CreateSP(this, &FControlRigEditor::IsHaltedAtBreakpoint));
 
 }
 
@@ -637,9 +639,21 @@ void FControlRigEditor::FillToolbar(FToolBarBuilder& ToolbarBuilder)
 			LOCTEXT("ExecutionMode_ToolTip", "Pick between different execution modes for testing the Control Rig"),
 			FSlateIcon(FEditorStyle::GetStyleSetName(), "LevelEditor.Recompile"),
 			true);
+
+
 		ToolbarBuilder.AddToolBarButton(FControlRigBlueprintCommands::Get().ResumeExecution,
 			NAME_None, TAttribute<FText>(), TAttribute<FText>(), FSlateIcon(FControlRigEditorStyle::Get().GetStyleSetName(), "ControlRig.ResumeExecution"));
 
+		ToolbarBuilder.AddToolBarButton(FControlRigBlueprintCommands::Get().StepOver,
+			NAME_None, TAttribute<FText>(), TAttribute<FText>(), FSlateIcon(FControlRigEditorStyle::Get().GetStyleSetName(), "ControlRig.StepOver"));
+
+		ToolbarBuilder.AddToolBarButton(FControlRigBlueprintCommands::Get().StepInto,
+			NAME_None, TAttribute<FText>(), TAttribute<FText>(), FSlateIcon(FControlRigEditorStyle::Get().GetStyleSetName(), "ControlRig.StepInto"));
+
+		ToolbarBuilder.AddToolBarButton(FControlRigBlueprintCommands::Get().StepOut,
+			NAME_None, TAttribute<FText>(), TAttribute<FText>(), FSlateIcon(FControlRigEditorStyle::Get().GetStyleSetName(), "ControlRig.StepOut"));
+		
+
 	}
 	ToolbarBuilder.EndSection();
 }
@@ -2246,29 +2260,32 @@ void FControlRigEditor::HandleControlRigExecutedEvent(UControlRig* InControlRig,
 	UpdateGraphCompilerErrors();
 }
 
-void FControlRigEditor::HandleControlRigExecutionHalted(const int32 InstructionIndex, URigVMNode* InNode)
+void FControlRigEditor::HandleControlRigExecutionHalted(const int32 InstructionIndex, UObject* InNodeObject)
 {
-	if (HaltedAtNode != nullptr)
+	if (URigVMNode* InNode = Cast<URigVMNode>(InNodeObject))
 	{
-		HaltedAtNode->SetExecutionIsHaltedAtThisNode(false);
-	}
-	HaltedAtNode = InNode;
+		if (HaltedAtNode != nullptr)
+		{
+			HaltedAtNode->SetExecutionIsHaltedAtThisNode(false);
+		}
+		HaltedAtNode = InNode;
 
-	if (InNode == nullptr)
-	{
-		return;
-	}	
-	InNode->SetExecutionIsHaltedAtThisNode(true);
+		if (InNode == nullptr)
+		{
+			return;
+		}	
+		InNode->SetExecutionIsHaltedAtThisNode(true);
 	
-	if (UControlRigBlueprint* Blueprint = GetControlRigBlueprint())
-	{
-		if (Blueprint->GetAllModels().Contains(InNode->GetGraph()))
+		if (UControlRigBlueprint* Blueprint = GetControlRigBlueprint())
 		{
-			if(UControlRigGraph* EdGraph = Cast<UControlRigGraph>(Blueprint->GetEdGraph(InNode->GetGraph())))
+			if (Blueprint->GetAllModels().Contains(InNode->GetGraph()))
 			{
-				if(UEdGraphNode* EdNode = EdGraph->FindNodeForModelNodeName(InNode->GetFName()))
+				if(UControlRigGraph* EdGraph = Cast<UControlRigGraph>(Blueprint->GetEdGraph(InNode->GetGraph())))
 				{
-					JumpToHyperlink(EdNode, false);
+					if(UEdGraphNode* EdNode = EdGraph->FindNodeForModelNodeName(InNode->GetFName()))
+					{
+						JumpToHyperlink(EdNode, false);
+					}
 				}
 			}
 		}
@@ -4793,6 +4810,23 @@ bool FControlRigEditor::OnActionMatchesName(FEdGraphSchemaAction* InAction, cons
 	return false;
 }
 
+void FControlRigEditor::HandleBreakpointActionRequested(const ERigVMBreakpointAction BreakpointAction)
+{
+	if (ControlRig)
+	{
+		ControlRig->ExecuteBreakpointAction(BreakpointAction);
+	}
+}
+
+bool FControlRigEditor::IsHaltedAtBreakpoint() const
+{
+	if (ControlRig && ControlRig->GetVM())
+	{
+		return ControlRig->GetVM()->GetHaltedAtInstruction() != INDEX_NONE;
+	}
+	return false;
+}
+
 void FControlRigEditor::UpdateGraphCompilerErrors()
 {
 	DECLARE_SCOPE_HIERARCHICAL_COUNTER_FUNC()
diff --git a/Engine/Plugins/Experimental/ControlRig/Source/ControlRigEditor/Private/Editor/ControlRigEditor.h b/Engine/Plugins/Experimental/ControlRig/Source/ControlRigEditor/Private/Editor/ControlRigEditor.h
--- a/Engine/Plugins/Experimental/ControlRig/Source/ControlRigEditor/Private/Editor/ControlRigEditor.h
+++ b/Engine/Plugins/Experimental/ControlRig/Source/ControlRigEditor/Private/Editor/ControlRigEditor.h
@@ -213,7 +213,7 @@ protected:
 	void HandleModifiedEvent(ERigVMGraphNotifType InNotifType, URigVMGraph* InGraph, UObject* InSubject);
 	void HandleVMCompiledEvent(UBlueprint* InBlueprint, URigVM* InVM);
 	void HandleControlRigExecutedEvent(UControlRig* InControlRig, const EControlRigState InState, const FName& InEventName);
-	void HandleControlRigExecutionHalted(const int32 InstructionIndex, URigVMNode* InNode);
+	void HandleControlRigExecutionHalted(const int32 InstructionIndex, UObject* InNode);
 
 	// FBaseToolKit overrides
 	void CreateEditorModeManager() override;
@@ -329,6 +329,10 @@ private:
 
 	 virtual bool OnActionMatchesName(FEdGraphSchemaAction* InAction, const FName& InName) const override;
 
+	void HandleBreakpointActionRequested(const ERigVMBreakpointAction BreakpointAction);
+
+	bool IsHaltedAtBreakpoint() const;
+
 protected:
 
 	/** Toolbox hosting widget */
diff --git a/Engine/Plugins/Experimental/ControlRig/Source/ControlRigEditor/Private/Editor/SControlRigStackView.cpp b/Engine/Plugins/Experimental/ControlRig/Source/ControlRigEditor/Private/Editor/SControlRigStackView.cpp
--- a/Engine/Plugins/Experimental/ControlRig/Source/ControlRigEditor/Private/Editor/SControlRigStackView.cpp
+++ b/Engine/Plugins/Experimental/ControlRig/Source/ControlRigEditor/Private/Editor/SControlRigStackView.cpp
@@ -621,7 +621,7 @@ void SControlRigStackView::OnVMCompiled(UBlueprint* InCompiledBlueprint, URigVM*
 	}
 }
 
-void SControlRigStackView::HandleExecutionHalted(const int32 HaltedAtInstruction, URigVMNode* InNode)
+void SControlRigStackView::HandleExecutionHalted(const int32 HaltedAtInstruction, UObject* InNode)
 {	
 	if (HaltedAtInstruction != INDEX_NONE && Operators.Num() > HaltedAtInstruction)
 	{
diff --git a/Engine/Plugins/Experimental/ControlRig/Source/ControlRigEditor/Private/Editor/SControlRigStackView.h b/Engine/Plugins/Experimental/ControlRig/Source/ControlRigEditor/Private/Editor/SControlRigStackView.h
--- a/Engine/Plugins/Experimental/ControlRig/Source/ControlRigEditor/Private/Editor/SControlRigStackView.h
+++ b/Engine/Plugins/Experimental/ControlRig/Source/ControlRigEditor/Private/Editor/SControlRigStackView.h
@@ -98,7 +98,7 @@ private:
 	void OnVMCompiled(UBlueprint* InCompiledBlueprint, URigVM* InCompiledVM);
 
 	//* Focus on the instruction when the execution is halted */
-	void HandleExecutionHalted(const int32 HaltedAtInstruction, URigVMNode* InNode);
+	void HandleExecutionHalted(const int32 HaltedAtInstruction, UObject* InNode);
 
 	/** Search box widget */
 	TSharedPtr<SSearchBox> FilterBox;
diff --git a/Engine/Source/Runtime/RigVM/Private/RigVMCore/RigVM.cpp b/Engine/Source/Runtime/RigVM/Private/RigVMCore/RigVM.cpp
--- a/Engine/Source/Runtime/RigVM/Private/RigVMCore/RigVM.cpp
+++ b/Engine/Source/Runtime/RigVM/Private/RigVMCore/RigVM.cpp
@@ -69,6 +69,7 @@ URigVM::URigVM()
 #if WITH_EDITOR
 	, DebugInfo(nullptr)
 	, HaltedAtInstruction(INDEX_NONE)
+	, HaltedAtInstructionHit(INDEX_NONE)
 #endif
     , FunctionNamesPtr(&FunctionNamesStorage)
     , FunctionsPtr(&FunctionsStorage)
@@ -351,6 +352,7 @@ TArray<FName> URigVM::GetEntryNames() const
 bool URigVM::ResumeExecution()
 {
 	HaltedAtInstruction = INDEX_NONE;
+	HaltedAtInstructionHit = INDEX_NONE;
 	if (DebugInfo)
 	{
 		if (FRigVMBreakpoint* Breakpoint = DebugInfo->FindBreakpoint(Context.InstructionIndex))
@@ -739,6 +741,157 @@ void URigVM::CacheMemoryHandlesIfRequired(FRigVMMemoryContainerPtrArray InMemory
 	}
 }
 
+#if WITH_EDITOR
+bool URigVM::ShouldHaltAtInstruction(const uint16 InstructionIndex)
+{
+	FRigVMByteCode& ByteCode = GetByteCode();
+	
+	if (FRigVMBreakpoint* Breakpoint = DebugInfo->FindBreakpoint(Context.InstructionIndex))
+	{
+		if (DebugInfo->IsActive(Context.InstructionIndex))
+		{
+			switch (CurrentBreakpointAction)
+			{
+				case ERigVMBreakpointAction::None:
+				{
+					// Halted at breakpoint. Check if this is a new breakpoint different from the previous halt.
+					if (HaltedAtInstruction != Context.InstructionIndex ||
+						HaltedAtInstructionHit != DebugInfo->GetBreakpointHits(Context.InstructionIndex))
+					{
+						HaltedAtInstruction = Context.InstructionIndex;
+						HaltedAtInstructionHit = DebugInfo->GetBreakpointHits(HaltedAtInstruction);
+						ExecutionHalted().Broadcast(Context.InstructionIndex, Breakpoint->Subject);
+					}
+					return true;
+				}
+				case ERigVMBreakpointAction::Resume:
+				{
+					CurrentBreakpointAction = ERigVMBreakpointAction::None;
+
+					if (DebugInfo->IsTemporaryBreakpoint(Breakpoint))
+					{
+						DebugInfo->RemoveBreakpoint(Breakpoint->InstructionIndex);
+					}
+					else
+					{
+						DebugInfo->IncrementBreakpointActivationOnHit(Context.InstructionIndex);
+						DebugInfo->HitBreakpoint(Context.InstructionIndex);
+					}
+					return false;
+					break;
+				}
+				case ERigVMBreakpointAction::StepOver:
+				case ERigVMBreakpointAction::StepInto:
+				case ERigVMBreakpointAction::StepOut:
+				{
+					// If we are stepping, check if we were halted at the current instruction, and remember it 
+					if (!DebugInfo->GetSteppingOriginBreakpoint())
+					{
+						DebugInfo->SetSteppingOriginBreakpoint(Breakpoint);
+						const TArray<UObject*>* FullCallstack = ByteCode.GetCallstackForInstruction(Context.InstructionIndex);
+						
+						// We want to keep the callstack up to the node that produced the halt
+						DebugInfo->SetSteppingOriginBreakpointCallstack(TArray<UObject*>(FullCallstack->GetData(), FullCallstack->Find((UObject*)DebugInfo->GetSteppingOriginBreakpoint()->Subject)+1));
+					}							
+					
+					break;	
+				}
+				default:
+				{
+					ensure(false);
+					break;
+				}
+			}
+		}
+		else
+		{
+			DebugInfo->HitBreakpoint(Context.InstructionIndex);
+		}
+	}
+
+	// If we are stepping, and the last active breakpoint was set, check if this is the new temporary breakpoint
+	if (DebugInfo->GetSteppingOriginBreakpoint())
+	{
+		const TArray<UObject*>* CurrentCallstack = ByteCode.GetCallstackForInstruction(Context.InstructionIndex);
+		if (CurrentCallstack && !CurrentCallstack->IsEmpty())
+		{
+			UObject* NewBreakpointNode = nullptr;
+
+			// Find the first difference in the callstack
+			int32 DifferenceIndex = INDEX_NONE;
+			TArray<UObject*>& PreviousCallstack = DebugInfo->GetSteppingOriginBreakpointCallstack();
+			for (int32 i=0; i<PreviousCallstack.Num(); ++i)
+			{
+				if (CurrentCallstack->Num() == i)
+				{
+					DifferenceIndex = i-1;
+					break;
+				}
+				if (PreviousCallstack[i] != CurrentCallstack->operator[](i))
+				{
+					DifferenceIndex = i;
+					break;
+				}
+			}
+
+			if (CurrentBreakpointAction == ERigVMBreakpointAction::StepOver)
+			{
+				if (DifferenceIndex != INDEX_NONE)
+				{
+					NewBreakpointNode = CurrentCallstack->operator[](DifferenceIndex);
+				}
+			}
+			else if (CurrentBreakpointAction == ERigVMBreakpointAction::StepInto)
+			{
+				if (DifferenceIndex == INDEX_NONE)
+				{
+					if (CurrentCallstack->Last() != PreviousCallstack.Last())
+					{
+						NewBreakpointNode = CurrentCallstack->operator[](FMath::Min(PreviousCallstack.Num(), CurrentCallstack->Num()-1));
+					}
+				}
+				else
+				{
+					NewBreakpointNode = CurrentCallstack->operator[](DifferenceIndex);
+				}
+			}
+			else if (CurrentBreakpointAction == ERigVMBreakpointAction::StepOut)
+			{
+				if (DifferenceIndex != INDEX_NONE && DifferenceIndex <= PreviousCallstack.Num() - 2)
+                {
+                	NewBreakpointNode = CurrentCallstack->operator[](DifferenceIndex);
+                }
+			}
+			
+			if (NewBreakpointNode)
+			{
+				if (DebugInfo->IsTemporaryBreakpoint(DebugInfo->GetSteppingOriginBreakpoint()))
+				{
+					DebugInfo->RemoveBreakpoint(DebugInfo->GetSteppingOriginBreakpoint()->InstructionIndex);
+				}
+				else
+				{
+					DebugInfo->IncrementBreakpointActivationOnHit(DebugInfo->GetSteppingOriginBreakpoint()->InstructionIndex);
+					DebugInfo->HitBreakpoint(DebugInfo->GetSteppingOriginBreakpoint()->InstructionIndex);
+				}
+				
+				FRigVMBreakpoint* NewBreakpoint = DebugInfo->AddBreakpoint(Context.InstructionIndex, NewBreakpointNode, true);
+				CurrentBreakpointAction = ERigVMBreakpointAction::None;					
+
+				// Halted at breakpoint. Check if this is a new breakpoint different from the previous halt.
+				HaltedAtInstruction = Context.InstructionIndex;
+				HaltedAtInstructionHit = DebugInfo->GetBreakpointHits(HaltedAtInstruction);
+				ExecutionHalted().Broadcast(Context.InstructionIndex, NewBreakpointNode);
+		
+				return true;
+			}
+		}
+	}
+
+	return false;
+}
+#endif
+
 bool URigVM::Initialize(FRigVMMemoryContainerPtrArray Memory, FRigVMFixedArray<void*> AdditionalArguments)
 {
 	if (ExecutingThreadId != INDEX_NONE)
@@ -1066,24 +1219,9 @@ bool URigVM::Execute(FRigVMMemoryContainerPtrArray Memory, FRigVMFixedArray<void
 	while (Instructions.IsValidIndex(Context.InstructionIndex))
 	{
 #if WITH_EDITOR
-		if (DebugInfo)
+		if (DebugInfo && ShouldHaltAtInstruction(Context.InstructionIndex))
 		{
-			if (FRigVMBreakpoint* Breakpoint = DebugInfo->FindBreakpoint(Context.InstructionIndex))
-			{
-				if (DebugInfo->IsActive(Context.InstructionIndex))
-				{
-					if (HaltedAtInstruction != Context.InstructionIndex)
-					{
-						HaltedAtInstruction = Context.InstructionIndex;
-						ExecutionHalted().Broadcast(Context.InstructionIndex, Breakpoint->Subject);
-					}
-					return true;
-				}
-				else
-				{
-					DebugInfo->HitBreakpoint(Context.InstructionIndex);
-				}
-			}
+			return true;
 		}
 		
 		InstructionVisitedDuringLastRun[Context.InstructionIndex]++;
diff --git a/Engine/Source/Runtime/RigVM/Private/RigVMCore/RigVMDebugInfo.cpp b/Engine/Source/Runtime/RigVM/Private/RigVMCore/RigVMDebugInfo.cpp
new file mode 100644
--- /dev/null
+++ b/Engine/Source/Runtime/RigVM/Private/RigVMCore/RigVMDebugInfo.cpp
@@ -0,0 +1,133 @@
+﻿// Copyright Epic Games, Inc. All Rights Reserved.
+
+#include "RigVMCore/RigVMDebugInfo.h"
+
+void FRigVMDebugInfo::ResetState()
+{
+	BreakpointHits.Empty();
+	BreakpointActivationOnHit.Empty();
+	TemporaryBreakpoint.Clear();
+}
+
+void FRigVMDebugInfo::StartExecution()
+{
+	BreakpointHits.Empty();
+	SteppingOriginBreakpoint = nullptr;
+	SteppingOriginBreakpointCallstack.Empty();
+}
+
+FRigVMBreakpoint* FRigVMDebugInfo::FindBreakpoint(const uint16 InstructionIndex) 
+{
+	if (FRigVMBreakpoint* Found = Breakpoints.Find(InstructionIndex))
+	{
+		return Found;
+	}
+	if (TemporaryBreakpoint.InstructionIndex == InstructionIndex)
+	{
+		return &TemporaryBreakpoint;
+	}
+	return nullptr;
+}
+
+FRigVMBreakpoint* FRigVMDebugInfo::AddBreakpoint(const uint16 InstructionIndex, UObject* InNode, const bool bIsTemporary)
+{
+	if (!Breakpoints.Contains(InstructionIndex))
+	{
+		if (bIsTemporary)
+		{
+			TemporaryBreakpoint = FRigVMBreakpoint(InstructionIndex, InNode);
+				
+			// Do not override the state if it already exists
+			BreakpointActivationOnHit.FindOrAdd(InstructionIndex, 0);
+			BreakpointHits.FindOrAdd(InstructionIndex, 0);
+			return &TemporaryBreakpoint;
+		}
+		else if (FRigVMBreakpoint* NewBP = &Breakpoints.Add(InstructionIndex, FRigVMBreakpoint(InstructionIndex, InNode)))
+		{
+			// Do not override the state if it already exists
+			BreakpointActivationOnHit.FindOrAdd(InstructionIndex, 0);
+			BreakpointHits.FindOrAdd(InstructionIndex, 0);
+			return NewBP;
+		}
+	}
+	return nullptr;
+}
+
+bool FRigVMDebugInfo::RemoveBreakpoint(const uint16 InstructionIndex)
+{
+	if (FRigVMBreakpoint* Found = Breakpoints.Find(InstructionIndex))
+	{
+		Breakpoints.Remove(InstructionIndex);
+		BreakpointHits.Remove(InstructionIndex);
+		BreakpointActivationOnHit.Remove(InstructionIndex);
+		return true;
+	}
+	if (TemporaryBreakpoint.InstructionIndex == InstructionIndex)
+	{
+		TemporaryBreakpoint.Clear();
+		BreakpointHits.Remove(InstructionIndex);
+		BreakpointActivationOnHit.Remove(InstructionIndex);
+		return true;
+	}
+	return false;
+}
+
+bool FRigVMDebugInfo::IsActive(const uint16 InstructionIndex) const
+{
+	const FRigVMBreakpoint* Breakpoint = Breakpoints.Find(InstructionIndex);
+	if (!Breakpoint && TemporaryBreakpoint.InstructionIndex == InstructionIndex)
+	{
+		Breakpoint = &TemporaryBreakpoint;
+	}
+	if (Breakpoint)
+	{
+		if (Breakpoint->bIsActive)
+		{
+			uint16 Hits = 0;
+			uint16 OnHit = 0;
+			if (BreakpointHits.Contains(InstructionIndex))
+			{
+				Hits = BreakpointHits.FindChecked(InstructionIndex);
+			}
+			if (BreakpointActivationOnHit.Contains(InstructionIndex))
+			{
+				OnHit = BreakpointActivationOnHit.FindChecked(InstructionIndex);
+			}
+			return Hits == OnHit;
+		}				
+	}
+	return false;		
+}
+
+void FRigVMDebugInfo::HitBreakpoint(const uint16 InstructionIndex)
+{
+	if (BreakpointHits.Contains(InstructionIndex))
+	{
+		BreakpointHits[InstructionIndex]++;
+	}
+	else
+	{
+		BreakpointHits.Add(InstructionIndex, 1);
+	}
+}
+
+void FRigVMDebugInfo::IncrementBreakpointActivationOnHit(const uint16 InstructionIndex)
+{
+	if (BreakpointActivationOnHit.Contains(InstructionIndex))
+	{
+		BreakpointActivationOnHit[InstructionIndex]++;
+	}
+	else
+	{
+		BreakpointActivationOnHit.Add(InstructionIndex, 1);
+	}
+}
+
+uint16 FRigVMDebugInfo::GetBreakpointHits(const uint16 InstructionIndex) const
+{
+	if (BreakpointActivationOnHit.Contains(InstructionIndex))
+	{
+		return BreakpointActivationOnHit[InstructionIndex];
+	}
+	return 0;
+}
diff --git a/Engine/Source/Runtime/RigVM/Public/RigVMCore/RigVM.h b/Engine/Source/Runtime/RigVM/Public/RigVMCore/RigVM.h
--- a/Engine/Source/Runtime/RigVM/Public/RigVMCore/RigVM.h
+++ b/Engine/Source/Runtime/RigVM/Public/RigVMCore/RigVM.h
@@ -13,6 +13,19 @@
 #endif
 #include "RigVM.generated.h"
 
+// Possible actions when the VM is halted at a breakpoint
+UENUM()
+enum class ERigVMBreakpointAction : uint8
+{
+	None,
+	Resume,
+	StepOver,
+	StepInto,
+	StepOut,
+	Max UMETA(Hidden),
+};
+
+
 // The type of parameter for a VM
 UENUM(BlueprintType)
 enum class ERigVMParameterType : uint8
@@ -124,7 +137,7 @@ public:
 	/** Bindable event for external objects to be notified when the VM reaches an Exit Operation */
 	DECLARE_EVENT(URigVM, FExecutionReachedExitEvent);
 #if WITH_EDITOR
-	DECLARE_EVENT_TwoParams(URigVM, FExecutionHaltedEvent, int32, URigVMNode*);
+	DECLARE_EVENT_TwoParams(URigVM, FExecutionHaltedEvent, int32, UObject*);
 #endif
 
 	URigVM();
@@ -529,6 +542,8 @@ public:
 	void SetDebugInfo(FRigVMDebugInfo* InDebugInfo) { DebugInfo = InDebugInfo; }
 
 	int32 GetHaltedAtInstruction() const { return HaltedAtInstruction; }
+
+	void SetBreakpointAction(const ERigVMBreakpointAction& Action) { CurrentBreakpointAction = Action; }
 #endif
 
 private:
@@ -547,6 +562,10 @@ private:
 #if WITH_EDITOR
 	FRigVMDebugInfo* DebugInfo;
 	int32 HaltedAtInstruction;
+	int32 HaltedAtInstructionHit;
+	ERigVMBreakpointAction CurrentBreakpointAction;
+
+	bool ShouldHaltAtInstruction(const uint16 InstructionIndex);
 #endif
 
 	UPROPERTY()
diff --git a/Engine/Source/Runtime/RigVM/Public/RigVMCore/RigVMDebugInfo.h b/Engine/Source/Runtime/RigVM/Public/RigVMCore/RigVMDebugInfo.h
--- a/Engine/Source/Runtime/RigVM/Public/RigVMCore/RigVMDebugInfo.h
+++ b/Engine/Source/Runtime/RigVM/Public/RigVMCore/RigVMDebugInfo.h
@@ -14,20 +14,30 @@ struct RIGVM_API FRigVMBreakpoint
 
 	FRigVMBreakpoint()
 	: bIsActive(true)
+	, InstructionIndex(INDEX_NONE)
 	, Subject(nullptr)
 	{
 
 	}
 	
-	FRigVMBreakpoint(URigVMNode* InNode)
+	FRigVMBreakpoint(const uint16 InInstructionIndex, UObject* InNode)
 	: bIsActive(true)
+	, InstructionIndex(InInstructionIndex)
 	, Subject(InNode)
 	{
 
 	}
 
+	FORCEINLINE void Clear()
+	{
+		bIsActive = true;
+		InstructionIndex = INDEX_NONE;
+		Subject = nullptr;
+	}
+
 	bool bIsActive;
-	URigVMNode* Subject;
+	uint16 InstructionIndex;
+	UObject* Subject;
 };
 
 USTRUCT()
@@ -36,19 +46,13 @@ struct RIGVM_API FRigVMDebugInfo
 	GENERATED_BODY()
 
 	FORCEINLINE FRigVMDebugInfo()
+		: SteppingOriginBreakpoint(nullptr)
 	{
 	}
 
-	FORCEINLINE void ResetState()
-	{
-		BreakpointHits.Empty();
-		BreakpointActivationOnHit.Empty();
-	}
+	void ResetState();
 
-	FORCEINLINE void StartExecution()
-	{
-		BreakpointHits.Empty();
-	}
+	void StartExecution();
 
 	FORCEINLINE void Clear()
 	{
@@ -56,80 +60,46 @@ struct RIGVM_API FRigVMDebugInfo
 		// Do not remove state
 	}
 
-	FRigVMBreakpoint* FindBreakpoint(const uint16 InstructionIndex) 
-	{
-		return Breakpoints.Find(InstructionIndex);
-	}
+	FRigVMBreakpoint* FindBreakpoint(const uint16 InstructionIndex);
 	
-	FRigVMBreakpoint* AddBreakpoint(const uint16 InstructionIndex, URigVMNode* InNode)
-	{
-		if (!Breakpoints.Contains(InstructionIndex))
-		{
-			if (FRigVMBreakpoint* NewBP = &Breakpoints.Add(InstructionIndex, InNode))
-			{
-				// Do not override the state if it already exists
-				BreakpointActivationOnHit.FindOrAdd(InstructionIndex, 0);
-				BreakpointHits.FindOrAdd(InstructionIndex, 0);
-				return NewBP;
-			}
-		}
-		return nullptr;
-	}
+	FRigVMBreakpoint* AddBreakpoint(const uint16 InstructionIndex, UObject* InNode, const bool bIsTemporary = false);
 
-	TMap<uint16 , FRigVMBreakpoint> GetBreakpoints() const { return Breakpoints; }
+	bool RemoveBreakpoint(const uint16 InstructionIndex);
 
-	void SetBreakpoints(const TMap<uint16 , FRigVMBreakpoint>& InBreakpoints) { Breakpoints = InBreakpoints; }
-
-	bool HasBreakpoint(const uint16 InstructionIndex) { return Breakpoints.Contains(InstructionIndex); }
+	TMap<uint16 , FRigVMBreakpoint> GetBreakpoints() const { return Breakpoints; }
 
-	bool IsActive(const uint16 InstructionIndex)
+	void SetBreakpoints(const TMap<uint16 , FRigVMBreakpoint>& InBreakpoints)
 	{
-		if (FRigVMBreakpoint* BP = Breakpoints.Find(InstructionIndex))
-		{
-			if (BP->bIsActive)
-			{
-				uint16 Hits = 0;
-				uint16 OnHit = 0;
-				if (BreakpointHits.Contains(InstructionIndex))
-				{
-					Hits = BreakpointHits.FindChecked(InstructionIndex);
-				}
-				if (BreakpointActivationOnHit.Contains(InstructionIndex))
-				{
-					OnHit = BreakpointActivationOnHit.FindChecked(InstructionIndex);
-				}
-				return Hits == OnHit;
-			}				
-		}
-		return false;		
+		Breakpoints = InBreakpoints;
 	}
 
-	void HitBreakpoint(const uint16 InstructionIndex)
+	bool IsTemporaryBreakpoint(FRigVMBreakpoint* Breakpoint) const
 	{
-		if (BreakpointHits.Contains(InstructionIndex))
-		{
-			BreakpointHits[InstructionIndex]++;
-		}
-		else
-		{
-			BreakpointHits.Add(InstructionIndex, 1);
-		}
+		return Breakpoint == &TemporaryBreakpoint;
 	}
 
-	void IncrementBreakpointActivationOnHit(const uint16 InstructionIndex)
-	{
-		if (BreakpointActivationOnHit.Contains(InstructionIndex))
-		{
-			BreakpointActivationOnHit[InstructionIndex]++;
-		}
-		else
-		{
-			BreakpointActivationOnHit.Add(InstructionIndex, 1);
-		}
-	}
+	bool IsActive(const uint16 InstructionIndex) const;
+
+	void HitBreakpoint(const uint16 InstructionIndex);
+
+	void IncrementBreakpointActivationOnHit(const uint16 InstructionIndex);
+
+	uint16 GetBreakpointHits(const uint16 InstructionIndex) const;
+
+	FRigVMBreakpoint* GetSteppingOriginBreakpoint() const { return SteppingOriginBreakpoint; }
+
+	void SetSteppingOriginBreakpoint(FRigVMBreakpoint* Breakpoint) { SteppingOriginBreakpoint = Breakpoint; }
+
+	TArray<UObject*>& GetSteppingOriginBreakpointCallstack() { return SteppingOriginBreakpointCallstack; }
+
+	void SetSteppingOriginBreakpointCallstack(TArray<UObject*> Callstack) { SteppingOriginBreakpointCallstack = Callstack; }
 
 private:
 	TMap<uint16 , FRigVMBreakpoint> Breakpoints;
+	FRigVMBreakpoint TemporaryBreakpoint;
 	TMap<uint16 , uint16> BreakpointActivationOnHit; // After how many instruction executions, this breakpoint becomes active
 	TMap<uint16 , uint16> BreakpointHits; // How many times this instruction has been executed
+
+	FRigVMBreakpoint* SteppingOriginBreakpoint;
+	TArray<UObject*> SteppingOriginBreakpointCallstack;
 };
