Control Rig: Refactor CopyOp to contain all required information for a copy

Also fixes UE-115756

diff --git a/Engine/Source/Developer/RigVMDeveloper/Private/RigVMCompiler/RigVMCompiler.cpp b/Engine/Source/Developer/RigVMDeveloper/Private/RigVMCompiler/RigVMCompiler.cpp
--- a/Engine/Source/Developer/RigVMDeveloper/Private/RigVMCompiler/RigVMCompiler.cpp
+++ b/Engine/Source/Developer/RigVMDeveloper/Private/RigVMCompiler/RigVMCompiler.cpp
@@ -805,7 +805,7 @@ void URigVMCompiler::TraverseAssign(const FRigVMAssignExprAST* InExpr, FRigVMCom
 			Local::SetupRegisterOffset(WorkData.VM, InExpr->GetTargetPin(), Target, TargetExpr, false);
 		}
 
-		WorkData.VM->GetByteCode().AddCopyOp(Source, Target);
+		WorkData.VM->GetByteCode().AddCopyOp(WorkData.VM->GetCopyOpForOperands(Source, Target));
 		int32 InstructionIndex = WorkData.VM->GetByteCode().GetNumInstructions() - 1;
 
 		if (Settings.SetupNodeInstructionIndex)
@@ -960,7 +960,7 @@ void URigVMCompiler::TraverseIf(const FRigVMIfExprAST* InExpr, FRigVMCompilerWor
 
 	FRigVMOperand& TrueOperand = WorkData.ExprToOperand.FindChecked(TrueExpr);
 
-	WorkData.VM->GetByteCode().AddCopyOp(TrueOperand, ResultOperand);
+	WorkData.VM->GetByteCode().AddCopyOp(WorkData.VM->GetCopyOpForOperands(TrueOperand, ResultOperand));
 	if (Settings.SetupNodeInstructionIndex)
 	{
 		WorkData.VM->GetByteCode().SetSubject(WorkData.VM->GetByteCode().GetNumInstructions() - 1, Callstack.GetCallPath(), Callstack.GetStack());
@@ -987,7 +987,7 @@ void URigVMCompiler::TraverseIf(const FRigVMIfExprAST* InExpr, FRigVMCompilerWor
 
 	FRigVMOperand& FalseOperand = WorkData.ExprToOperand.FindChecked(FalseExpr);
 
-	WorkData.VM->GetByteCode().AddCopyOp(FalseOperand, ResultOperand);
+	WorkData.VM->GetByteCode().AddCopyOp(WorkData.VM->GetCopyOpForOperands(FalseOperand, ResultOperand));
 	if (Settings.SetupNodeInstructionIndex)
 	{
 		WorkData.VM->GetByteCode().SetSubject(WorkData.VM->GetByteCode().GetNumInstructions() - 1, Callstack.GetCallPath(), Callstack.GetStack());
@@ -1108,7 +1108,7 @@ void URigVMCompiler::TraverseSelect(const FRigVMSelectExprAST* InExpr, FRigVMCom
 
 		// add copy op to copy the result
 		FRigVMOperand& CaseOperand = WorkData.ExprToOperand.FindChecked(CaseExpressions[CaseIndex]);
-		WorkData.VM->GetByteCode().AddCopyOp(CaseOperand, ResultOperand);
+		WorkData.VM->GetByteCode().AddCopyOp(WorkData.VM->GetCopyOpForOperands(CaseOperand, ResultOperand));
 		if (Settings.SetupNodeInstructionIndex)
 		{
 			WorkData.VM->GetByteCode().SetSubject(WorkData.VM->GetByteCode().GetNumInstructions() - 1, Callstack.GetCallPath(), Callstack.GetStack());
@@ -1237,14 +1237,17 @@ void URigVMCompiler::InitializeLocalVariables(const FRigVMExprAST* InExpr, FRigV
 						{
 							FString TargetPath = FString::Printf(TEXT("LocalVariable::%s|%s"), *Node->GetNodePath(), *Variable.Name.ToString());
 							FString SourcePath = FString::Printf(TEXT("LocalVariable::%s|%s::Const"), *Node->GetContainedGraph()->GetGraphName(), *Variable.Name.ToString());
-							FRigVMOperand* Target = WorkData.PinPathToOperand->Find(TargetPath);
-							FRigVMOperand* Source = WorkData.PinPathToOperand->Find(SourcePath);
-							if (Source && Target) 
+							FRigVMOperand* TargetPtr = WorkData.PinPathToOperand->Find(TargetPath);
+							FRigVMOperand* SourcePtr = WorkData.PinPathToOperand->Find(SourcePath);
+							if (SourcePtr && TargetPtr) 
 							{
+								const FRigVMOperand& Source = *SourcePtr;
+								const FRigVMOperand& Target = *TargetPtr;
+
 								bool bAlreadyCopied = false;
 								for (const FRigVMCopyOp& CopyOp : PreviousCopyInstructions)
 								{
-									if (CopyOp.Source == *Source && CopyOp.Target == *Target)
+									if (CopyOp.Source == Source && CopyOp.Target == Target)
 									{
 										bAlreadyCopied = true;
 										break;
@@ -1253,7 +1256,7 @@ void URigVMCompiler::InitializeLocalVariables(const FRigVMExprAST* InExpr, FRigV
 
 								if (!bAlreadyCopied)
 								{
-									ByteCode.AddCopyOp(*Source, *Target);
+									ByteCode.AddCopyOp(WorkData.VM->GetCopyOpForOperands(Source, Target));
 								}	
 							}					
 						}
@@ -1280,11 +1283,13 @@ void URigVMCompiler::InitializeLocalVariables(const FRigVMExprAST* InExpr, FRigV
 					{
 						FString TargetPath = FString::Printf(TEXT("LocalVariable::%s"), *Variable.Name.ToString());
 						FString SourcePath = FString::Printf(TEXT("LocalVariable::|%s::Const"), *Variable.Name.ToString());
-						FRigVMOperand* Target = WorkData.PinPathToOperand->Find(TargetPath);
-						FRigVMOperand* Source = WorkData.PinPathToOperand->Find(SourcePath);
-						if (Source && Target)
+						FRigVMOperand* TargetPtr = WorkData.PinPathToOperand->Find(TargetPath);
+						FRigVMOperand* SourcePtr = WorkData.PinPathToOperand->Find(SourcePath);
+						if (SourcePtr && TargetPtr)
 						{
-							ByteCode.AddCopyOp(*Source, *Target);
+							const FRigVMOperand& Source = *SourcePtr;
+							const FRigVMOperand& Target = *TargetPtr;
+							ByteCode.AddCopyOp(WorkData.VM->GetCopyOpForOperands(Source, Target));
 						}					
 					}
 				}
diff --git a/Engine/Source/Runtime/Core/Public/UObject/UE5MainStreamObjectVersion.h b/Engine/Source/Runtime/Core/Public/UObject/UE5MainStreamObjectVersion.h
--- a/Engine/Source/Runtime/Core/Public/UObject/UE5MainStreamObjectVersion.h
+++ b/Engine/Source/Runtime/Core/Public/UObject/UE5MainStreamObjectVersion.h
@@ -139,6 +139,9 @@ struct CORE_API FUE5MainStreamObjectVersion
 		// Switch FTextureSource to use virtualized bulkdata
 		TextureSourceVirtualization,
 
+		// RigVM to store more information alongside the Copy Operator
+		RigVMCopyOpStoreNumBytes,
+
 		// -----<new versions can be added above this line>-------------------------------------------------
 		VersionPlusOne,
 		LatestVersion = VersionPlusOne - 1
diff --git a/Engine/Source/Runtime/RigVM/Private/RigVMCore/RigVM.cpp b/Engine/Source/Runtime/RigVM/Private/RigVMCore/RigVM.cpp
--- a/Engine/Source/Runtime/RigVM/Private/RigVMCore/RigVM.cpp
+++ b/Engine/Source/Runtime/RigVM/Private/RigVMCore/RigVM.cpp
@@ -3,6 +3,7 @@
 #include "RigVMCore/RigVM.h"
 #include "UObject/Package.h"
 #include "UObject/AnimObjectVersion.h"
+#include "UObject/UE5MainStreamObjectVersion.h"
 #include "HAL/PlatformTLS.h"
 
 void FRigVMParameter::Serialize(FArchive& Ar)
@@ -97,6 +98,7 @@ URigVM::~URigVM()
 void URigVM::Serialize(FArchive& Ar)
 {
 	Ar.UsingCustomVersion(FAnimObjectVersion::GUID);
+	Ar.UsingCustomVersion(FUE5MainStreamObjectVersion::GUID);
 
 	if (Ar.CustomVer(FAnimObjectVersion::GUID) < FAnimObjectVersion::StoreMarkerNamesOnSkeleton)
 	{
@@ -140,6 +142,12 @@ void URigVM::Load(FArchive& Ar)
 	Ar << ByteCodeStorage;
 	Ar << Parameters;
 
+	if (Ar.CustomVer(FAnimObjectVersion::GUID) < FUE5MainStreamObjectVersion::RigVMCopyOpStoreNumBytes)
+	{
+		Reset();
+		return;
+	}
+
 	if (WorkMemoryStorage.bEncounteredErrorDuringLoad ||
 		LiteralMemoryStorage.bEncounteredErrorDuringLoad)
 	{
@@ -594,76 +602,7 @@ void URigVM::CacheMemoryHandlesIfRequired(FRigVMMemoryContainerPtrArray InMemory
 				CacheSingleMemoryHandle(Op.Source);
 				CacheSingleMemoryHandle(Op.Target);
 
-				uint16 NumBytes = 0;
-				ERigVMRegisterType TargetType = ERigVMRegisterType::Invalid;
-				UScriptStruct* ScriptStruct = nullptr;
-
-				if (Op.Target.GetMemoryType() == ERigVMMemoryType::External)
-				{
-					ensure(ExternalVariables.IsValidIndex(Op.Target.GetRegisterIndex()));
-					const FRigVMExternalVariable& ExternalVariable = ExternalVariables[Op.Target.GetRegisterIndex()];
-
-					NumBytes = ExternalVariable.Size;
-					TargetType = ERigVMRegisterType::Plain;
-					if (UScriptStruct* ExternalScriptStruct = Cast<UScriptStruct>(ExternalVariable.TypeObject))	
-					{
-						TargetType = ERigVMRegisterType::Struct;
-						ScriptStruct = ExternalScriptStruct;
-					}
-					else if (ExternalVariable.TypeName == TEXT("FString"))
-					{
-						TargetType = ERigVMRegisterType::String;
-					}
-					else if (ExternalVariable.TypeName == TEXT("FName"))
-					{
-						TargetType = ERigVMRegisterType::Name;
-					}
-				}
-				else
-				{
-					const FRigVMRegister& TargetRegister = CachedMemory[Op.Target.GetContainerIndex()]->Registers[Op.Target.GetRegisterIndex()];
-					NumBytes = TargetRegister.GetNumBytesPerSlice();
-
-					TargetType = TargetRegister.Type;
-
-					if (Op.Target.GetRegisterOffset() == INDEX_NONE)
-					{
-						if (TargetRegister.IsArray())
-						{
-							const FRigVMRegister& SourceRegister = CachedMemory[Op.Source.GetContainerIndex()]->Registers[Op.Source.GetRegisterIndex()];
-							if (!SourceRegister.IsArray())
-							{
-								if (Op.Source.GetRegisterOffset() == INDEX_NONE)
-								{
-									NumBytes = TargetRegister.ElementSize;
-								}
-								else
-								{
-									const FRigVMRegisterOffset& SourceOffset = CachedMemory[Op.Source.GetContainerIndex()]->RegisterOffsets[Op.Source.GetRegisterOffset()];
-									if (SourceOffset.GetCPPType() != TEXT("TArray"))
-									{
-										NumBytes = SourceOffset.GetElementSize();
-									}
-								}
-							}
-						}
-					}
-					else
-					{
-						TargetType = CachedMemory[Op.Target.GetContainerIndex()]->RegisterOffsets[Op.Target.GetRegisterOffset()].GetType();
-						NumBytes = CachedMemory[Op.Target.GetContainerIndex()]->RegisterOffsets[Op.Target.GetRegisterOffset()].GetElementSize();
-					}
-
-					if (TargetType == ERigVMRegisterType::Struct)
-					{
-						ScriptStruct = CachedMemory[Op.Target.GetContainerIndex()]->GetScriptStruct(Op.Target.GetRegisterIndex(), Op.Target.GetRegisterOffset());
-					}
-				}
-
-				CachedMemoryHandles.Add(FRigVMMemoryHandle((uint8*)reinterpret_cast<void*>(NumBytes)));
-				CachedMemoryHandles.Add(FRigVMMemoryHandle((uint8*)reinterpret_cast<void*>((uint16)TargetType)));
-
-				if (TargetType == ERigVMRegisterType::Struct)
+				if (UScriptStruct* ScriptStruct = GetScriptStructForCopyOp(Op))
 				{
 					CachedMemoryHandles.Add(FRigVMMemoryHandle((uint8*)ScriptStruct));
 				}
@@ -1052,8 +991,7 @@ bool URigVM::Initialize(FRigVMMemoryContainerPtrArray Memory, FRigVMFixedArray<v
 				void* SourcePtr = SourceHandle;
 				void* TargetPtr = TargetHandle;
 
-				uint64 NumBytes = reinterpret_cast<uint64>(CachedMemoryHandles[FirstHandleForInstruction[Context.InstructionIndex] + 2].GetData());
-				ERigVMRegisterType MemoryType = (ERigVMRegisterType)reinterpret_cast<uint64>(CachedMemoryHandles[FirstHandleForInstruction[Context.InstructionIndex] + 3].GetData());
+				const uint64 NumBytes = (uint64)Op.NumBytes;
 
 				if (TargetHandle.Type == FRigVMMemoryHandle::Dynamic)
 				{
@@ -1077,7 +1015,7 @@ bool URigVM::Initialize(FRigVMMemoryContainerPtrArray Memory, FRigVMFixedArray<v
 					TargetPtr = (*Storage)[ArrayIndex].GetData();
 				}
 
-				switch (MemoryType)
+				switch (Op.RegisterType)
 				{
 					case ERigVMRegisterType::Plain:
 					{
@@ -1108,7 +1046,7 @@ bool URigVM::Initialize(FRigVMMemoryContainerPtrArray Memory, FRigVMFixedArray<v
 					}
 					case ERigVMRegisterType::Struct:
 					{
-						UScriptStruct* ScriptStruct = (UScriptStruct*)CachedMemoryHandles[FirstHandleForInstruction[Context.InstructionIndex] + 4].GetData();
+						UScriptStruct* ScriptStruct = (UScriptStruct*)CachedMemoryHandles[FirstHandleForInstruction[Context.InstructionIndex] + 2].GetData();
 						int32 NumStructs = NumBytes / ScriptStruct->GetStructureSize();
 						if (NumStructs > 0 && TargetPtr)
 						{
@@ -1360,8 +1298,7 @@ bool URigVM::Execute(FRigVMMemoryContainerPtrArray Memory, FRigVMFixedArray<void
 				void* SourcePtr = SourceHandle;
 				void* TargetPtr = TargetHandle;
 
-				const uint64 NumBytes = reinterpret_cast<uint64>(CachedMemoryHandles[FirstHandleForInstruction[Context.InstructionIndex] + 2].GetData());
-				const ERigVMRegisterType MemoryType = (ERigVMRegisterType)reinterpret_cast<uint64>(CachedMemoryHandles[FirstHandleForInstruction[Context.InstructionIndex] + 3].GetData());
+				const uint64 NumBytes = (uint64)Op.NumBytes;
 
 				if (TargetHandle.Type == FRigVMMemoryHandle::Dynamic)
 				{
@@ -1385,7 +1322,7 @@ bool URigVM::Execute(FRigVMMemoryContainerPtrArray Memory, FRigVMFixedArray<void
 					TargetPtr = (*Storage)[ArrayIndex].GetData();
 				}
 
-				switch (MemoryType)
+				switch (Op.RegisterType)
 				{
 					case ERigVMRegisterType::Plain:
 					{
@@ -1416,7 +1353,7 @@ bool URigVM::Execute(FRigVMMemoryContainerPtrArray Memory, FRigVMFixedArray<void
 					}
 					case ERigVMRegisterType::Struct:
 					{
-						UScriptStruct* ScriptStruct = (UScriptStruct*)CachedMemoryHandles[FirstHandleForInstruction[Context.InstructionIndex] + 4].GetData();
+						UScriptStruct* ScriptStruct = (UScriptStruct*)CachedMemoryHandles[FirstHandleForInstruction[Context.InstructionIndex] + 2].GetData();
 						const int32 NumStructs = NumBytes / ScriptStruct->GetStructureSize();
 						if (NumStructs > 0 && TargetPtr)
 						{
@@ -2154,3 +2091,95 @@ void URigVM::CopyOperandForDebuggingImpl(const FRigVMOperand& InArg, const FRigV
 	
 #endif
 }
+
+FRigVMCopyOp URigVM::GetCopyOpForOperands(const FRigVMOperand& InSource, const FRigVMOperand& InTarget)
+{
+	TPair<ERigVMRegisterType, uint16> SourceCopyInfo = GetCopyInfoForOperand(InSource);
+	TPair<ERigVMRegisterType, uint16> TargetCopyInfo = GetCopyInfoForOperand(InTarget);
+
+	check(SourceCopyInfo.Key != ERigVMRegisterType::Invalid);
+	check(SourceCopyInfo.Value > 0);
+	check(TargetCopyInfo.Key != ERigVMRegisterType::Invalid);
+	check(TargetCopyInfo.Value > 0);
+	//check(SourceCopyInfo.Value == TargetCopyInfo.Value);
+
+	return FRigVMCopyOp(InSource, InTarget, TargetCopyInfo.Value, TargetCopyInfo.Key);
+}
+
+TPair<ERigVMRegisterType, uint16> URigVM::GetCopyInfoForOperand(const FRigVMOperand& InOperand)
+{
+	if(CachedMemory.IsEmpty())
+	{
+		CachedMemory.Add(WorkMemoryPtr);
+		CachedMemory.Add(LiteralMemoryPtr);
+		CachedMemory.Add(DebugMemoryPtr);
+	}
+	
+	ERigVMRegisterType RegisterType = ERigVMRegisterType::Invalid;
+	uint16 NumBytesToCopy = 0;
+
+	if (InOperand.GetRegisterOffset() != INDEX_NONE)
+	{
+		const FRigVMRegisterOffset& RegisterOffset = CachedMemory[InOperand.GetContainerIndex()]->RegisterOffsets[InOperand.GetRegisterOffset()];
+		RegisterType = RegisterOffset.GetType();
+		NumBytesToCopy = RegisterOffset.GetElementSize();
+	}
+	else if (InOperand.GetMemoryType() == ERigVMMemoryType::External)
+	{
+		ensure(ExternalVariables.IsValidIndex(InOperand.GetRegisterIndex()));
+		const FRigVMExternalVariable& ExternalVariable = ExternalVariables[InOperand.GetRegisterIndex()];
+
+		NumBytesToCopy = ExternalVariable.Size;
+		RegisterType = ERigVMRegisterType::Plain;
+		
+		if (UScriptStruct* ExternalScriptStruct = Cast<UScriptStruct>(ExternalVariable.TypeObject))	
+		{
+			RegisterType = ERigVMRegisterType::Struct;
+		}
+		else if (ExternalVariable.TypeName == TEXT("FString"))
+		{
+			RegisterType = ERigVMRegisterType::String;
+		}
+		else if (ExternalVariable.TypeName == TEXT("FName"))
+		{
+			RegisterType = ERigVMRegisterType::Name;
+		}
+	}
+	else
+	{
+		const FRigVMRegister& Register = CachedMemory[InOperand.GetContainerIndex()]->Registers[InOperand.GetRegisterIndex()];
+
+		RegisterType = Register.Type;
+		NumBytesToCopy = Register.GetNumBytesPerSlice();
+	}
+
+	return TPair<ERigVMRegisterType, uint16>(RegisterType, NumBytesToCopy); 
+}
+
+UScriptStruct* URigVM::GetScriptStructForCopyOp(const FRigVMCopyOp& InCopyOp) const
+{
+	UScriptStruct* SourceScriptStruct = GetScripStructForOperand(InCopyOp.Source);
+#if WITH_EDITOR
+	UScriptStruct* TargetScriptStruct = GetScripStructForOperand(InCopyOp.Target);
+	check(SourceScriptStruct == TargetScriptStruct);
+#endif
+	return SourceScriptStruct;
+}
+
+UScriptStruct* URigVM::GetScripStructForOperand(const FRigVMOperand& InOperand) const
+{
+	if(InOperand.GetRegisterOffset() != INDEX_NONE)
+	{
+		const FRigVMRegisterOffset& RegisterOffset = CachedMemory[InOperand.GetContainerIndex()]->RegisterOffsets[InOperand.GetRegisterOffset()];
+		return RegisterOffset.GetScriptStruct();
+	}
+
+	if(InOperand.GetMemoryType() == ERigVMMemoryType::External)
+	{
+		const FRigVMExternalVariable& ExternalVariable = ExternalVariables[InOperand.GetRegisterIndex()];
+		return Cast<UScriptStruct>(ExternalVariable.TypeObject);
+	}
+
+	return CachedMemory[InOperand.GetContainerIndex()]->GetScriptStruct(InOperand.GetRegisterIndex());
+}
+
diff --git a/Engine/Source/Runtime/RigVM/Private/RigVMCore/RigVMByteCode.cpp b/Engine/Source/Runtime/RigVM/Private/RigVMCore/RigVMByteCode.cpp
--- a/Engine/Source/Runtime/RigVM/Private/RigVMCore/RigVMByteCode.cpp
+++ b/Engine/Source/Runtime/RigVM/Private/RigVMCore/RigVMByteCode.cpp
@@ -2,6 +2,7 @@
 
 #include "RigVMCore/RigVMByteCode.h"
 #include "UObject/AnimObjectVersion.h"
+#include "UObject/UE5MainStreamObjectVersion.h"
 #include "UObject/FortniteMainBranchObjectVersion.h"
 
 void FRigVMExecuteOp::Serialize(FArchive& Ar)
@@ -25,9 +26,30 @@ void FRigVMBinaryOp::Serialize(FArchive& Ar)
 
 void FRigVMCopyOp::Serialize(FArchive& Ar)
 {
+	Ar.UsingCustomVersion(FUE5MainStreamObjectVersion::GUID);
+	
 	Ar << OpCode;
 	Ar << Source;
 	Ar << Target;
+
+	if(Ar.IsLoading())
+	{
+		if (Ar.CustomVer(FAnimObjectVersion::GUID) < FUE5MainStreamObjectVersion::RigVMCopyOpStoreNumBytes)
+		{
+			NumBytes = 0;
+			RegisterType = ERigVMRegisterType::Invalid;
+		}
+		else
+		{
+			Ar << NumBytes;
+			Ar << RegisterType;
+		}
+	}
+	else
+	{
+		Ar << NumBytes;
+		Ar << RegisterType;
+	}
 }
 
 void FRigVMComparisonOp::Serialize(FArchive& Ar)
@@ -129,6 +151,7 @@ void FRigVMByteCode::Serialize(FArchive& Ar)
 {
 	Ar.UsingCustomVersion(FAnimObjectVersion::GUID);
 	Ar.UsingCustomVersion(FFortniteMainBranchObjectVersion::GUID);
+	Ar.UsingCustomVersion(FUE5MainStreamObjectVersion::GUID);
 
 	if (Ar.CustomVer(FAnimObjectVersion::GUID) < FAnimObjectVersion::StoreMarkerNamesOnSkeleton)
 	{
@@ -751,13 +774,21 @@ uint64 FRigVMByteCode::AddTrueOp(const FRigVMOperand& InArg)
 	return AddOp(Op);
 }
 
-uint64 FRigVMByteCode::AddCopyOp(const FRigVMOperand& InSource, const FRigVMOperand& InTarget)
+uint64 FRigVMByteCode::AddCopyOp(const FRigVMOperand& InSource, const FRigVMOperand& InTarget, uint16 InNumBytes, ERigVMRegisterType InTargetType)
 {
-	ensure(InTarget.GetMemoryType() != ERigVMMemoryType::Literal);
-	FRigVMCopyOp Op(InSource, InTarget);
+	check(InTarget.GetMemoryType() != ERigVMMemoryType::Literal);
+	check(InNumBytes > 0);
+	check(InTargetType != ERigVMRegisterType::Invalid);
+	
+	FRigVMCopyOp Op(InSource, InTarget, InNumBytes, InTargetType);
 	return AddOp(Op);
 }
 
+uint64 FRigVMByteCode::AddCopyOp(const FRigVMCopyOp& InCopyOp)
+{
+	return AddCopyOp(InCopyOp.Source, InCopyOp.Target, InCopyOp.NumBytes, InCopyOp.RegisterType);
+}
+
 uint64 FRigVMByteCode::AddIncrementOp(const FRigVMOperand& InArg)
 {
 	ensure(InArg.GetMemoryType() != ERigVMMemoryType::Literal);
@@ -925,6 +956,7 @@ FString FRigVMByteCode::DumpToText() const
 				FString TargetContent;
 				FRigVMOperand::StaticStruct()->ExportText(TargetContent, &Op.Source, nullptr, nullptr, PPF_None, nullptr);
 				Line += FString::Printf(TEXT(", Target %s"), *TargetContent);
+				Line += FString::Printf(TEXT(", NumBytes %d"), (int32)Op.NumBytes);
 				break;
 			}
 			case ERigVMOpCode::Zero:
diff --git a/Engine/Source/Runtime/RigVM/Private/RigVMCore/RigVMMemory.cpp b/Engine/Source/Runtime/RigVM/Private/RigVMCore/RigVMMemory.cpp
--- a/Engine/Source/Runtime/RigVM/Private/RigVMCore/RigVMMemory.cpp
+++ b/Engine/Source/Runtime/RigVM/Private/RigVMCore/RigVMMemory.cpp
@@ -173,10 +173,10 @@ void FRigVMRegisterOffset::Load(FArchive& Ar)
 	}
 }
 
-FRigVMRegisterOffset::FRigVMRegisterOffset(UScriptStruct* InScriptStruct, const FString& InSegmentPath, int32 InInitialOffset, uint16 InElementSize)
+FRigVMRegisterOffset::FRigVMRegisterOffset(UScriptStruct* InScriptStruct, const FString& InSegmentPath, int32 InInitialOffset, uint16 InElementSize, const FName& InCPPType)
 	: Segments()
 	, Type(ERigVMRegisterType::Plain)
-	, CPPType()
+	, CPPType(InCPPType)
 	, ScriptStruct(nullptr)
 	, ParentScriptStruct(nullptr)
 	, ArrayIndex(0)
@@ -244,6 +244,9 @@ FRigVMRegisterOffset::FRigVMRegisterOffset(UScriptStruct* InScriptStruct, const
 				{
 					Offset.Segments.Add(-1);
 					Property = ArrayProperty->Inner;
+
+					Offset.CPPType = *ArrayProperty->Inner->GetCPPType();
+					Offset.ElementSize = ArrayProperty->Inner->ElementSize;
 				}
 
 				if (FStructProperty* StructProperty = CastField<FStructProperty>(Property))
@@ -1691,12 +1694,19 @@ int32 FRigVMMemoryContainer::GetOrAddRegisterOffset(int32 InRegisterIndex, UScri
 		return INDEX_NONE;
 	}
 
-	if (InElementSize == 0 && Registers.IsValidIndex(InRegisterIndex))
+	FName BaseCPPType = NAME_None;
+	if(Registers.IsValidIndex(InRegisterIndex))
 	{
-		InElementSize = (int32)Registers[InRegisterIndex].ElementSize;
+		if(InElementSize == 0)
+		{
+			InElementSize = (int32)Registers[InRegisterIndex].ElementSize;
+		}
+#if WITH_EDITORONLY_DATA
+		BaseCPPType = Registers[InRegisterIndex].BaseCPPType;
+#endif
 	}
 
-	FRigVMRegisterOffset Offset(InScriptStruct, InSegmentPath, InInitialOffset, InElementSize);
+	FRigVMRegisterOffset Offset(InScriptStruct, InSegmentPath, InInitialOffset, InElementSize, BaseCPPType);
 	int32 ExistingIndex = RegisterOffsets.Find(Offset);
 	if (ExistingIndex == INDEX_NONE)
 	{
diff --git a/Engine/Source/Runtime/RigVM/Public/RigVMCore/RigVM.h b/Engine/Source/Runtime/RigVM/Public/RigVMCore/RigVM.h
--- a/Engine/Source/Runtime/RigVM/Public/RigVMCore/RigVM.h
+++ b/Engine/Source/Runtime/RigVM/Public/RigVMCore/RigVM.h
@@ -622,6 +622,11 @@ private:
 
 	void CopyOperandForDebuggingImpl(const FRigVMOperand& InArg, const FRigVMMemoryHandle& InHandle, const FRigVMOperand& InDebugOperand);
 
+	FRigVMCopyOp GetCopyOpForOperands(const FRigVMOperand& InSource, const FRigVMOperand& InTarget);
+	TPair<ERigVMRegisterType, uint16> GetCopyInfoForOperand(const FRigVMOperand& InOperand);
+	UScriptStruct* GetScriptStructForCopyOp(const FRigVMCopyOp& InCopyOp) const;
+	UScriptStruct* GetScripStructForOperand(const FRigVMOperand& InOperand) const;
+	
 	TMap<FRigVMOperand, TArray<FRigVMOperand>> OperandToDebugRegisters;
 
 	int32 ExecutingThreadId;
diff --git a/Engine/Source/Runtime/RigVM/Public/RigVMCore/RigVMByteCode.h b/Engine/Source/Runtime/RigVM/Public/RigVMCore/RigVMByteCode.h
--- a/Engine/Source/Runtime/RigVM/Public/RigVMCore/RigVMByteCode.h
+++ b/Engine/Source/Runtime/RigVM/Public/RigVMCore/RigVMByteCode.h
@@ -228,21 +228,29 @@ struct RIGVM_API FRigVMCopyOp : public FRigVMBaseOp
 	: FRigVMBaseOp(ERigVMOpCode::Copy)
 	, Source()
 	, Target()
+	, NumBytes(0)
+	, RegisterType(ERigVMRegisterType::Invalid)
 	{
 	}
 
 	FRigVMCopyOp(
 		FRigVMOperand InSource,
-		FRigVMOperand InTarget
+		FRigVMOperand InTarget,
+		uint16 InNumBytes,
+		ERigVMRegisterType InRegisterType
 	)
 		: FRigVMBaseOp(ERigVMOpCode::Copy)
 		, Source(InSource)
 		, Target(InTarget)
+		, NumBytes(InNumBytes)
+		, RegisterType(InRegisterType)
 	{
 	}
 
 	FRigVMOperand Source;
 	FRigVMOperand Target;
+	uint16 NumBytes;
+	ERigVMRegisterType RegisterType;
 
 	void Serialize(FArchive& Ar);
 	FORCEINLINE friend FArchive& operator<<(FArchive& Ar, FRigVMCopyOp& P)
@@ -540,7 +548,10 @@ public:
 	uint64 AddTrueOp(const FRigVMOperand& InArg);
 
 	// adds a copy operator to copy the content of a source argument to a target argument
-	uint64 AddCopyOp(const FRigVMOperand& InSource, const FRigVMOperand& InTarget);
+	uint64 AddCopyOp(const FRigVMOperand& InSource, const FRigVMOperand& InTarget, uint16 InNumBytes, ERigVMRegisterType InTargetType);
+
+	// adds a copy operator to copy the content of a source argument to a target argument
+	uint64 AddCopyOp(const FRigVMCopyOp& InCopyOp);
 
 	// adds an increment operator to increment a int32 argument
 	uint64 AddIncrementOp(const FRigVMOperand& InArg);
diff --git a/Engine/Source/Runtime/RigVM/Public/RigVMCore/RigVMMemory.h b/Engine/Source/Runtime/RigVM/Public/RigVMCore/RigVMMemory.h
--- a/Engine/Source/Runtime/RigVM/Public/RigVMCore/RigVMMemory.h
+++ b/Engine/Source/Runtime/RigVM/Public/RigVMCore/RigVMMemory.h
@@ -87,7 +87,14 @@ public:
 	FORCEINLINE_DEBUGGABLE ERigVMMemoryType GetMemoryType() const { return MemoryType; }
 
 	// returns the index of the container of this argument
-	FORCEINLINE_DEBUGGABLE int32 GetContainerIndex() const { return (int32)MemoryType; }
+	FORCEINLINE_DEBUGGABLE int32 GetContainerIndex() const
+	{
+		if(MemoryType == ERigVMMemoryType::External)
+		{
+			return (int)ERigVMMemoryType::Work;
+		}
+		return (int32)MemoryType;
+	}
 
 	// returns the index of the register of this argument
 	FORCEINLINE_DEBUGGABLE int32 GetRegisterIndex() const { return RegisterIndex == UINT16_MAX ? INDEX_NONE : (int32)RegisterIndex; }
@@ -328,7 +335,7 @@ public:
 	}
 
 	// constructs a path given a struct and a segment path
-	FRigVMRegisterOffset(UScriptStruct* InScriptStruct, const FString& InSegmentPath, int32 InInitialOffset = 0, uint16 InElementSize = 0);
+	FRigVMRegisterOffset(UScriptStruct* InScriptStruct, const FString& InSegmentPath, int32 InInitialOffset = 0, uint16 InElementSize = 0, const FName& InCPPType = NAME_None);
 
 	// returns the data pointer within a container
 	uint8* GetData(uint8* InContainer) const;
@@ -488,6 +495,11 @@ public:
 		return Type;
 	}
 
+	FORCEINLINE bool IsDynamic() const
+	{
+		return (Type == Dynamic) || (Type == NestedDynamic);
+	}
+
 private:
 
 	FORCEINLINE_DEBUGGABLE uint8* GetData_Internal_NoOffset(int32 SliceIndex, bool bGetArrayData = false) const
@@ -505,6 +517,10 @@ private:
 			{
 				return Ptr + SliceIndex * Size;
 			}
+			case FType::ArraySize:
+			{
+				return Ptr;
+			}
 			case FType::Dynamic:
 			{
 				ensure(RegisterOffset == nullptr);
