Control Rig: RigVM to recompute num bytes on load

diff --git a/Engine/Source/Runtime/RigVM/Private/RigVMCore/RigVM.cpp b/Engine/Source/Runtime/RigVM/Private/RigVMCore/RigVM.cpp
--- a/Engine/Source/Runtime/RigVM/Private/RigVMCore/RigVM.cpp
+++ b/Engine/Source/Runtime/RigVM/Private/RigVMCore/RigVM.cpp
@@ -164,6 +164,9 @@ void URigVM::Load(FArchive& Ar)
 			ParametersNameMap.Add(Parameters[Index].Name, Index);
 		}
 
+		// rebuild the bytecode to adjust for byte shifts in shipping
+		RebuildByteCodeOnLoad();
+
 		InvalidateCachedMemory();
 	}
 }
@@ -680,6 +683,34 @@ void URigVM::CacheMemoryHandlesIfRequired(FRigVMMemoryContainerPtrArray InMemory
 	}
 }
 
+void URigVM::RebuildByteCodeOnLoad()
+{
+	Instructions = GetByteCode().GetInstructions();
+	for(int32 InstructionIndex = 0; InstructionIndex < Instructions.Num(); InstructionIndex++)
+	{
+		const FRigVMInstruction& Instruction = Instructions[InstructionIndex];
+		switch(Instruction.OpCode)
+		{
+			case ERigVMOpCode::Copy:
+			{
+				// create a local copy of the original op
+				FRigVMCopyOp OldCopyOp = GetByteCode().GetOpAt<FRigVMCopyOp>(Instruction);
+				FRigVMCopyOp& NewCopyOp = GetByteCode().GetOpAt<FRigVMCopyOp>(Instruction);
+				NewCopyOp = GetCopyOpForOperands(OldCopyOp.Source, OldCopyOp.Target);
+				check(OldCopyOp.Source == NewCopyOp.Source);
+				check(OldCopyOp.Target == NewCopyOp.Target);
+				//check(OldCopyOp.NumBytes == NewCopyOp.NumBytes);
+				check(OldCopyOp.RegisterType == NewCopyOp.RegisterType);
+				break;
+			}
+			default:
+			{
+				break;
+			}
+		}
+	}
+}
+
 #if WITH_EDITOR
 bool URigVM::ShouldHaltAtInstruction(const uint16 InstructionIndex)
 {
diff --git a/Engine/Source/Runtime/RigVM/Public/RigVMCore/RigVM.h b/Engine/Source/Runtime/RigVM/Public/RigVMCore/RigVM.h
--- a/Engine/Source/Runtime/RigVM/Public/RigVMCore/RigVM.h
+++ b/Engine/Source/Runtime/RigVM/Public/RigVMCore/RigVM.h
@@ -552,6 +552,7 @@ private:
 	void RefreshInstructionsIfRequired();
 	void InvalidateCachedMemory();
 	void CacheMemoryHandlesIfRequired(FRigVMMemoryContainerPtrArray InMemory);
+	void RebuildByteCodeOnLoad();
 
 	UPROPERTY(transient)
 	FRigVMInstructionArray Instructions;
