Control Rig: Fix compiler regressions introduced in 16402094

diff --git a/Engine/Source/Runtime/RigVM/Private/RigVMCore/RigVM.cpp b/Engine/Source/Runtime/RigVM/Private/RigVMCore/RigVM.cpp
--- a/Engine/Source/Runtime/RigVM/Private/RigVMCore/RigVM.cpp
+++ b/Engine/Source/Runtime/RigVM/Private/RigVMCore/RigVM.cpp
@@ -693,8 +693,17 @@ void URigVM::RebuildByteCodeOnLoad()
 		{
 			case ERigVMOpCode::Copy:
 			{
-				// create a local copy of the original op
 				FRigVMCopyOp OldCopyOp = GetByteCode().GetOpAt<FRigVMCopyOp>(Instruction);
+				if((OldCopyOp.Source.GetMemoryType() == ERigVMMemoryType::External) ||
+					(OldCopyOp.Target.GetMemoryType() == ERigVMMemoryType::External))
+				{
+					if(ExternalVariables.IsEmpty())
+					{
+						break;
+					}
+				}
+					
+				// create a local copy of the original op
 				FRigVMCopyOp& NewCopyOp = GetByteCode().GetOpAt<FRigVMCopyOp>(Instruction);
 				NewCopyOp = GetCopyOpForOperands(OldCopyOp.Source, OldCopyOp.Target);
 				check(OldCopyOp.Source == NewCopyOp.Source);
