ControlRig: Fix debug stepping not working

diff --git a/Engine/Plugins/Experimental/ControlRig/Source/ControlRig/Private/ControlRig.cpp b/Engine/Plugins/Experimental/ControlRig/Source/ControlRig/Private/ControlRig.cpp
--- a/Engine/Plugins/Experimental/ControlRig/Source/ControlRig/Private/ControlRig.cpp
+++ b/Engine/Plugins/Experimental/ControlRig/Source/ControlRig/Private/ControlRig.cpp
@@ -1006,7 +1006,7 @@ void UControlRig::ExecuteUnits(FRigUnitContext& InOutContext, const FName& InEve
 #if WITH_EDITOR
 			if(URigVM* SnapShotVM = GetSnapshotVM(false)) // don't create it for normal runs
 			{
-				if (VM->GetHaltedAtInstruction() != INDEX_NONE)
+				if (VM->GetHaltedAtBreakpoint() != nullptr)
 				{
 					VM->CopyFrom(SnapShotVM, false, false, false, true, true);	
 				}
@@ -2433,14 +2433,14 @@ URigVM* UControlRig::GetSnapshotVM(bool bCreateIfNeeded)
 #endif
 }
 
-void UControlRig::AddBreakpoint(int32 InstructionIndex, URigVMNode* InNode)
+void UControlRig::AddBreakpoint(int32 InstructionIndex, URigVMNode* InNode, uint16 InDepth)
 {
-	DebugInfo.AddBreakpoint(InstructionIndex, InNode);
+	DebugInfo.AddBreakpoint(InstructionIndex, InNode, InDepth);
 }
 
 bool UControlRig::ExecuteBreakpointAction(const ERigVMBreakpointAction BreakpointAction)
 {
-	if (VM->GetHaltedAtInstruction() != INDEX_NONE)
+	if (VM->GetHaltedAtBreakpoint() != nullptr)
 	{
 		VM->SetBreakpointAction(BreakpointAction);
 		return true;
diff --git a/Engine/Plugins/Experimental/ControlRig/Source/ControlRig/Public/ControlRig.h b/Engine/Plugins/Experimental/ControlRig/Source/ControlRig/Public/ControlRig.h
--- a/Engine/Plugins/Experimental/ControlRig/Source/ControlRig/Public/ControlRig.h
+++ b/Engine/Plugins/Experimental/ControlRig/Source/ControlRig/Public/ControlRig.h
@@ -664,8 +664,8 @@ public:
 
 	void SetIsInDebugMode(const bool bValue) { bIsInDebugMode = bValue; }
 	
-	/** Adds a breakpoint in the VM at the InstructionIndex */
-	void AddBreakpoint(int32 InstructionIndex, URigVMNode* InNode);
+	/** Adds a breakpoint in the VM at the InstructionIndex for the Node */
+	void AddBreakpoint(int32 InstructionIndex, URigVMNode* InNode, uint16 InDepth);
 
 	/** If the VM is halted at a breakpoint, it sets a breakpoint action so that
 	 *  it is applied on the next VM execution */
diff --git a/Engine/Plugins/Experimental/ControlRig/Source/ControlRigDeveloper/Private/ControlRigBlueprint.cpp b/Engine/Plugins/Experimental/ControlRig/Source/ControlRigDeveloper/Private/ControlRigBlueprint.cpp
--- a/Engine/Plugins/Experimental/ControlRig/Source/ControlRigDeveloper/Private/ControlRigBlueprint.cpp
+++ b/Engine/Plugins/Experimental/ControlRig/Source/ControlRigDeveloper/Private/ControlRigBlueprint.cpp
@@ -847,7 +847,7 @@ bool UControlRigBlueprint::AddBreakpointToControlRig(URigVMNode* InBreakpointNod
 				if (!AddedCallpaths.Contains(BreakpointCallPath))
 				{
 					AddedCallpaths.Add(BreakpointCallPath);
-					CDO->AddBreakpoint(i, InBreakpointNode);
+					CDO->AddBreakpoint(i, InBreakpointNode, BreakpointProxy.GetCallstack().Num());
 				}
 			}
 		}
diff --git a/Engine/Plugins/Experimental/ControlRig/Source/ControlRigEditor/Private/Editor/ControlRigEditor.cpp b/Engine/Plugins/Experimental/ControlRig/Source/ControlRigEditor/Private/Editor/ControlRigEditor.cpp
--- a/Engine/Plugins/Experimental/ControlRig/Source/ControlRigEditor/Private/Editor/ControlRigEditor.cpp
+++ b/Engine/Plugins/Experimental/ControlRig/Source/ControlRigEditor/Private/Editor/ControlRigEditor.cpp
@@ -5191,7 +5191,7 @@ bool FControlRigEditor::IsHaltedAtBreakpoint() const
 {
 	if (ControlRig && ControlRig->GetVM())
 	{
-		return ControlRig->GetVM()->GetHaltedAtInstruction() != INDEX_NONE;
+		return ControlRig->GetVM()->GetHaltedAtBreakpoint() != nullptr;
 	}
 	return false;
 }
diff --git a/Engine/Source/Runtime/RigVM/Private/RigVMCore/RigVM.cpp b/Engine/Source/Runtime/RigVM/Private/RigVMCore/RigVM.cpp
--- a/Engine/Source/Runtime/RigVM/Private/RigVMCore/RigVM.cpp
+++ b/Engine/Source/Runtime/RigVM/Private/RigVMCore/RigVM.cpp
@@ -69,8 +69,8 @@ URigVM::URigVM()
     , ByteCodePtr(&ByteCodeStorage)
 #if WITH_EDITOR
 	, DebugInfo(nullptr)
-	, HaltedAtInstruction(INDEX_NONE)
-	, HaltedAtInstructionHit(INDEX_NONE)
+	, HaltedAtBreakpoint(nullptr)
+	, HaltedAtBreakpointHit(INDEX_NONE)
 #endif
     , FunctionNamesPtr(&FunctionNamesStorage)
     , FunctionsPtr(&FunctionsStorage)
@@ -362,13 +362,14 @@ TArray<FName> URigVM::GetEntryNames() const
 #if WITH_EDITOR
 bool URigVM::ResumeExecution()
 {
-	HaltedAtInstruction = INDEX_NONE;
-	HaltedAtInstructionHit = INDEX_NONE;
+	HaltedAtBreakpoint = nullptr;
+	HaltedAtBreakpointHit = INDEX_NONE;
 	if (DebugInfo)
 	{
-		if (FRigVMBreakpoint* Breakpoint = DebugInfo->FindBreakpoint(Context.InstructionIndex))
+		if (const TSharedPtr<FRigVMBreakpoint> CurrentBreakpoint = DebugInfo->GetCurrentActiveBreakpoint())
 		{
-			DebugInfo->IncrementBreakpointActivationOnHit(Context.InstructionIndex);
+			DebugInfo->IncrementBreakpointActivationOnHit(CurrentBreakpoint);
+			DebugInfo->SetCurrentActiveBreakpoint(nullptr);
 			return true;
 		}
 	}
@@ -724,21 +725,35 @@ void URigVM::RebuildByteCodeOnLoad()
 bool URigVM::ShouldHaltAtInstruction(const uint16 InstructionIndex)
 {
 	FRigVMByteCode& ByteCode = GetByteCode();
-	
-	if (FRigVMBreakpoint* Breakpoint = DebugInfo->FindBreakpoint(Context.InstructionIndex))
+
+	if (HaltedAtBreakpoint && InstructionIndex < HaltedAtBreakpoint->InstructionIndex)
 	{
-		if (DebugInfo->IsActive(Context.InstructionIndex))
+		return false;
+	}
+
+	TArray<TSharedPtr<FRigVMBreakpoint>> BreakpointsAtInstruction = DebugInfo->FindBreakpointsAtInstruction(InstructionIndex);
+	for (TSharedPtr<FRigVMBreakpoint> Breakpoint : BreakpointsAtInstruction)
+	{
+		if (DebugInfo->IsActive(Breakpoint))
 		{
 			switch (CurrentBreakpointAction)
 			{
 				case ERigVMBreakpointAction::None:
 				{
 					// Halted at breakpoint. Check if this is a new breakpoint different from the previous halt.
-					if (HaltedAtInstruction != Context.InstructionIndex ||
-						HaltedAtInstructionHit != DebugInfo->GetBreakpointHits(Context.InstructionIndex))
+					if (HaltedAtBreakpoint != Breakpoint ||
+						HaltedAtBreakpointHit != DebugInfo->GetBreakpointHits(Breakpoint))
 					{
-						HaltedAtInstruction = Context.InstructionIndex;
-						HaltedAtInstructionHit = DebugInfo->GetBreakpointHits(HaltedAtInstruction);
+						HaltedAtBreakpoint = Breakpoint;
+						HaltedAtBreakpointHit = DebugInfo->GetBreakpointHits(Breakpoint);
+						DebugInfo->SetCurrentActiveBreakpoint(Breakpoint);
+						
+						// We want to keep the callstack up to the node that produced the halt
+						const TArray<UObject*>* FullCallstack = ByteCode.GetCallstackForInstruction(Context.InstructionIndex);
+						if (FullCallstack)
+						{
+							DebugInfo->SetCurrentActiveBreakpointCallstack(TArray<UObject*>(FullCallstack->GetData(), FullCallstack->Find((UObject*)Breakpoint->Subject)+1));
+						}
 						ExecutionHalted().Broadcast(Context.InstructionIndex, Breakpoint->Subject);
 					}
 					return true;
@@ -749,14 +764,13 @@ bool URigVM::ShouldHaltAtInstruction(const uint16 InstructionIndex)
 
 					if (DebugInfo->IsTemporaryBreakpoint(Breakpoint))
 					{
-						DebugInfo->RemoveBreakpoint(Breakpoint->InstructionIndex);
+						DebugInfo->RemoveBreakpoint(Breakpoint);
 					}
 					else
 					{
-						DebugInfo->IncrementBreakpointActivationOnHit(Context.InstructionIndex);
-						DebugInfo->HitBreakpoint(Context.InstructionIndex);
+						DebugInfo->IncrementBreakpointActivationOnHit(Breakpoint);
+						DebugInfo->HitBreakpoint(Breakpoint);
 					}
-					return false;
 					break;
 				}
 				case ERigVMBreakpointAction::StepOver:
@@ -764,13 +778,16 @@ bool URigVM::ShouldHaltAtInstruction(const uint16 InstructionIndex)
 				case ERigVMBreakpointAction::StepOut:
 				{
 					// If we are stepping, check if we were halted at the current instruction, and remember it 
-					if (!DebugInfo->GetSteppingOriginBreakpoint())
+					if (!DebugInfo->GetCurrentActiveBreakpoint())
 					{
-						DebugInfo->SetSteppingOriginBreakpoint(Breakpoint);
+						DebugInfo->SetCurrentActiveBreakpoint(Breakpoint);
 						const TArray<UObject*>* FullCallstack = ByteCode.GetCallstackForInstruction(Context.InstructionIndex);
 						
 						// We want to keep the callstack up to the node that produced the halt
-						DebugInfo->SetSteppingOriginBreakpointCallstack(TArray<UObject*>(FullCallstack->GetData(), FullCallstack->Find((UObject*)DebugInfo->GetSteppingOriginBreakpoint()->Subject)+1));
+						if (FullCallstack)
+						{
+							DebugInfo->SetCurrentActiveBreakpointCallstack(TArray<UObject*>(FullCallstack->GetData(), FullCallstack->Find((UObject*)DebugInfo->GetCurrentActiveBreakpoint()->Subject)+1));
+						}
 					}							
 					
 					break;	
@@ -784,12 +801,12 @@ bool URigVM::ShouldHaltAtInstruction(const uint16 InstructionIndex)
 		}
 		else
 		{
-			DebugInfo->HitBreakpoint(Context.InstructionIndex);
+			DebugInfo->HitBreakpoint(Breakpoint);
 		}
 	}
 
 	// If we are stepping, and the last active breakpoint was set, check if this is the new temporary breakpoint
-	if (DebugInfo->GetSteppingOriginBreakpoint())
+	if (CurrentBreakpointAction != ERigVMBreakpointAction::None && DebugInfo->GetCurrentActiveBreakpoint())
 	{
 		const TArray<UObject*>* CurrentCallstack = ByteCode.GetCallstackForInstruction(Context.InstructionIndex);
 		if (CurrentCallstack && !CurrentCallstack->IsEmpty())
@@ -798,7 +815,7 @@ bool URigVM::ShouldHaltAtInstruction(const uint16 InstructionIndex)
 
 			// Find the first difference in the callstack
 			int32 DifferenceIndex = INDEX_NONE;
-			TArray<UObject*>& PreviousCallstack = DebugInfo->GetSteppingOriginBreakpointCallstack();
+			TArray<UObject*>& PreviousCallstack = DebugInfo->GetCurrentActiveBreakpointCallstack();
 			for (int32 i=0; i<PreviousCallstack.Num(); ++i)
 			{
 				if (CurrentCallstack->Num() == i)
@@ -844,22 +861,23 @@ bool URigVM::ShouldHaltAtInstruction(const uint16 InstructionIndex)
 			
 			if (NewBreakpointNode)
 			{
-				if (DebugInfo->IsTemporaryBreakpoint(DebugInfo->GetSteppingOriginBreakpoint()))
+				// Remove or hit previous breakpoint
+				if (DebugInfo->IsTemporaryBreakpoint(DebugInfo->GetCurrentActiveBreakpoint()))
 				{
-					DebugInfo->RemoveBreakpoint(DebugInfo->GetSteppingOriginBreakpoint()->InstructionIndex);
+					DebugInfo->RemoveBreakpoint(DebugInfo->GetCurrentActiveBreakpoint());
 				}
 				else
 				{
-					DebugInfo->IncrementBreakpointActivationOnHit(DebugInfo->GetSteppingOriginBreakpoint()->InstructionIndex);
-					DebugInfo->HitBreakpoint(DebugInfo->GetSteppingOriginBreakpoint()->InstructionIndex);
+					DebugInfo->IncrementBreakpointActivationOnHit(DebugInfo->GetCurrentActiveBreakpoint());
+					DebugInfo->HitBreakpoint(DebugInfo->GetCurrentActiveBreakpoint());
 				}
-				
-				FRigVMBreakpoint* NewBreakpoint = DebugInfo->AddBreakpoint(Context.InstructionIndex, NewBreakpointNode, true);
+
+				// Create new temporary breakpoint
+				TSharedPtr<FRigVMBreakpoint> NewBreakpoint = DebugInfo->AddBreakpoint(Context.InstructionIndex, NewBreakpointNode, 0, true);
 				CurrentBreakpointAction = ERigVMBreakpointAction::None;					
 
-				// Halted at breakpoint. Check if this is a new breakpoint different from the previous halt.
-				HaltedAtInstruction = Context.InstructionIndex;
-				HaltedAtInstructionHit = DebugInfo->GetBreakpointHits(HaltedAtInstruction);
+				HaltedAtBreakpoint = NewBreakpoint;
+				HaltedAtBreakpointHit = DebugInfo->GetBreakpointHits(HaltedAtBreakpoint);
 				ExecutionHalted().Broadcast(Context.InstructionIndex, NewBreakpointNode);
 		
 				return true;
@@ -1581,9 +1599,10 @@ bool URigVM::Execute(FRigVMMemoryContainerPtrArray Memory, FRigVMFixedArray<void
 			{
 				ExecutionReachedExit().Broadcast();
 #if WITH_EDITOR					
-				if (HaltedAtInstruction != INDEX_NONE)
+				if (HaltedAtBreakpoint != nullptr)
 				{
-					HaltedAtInstruction = INDEX_NONE;
+					HaltedAtBreakpoint = nullptr;
+					DebugInfo->SetCurrentActiveBreakpoint(nullptr);
 					ExecutionHalted().Broadcast(INDEX_NONE, nullptr);
 				}
 #endif
@@ -1612,9 +1631,10 @@ bool URigVM::Execute(FRigVMMemoryContainerPtrArray Memory, FRigVMFixedArray<void
 	}
 
 #if WITH_EDITOR
-	if (HaltedAtInstruction != INDEX_NONE)
+	if (HaltedAtBreakpoint != nullptr)
 	{
-		HaltedAtInstruction = INDEX_NONE;
+		DebugInfo->SetCurrentActiveBreakpoint(nullptr);
+		HaltedAtBreakpoint = nullptr;
 		ExecutionHalted().Broadcast(INDEX_NONE, nullptr);
 	}
 #endif
diff --git a/Engine/Source/Runtime/RigVM/Private/RigVMCore/RigVMDebugInfo.cpp b/Engine/Source/Runtime/RigVM/Private/RigVMCore/RigVMDebugInfo.cpp
--- a/Engine/Source/Runtime/RigVM/Private/RigVMCore/RigVMDebugInfo.cpp
+++ b/Engine/Source/Runtime/RigVM/Private/RigVMCore/RigVMDebugInfo.cpp
@@ -6,78 +6,121 @@ void FRigVMDebugInfo::ResetState()
 {
 	BreakpointHits.Empty();
 	BreakpointActivationOnHit.Empty();
-	TemporaryBreakpoint.Clear();
+	TemporaryBreakpoint = nullptr;
 }
 
 void FRigVMDebugInfo::StartExecution()
 {
 	BreakpointHits.Empty();
-	SteppingOriginBreakpoint = nullptr;
-	SteppingOriginBreakpointCallstack.Empty();
+	CurrentActiveBreakpoint = nullptr;
+	CurrentActiveBreakpointCallstack.Empty();
 }
 
-FRigVMBreakpoint* FRigVMDebugInfo::FindBreakpoint(const uint16 InstructionIndex) 
+TSharedPtr<FRigVMBreakpoint> FRigVMDebugInfo::FindBreakpoint(const uint16 InInstructionIndex, const UObject* InSubject)
 {
-	if (FRigVMBreakpoint* Found = Breakpoints.Find(InstructionIndex))
+	TArray<TSharedPtr<FRigVMBreakpoint>> BreakpointAtInstruction = FindBreakpointsAtInstruction(InInstructionIndex);
+	for (TSharedPtr<FRigVMBreakpoint> BP : BreakpointAtInstruction)
 	{
-		return Found;
+		if (BP->Subject == InSubject)
+		{
+			return BP;
+		}
 	}
-	if (TemporaryBreakpoint.InstructionIndex == InstructionIndex)
+	return nullptr;
+}
+
+TArray<TSharedPtr<FRigVMBreakpoint>> FRigVMDebugInfo::FindBreakpointsAtInstruction(const uint16 InInstructionIndex) 
+{
+	TArray<TSharedPtr<FRigVMBreakpoint>> Result;
+	for (TSharedPtr<FRigVMBreakpoint> BP : Breakpoints)
 	{
-		return &TemporaryBreakpoint;
+		if (BP->InstructionIndex == InInstructionIndex)
+		{
+			Result.Add(BP);
+		}
 	}
-	return nullptr;
+
+	if (TemporaryBreakpoint.IsValid() && TemporaryBreakpoint->InstructionIndex == InInstructionIndex)
+	{
+		Result.Add(TemporaryBreakpoint);
+	}
+	return Result;
 }
 
-FRigVMBreakpoint* FRigVMDebugInfo::AddBreakpoint(const uint16 InstructionIndex, UObject* InNode, const bool bIsTemporary)
+TSharedPtr<FRigVMBreakpoint> FRigVMDebugInfo::AddBreakpoint(const uint16 InstructionIndex, UObject* InNode, const uint16 InDepth,
+	const bool bIsTemporary)
 {
-	if (!Breakpoints.Contains(InstructionIndex))
+	for (const TSharedPtr<FRigVMBreakpoint> BP : FindBreakpointsAtInstruction(InstructionIndex))
 	{
-		if (bIsTemporary)
+		if (BP->Subject == InNode)
 		{
-			TemporaryBreakpoint = FRigVMBreakpoint(InstructionIndex, InNode);
-				
-			// Do not override the state if it already exists
-			BreakpointActivationOnHit.FindOrAdd(InstructionIndex, 0);
-			BreakpointHits.FindOrAdd(InstructionIndex, 0);
-			return &TemporaryBreakpoint;
+			return nullptr;
 		}
-		else if (FRigVMBreakpoint* NewBP = &Breakpoints.Add(InstructionIndex, FRigVMBreakpoint(InstructionIndex, InNode)))
+	}
+	
+	
+	if (bIsTemporary)
+	{
+		TemporaryBreakpoint = MakeShared<FRigVMBreakpoint>(InstructionIndex, InNode, 0);
+			
+		// Do not override the state if it already exists
+		BreakpointActivationOnHit.FindOrAdd(TemporaryBreakpoint, 0);
+		BreakpointHits.FindOrAdd(TemporaryBreakpoint, 0);
+		return TemporaryBreakpoint;
+	}
+	else
+	{
+		// Breakpoints are sorted by instruction index and callstack depth
+		int32 Index = 0;
+		for(; Index<Breakpoints.Num(); ++Index)
 		{
-			// Do not override the state if it already exists
-			BreakpointActivationOnHit.FindOrAdd(InstructionIndex, 0);
-			BreakpointHits.FindOrAdd(InstructionIndex, 0);
-			return NewBP;
+			if (InstructionIndex < Breakpoints[Index]->InstructionIndex ||
+				(InstructionIndex == Breakpoints[Index]->InstructionIndex && InDepth < Breakpoints[Index]->Depth))
+			{
+				break;
+			}
 		}
+		TSharedRef<FRigVMBreakpoint> NewBP = MakeShared<FRigVMBreakpoint>(InstructionIndex, InNode, InDepth);
+		Breakpoints.Insert(NewBP, Index);
+		
+		// Do not override the state if it already exists
+		BreakpointActivationOnHit.FindOrAdd(NewBP, 0);
+		BreakpointHits.FindOrAdd(NewBP, 0);
+		return NewBP;
 	}
+	
 	return nullptr;
 }
 
-bool FRigVMDebugInfo::RemoveBreakpoint(const uint16 InstructionIndex)
+bool FRigVMDebugInfo::RemoveBreakpoint(const TSharedPtr<FRigVMBreakpoint> InBreakpoint)
 {
-	if (FRigVMBreakpoint* Found = Breakpoints.Find(InstructionIndex))
+	TSharedRef<FRigVMBreakpoint> BreakpointRef = InBreakpoint.ToSharedRef();
+
+	bool Found = false;
+	int32 Index = Breakpoints.Find(BreakpointRef);
+	if (Index != INDEX_NONE)
 	{
-		Breakpoints.Remove(InstructionIndex);
-		BreakpointHits.Remove(InstructionIndex);
-		BreakpointActivationOnHit.Remove(InstructionIndex);
-		return true;
+		BreakpointHits.Remove(BreakpointRef);
+		BreakpointActivationOnHit.Remove(BreakpointRef);
+		Breakpoints.RemoveAt(Index);
+		Found = true;
 	}
-	if (TemporaryBreakpoint.InstructionIndex == InstructionIndex)
+	if (TemporaryBreakpoint.IsValid() && InBreakpoint == TemporaryBreakpoint)
 	{
-		TemporaryBreakpoint.Clear();
-		BreakpointHits.Remove(InstructionIndex);
-		BreakpointActivationOnHit.Remove(InstructionIndex);
-		return true;
+		BreakpointHits.Remove(TemporaryBreakpoint);
+		BreakpointActivationOnHit.Remove(TemporaryBreakpoint);
+		TemporaryBreakpoint = nullptr;
+		Found = true;
 	}
-	return false;
+	return Found;
 }
 
-bool FRigVMDebugInfo::IsActive(const uint16 InstructionIndex) const
+bool FRigVMDebugInfo::IsActive(const TSharedPtr<FRigVMBreakpoint> InBreakpoint) const
 {
-	const FRigVMBreakpoint* Breakpoint = Breakpoints.Find(InstructionIndex);
-	if (!Breakpoint && TemporaryBreakpoint.InstructionIndex == InstructionIndex)
+	TSharedPtr<FRigVMBreakpoint> Breakpoint = InBreakpoint;
+	if (!Breakpoint && TemporaryBreakpoint.IsValid() && TemporaryBreakpoint == InBreakpoint)
 	{
-		Breakpoint = &TemporaryBreakpoint;
+		Breakpoint = TemporaryBreakpoint;
 	}
 	if (Breakpoint)
 	{
@@ -85,13 +128,13 @@ bool FRigVMDebugInfo::IsActive(const uint16 InstructionIndex) const
 		{
 			uint16 Hits = 0;
 			uint16 OnHit = 0;
-			if (BreakpointHits.Contains(InstructionIndex))
+			if (BreakpointHits.Contains(Breakpoint))
 			{
-				Hits = BreakpointHits.FindChecked(InstructionIndex);
+				Hits = BreakpointHits.FindChecked(Breakpoint);
 			}
-			if (BreakpointActivationOnHit.Contains(InstructionIndex))
+			if (BreakpointActivationOnHit.Contains(Breakpoint))
 			{
-				OnHit = BreakpointActivationOnHit.FindChecked(InstructionIndex);
+				OnHit = BreakpointActivationOnHit.FindChecked(Breakpoint);
 			}
 			return Hits == OnHit;
 		}				
@@ -99,35 +142,35 @@ bool FRigVMDebugInfo::IsActive(const uint16 InstructionIndex) const
 	return false;		
 }
 
-void FRigVMDebugInfo::HitBreakpoint(const uint16 InstructionIndex)
+void FRigVMDebugInfo::HitBreakpoint(const TSharedPtr<FRigVMBreakpoint> InBreakpoint)
 {
-	if (BreakpointHits.Contains(InstructionIndex))
+	if (BreakpointHits.Contains(InBreakpoint))
 	{
-		BreakpointHits[InstructionIndex]++;
+		BreakpointHits[InBreakpoint]++;
 	}
 	else
 	{
-		BreakpointHits.Add(InstructionIndex, 1);
+		BreakpointHits.Add(InBreakpoint, 1);
 	}
 }
 
-void FRigVMDebugInfo::IncrementBreakpointActivationOnHit(const uint16 InstructionIndex)
+void FRigVMDebugInfo::IncrementBreakpointActivationOnHit(const TSharedPtr<FRigVMBreakpoint> InBreakpoint)
 {
-	if (BreakpointActivationOnHit.Contains(InstructionIndex))
+	if (BreakpointActivationOnHit.Contains(InBreakpoint))
 	{
-		BreakpointActivationOnHit[InstructionIndex]++;
+		BreakpointActivationOnHit[InBreakpoint]++;
 	}
 	else
 	{
-		BreakpointActivationOnHit.Add(InstructionIndex, 1);
+		BreakpointActivationOnHit.Add(InBreakpoint, 1);
 	}
 }
 
-uint16 FRigVMDebugInfo::GetBreakpointHits(const uint16 InstructionIndex) const
+uint16 FRigVMDebugInfo::GetBreakpointHits(const TSharedPtr<FRigVMBreakpoint> InBreakpoint) const
 {
-	if (BreakpointActivationOnHit.Contains(InstructionIndex))
+	if (BreakpointActivationOnHit.Contains(InBreakpoint))
 	{
-		return BreakpointActivationOnHit[InstructionIndex];
+		return BreakpointActivationOnHit[InBreakpoint];
 	}
 	return 0;
 }
diff --git a/Engine/Source/Runtime/RigVM/Public/RigVMCore/RigVM.h b/Engine/Source/Runtime/RigVM/Public/RigVMCore/RigVM.h
--- a/Engine/Source/Runtime/RigVM/Public/RigVMCore/RigVM.h
+++ b/Engine/Source/Runtime/RigVM/Public/RigVMCore/RigVM.h
@@ -541,7 +541,7 @@ public:
 
 	void SetDebugInfo(FRigVMDebugInfo* InDebugInfo) { DebugInfo = InDebugInfo; }
 
-	int32 GetHaltedAtInstruction() const { return HaltedAtInstruction; }
+	TSharedPtr<FRigVMBreakpoint> GetHaltedAtBreakpoint() const { return HaltedAtBreakpoint; }
 
 	void SetBreakpointAction(const ERigVMBreakpointAction& Action) { CurrentBreakpointAction = Action; }
 #endif
@@ -562,8 +562,8 @@ private:
 
 #if WITH_EDITOR
 	FRigVMDebugInfo* DebugInfo;
-	int32 HaltedAtInstruction;
-	int32 HaltedAtInstructionHit;
+	TSharedPtr<FRigVMBreakpoint> HaltedAtBreakpoint;
+	int32 HaltedAtBreakpointHit;
 	ERigVMBreakpointAction CurrentBreakpointAction;
 
 	bool ShouldHaltAtInstruction(const uint16 InstructionIndex);
diff --git a/Engine/Source/Runtime/RigVM/Public/RigVMCore/RigVMDebugInfo.h b/Engine/Source/Runtime/RigVM/Public/RigVMCore/RigVMDebugInfo.h
--- a/Engine/Source/Runtime/RigVM/Public/RigVMCore/RigVMDebugInfo.h
+++ b/Engine/Source/Runtime/RigVM/Public/RigVMCore/RigVMDebugInfo.h
@@ -16,14 +16,16 @@ struct RIGVM_API FRigVMBreakpoint
 	: bIsActive(true)
 	, InstructionIndex(INDEX_NONE)
 	, Subject(nullptr)
+	, Depth(0)
 	{
 
 	}
 	
-	FRigVMBreakpoint(const uint16 InInstructionIndex, UObject* InNode)
+	FRigVMBreakpoint(const uint16 InInstructionIndex, UObject* InNode, const uint16 InDepth)
 	: bIsActive(true)
 	, InstructionIndex(InInstructionIndex)
 	, Subject(InNode)
+	, Depth(InDepth)
 	{
 
 	}
@@ -33,11 +35,18 @@ struct RIGVM_API FRigVMBreakpoint
 		bIsActive = true;
 		InstructionIndex = INDEX_NONE;
 		Subject = nullptr;
+		Depth = 0;
 	}
 
 	bool bIsActive;
 	uint16 InstructionIndex;
 	UObject* Subject;
+	uint16 Depth;
+
+	bool operator==(const FRigVMBreakpoint& Other) const
+	{
+		return InstructionIndex == Other.InstructionIndex && Subject == Other.Subject;
+	}
 };
 
 USTRUCT()
@@ -46,7 +55,8 @@ struct RIGVM_API FRigVMDebugInfo
 	GENERATED_BODY()
 
 	FORCEINLINE FRigVMDebugInfo()
-		: SteppingOriginBreakpoint(nullptr)
+		: TemporaryBreakpoint(nullptr)
+	, CurrentActiveBreakpoint(nullptr)		
 	{
 	}
 
@@ -60,46 +70,51 @@ struct RIGVM_API FRigVMDebugInfo
 		// Do not remove state
 	}
 
-	FRigVMBreakpoint* FindBreakpoint(const uint16 InstructionIndex);
+	TSharedPtr<FRigVMBreakpoint> FindBreakpoint(const uint16 InInstructionIndex, const UObject* InSubject);
+	TArray<TSharedPtr<FRigVMBreakpoint>> FindBreakpointsAtInstruction(const uint16 InInstructionIndex);
 	
-	FRigVMBreakpoint* AddBreakpoint(const uint16 InstructionIndex, UObject* InNode, const bool bIsTemporary = false);
+	TSharedPtr<FRigVMBreakpoint> AddBreakpoint(const uint16 InstructionIndex, UObject* InNode, const uint16 InDepth, const bool bIsTemporary = false);
 
-	bool RemoveBreakpoint(const uint16 InstructionIndex);
+	bool RemoveBreakpoint(const TSharedPtr<FRigVMBreakpoint> Breakpoint);
 
-	TMap<uint16 , FRigVMBreakpoint> GetBreakpoints() const { return Breakpoints; }
+	TArray<TSharedRef<FRigVMBreakpoint>> GetBreakpoints() const { return Breakpoints; }
 
-	void SetBreakpoints(const TMap<uint16 , FRigVMBreakpoint>& InBreakpoints)
+	void SetBreakpoints(const TArray<TSharedRef<FRigVMBreakpoint>>& InBreakpoints)
 	{
 		Breakpoints = InBreakpoints;
 	}
 
-	bool IsTemporaryBreakpoint(FRigVMBreakpoint* Breakpoint) const
+	bool IsTemporaryBreakpoint(TSharedPtr<FRigVMBreakpoint> Breakpoint) const
 	{
-		return Breakpoint == &TemporaryBreakpoint;
+		if (Breakpoint.IsValid())
+		{
+			return Breakpoint == TemporaryBreakpoint;
+		}
+		return false;
 	}
 
-	bool IsActive(const uint16 InstructionIndex) const;
+	bool IsActive(const TSharedPtr<FRigVMBreakpoint> InBreakpoint) const;
 
-	void HitBreakpoint(const uint16 InstructionIndex);
+	void HitBreakpoint(const TSharedPtr<FRigVMBreakpoint> InBreakpoint);
 
-	void IncrementBreakpointActivationOnHit(const uint16 InstructionIndex);
+	void IncrementBreakpointActivationOnHit(const TSharedPtr<FRigVMBreakpoint> InBreakpoint);
 
-	uint16 GetBreakpointHits(const uint16 InstructionIndex) const;
+	uint16 GetBreakpointHits(const TSharedPtr<FRigVMBreakpoint> InBreakpoint) const;
 
-	FRigVMBreakpoint* GetSteppingOriginBreakpoint() const { return SteppingOriginBreakpoint; }
+	TSharedPtr<FRigVMBreakpoint> GetCurrentActiveBreakpoint() const { return CurrentActiveBreakpoint; }
 
-	void SetSteppingOriginBreakpoint(FRigVMBreakpoint* Breakpoint) { SteppingOriginBreakpoint = Breakpoint; }
+	void SetCurrentActiveBreakpoint(TSharedPtr<FRigVMBreakpoint> Breakpoint) { CurrentActiveBreakpoint = Breakpoint; }
 
-	TArray<UObject*>& GetSteppingOriginBreakpointCallstack() { return SteppingOriginBreakpointCallstack; }
+	TArray<UObject*>& GetCurrentActiveBreakpointCallstack() { return CurrentActiveBreakpointCallstack; }
 
-	void SetSteppingOriginBreakpointCallstack(TArray<UObject*> Callstack) { SteppingOriginBreakpointCallstack = Callstack; }
+	void SetCurrentActiveBreakpointCallstack(TArray<UObject*> Callstack) { CurrentActiveBreakpointCallstack = Callstack; }
 
 private:
-	TMap<uint16 , FRigVMBreakpoint> Breakpoints;
-	FRigVMBreakpoint TemporaryBreakpoint;
-	TMap<uint16 , uint16> BreakpointActivationOnHit; // After how many instruction executions, this breakpoint becomes active
-	TMap<uint16 , uint16> BreakpointHits; // How many times this instruction has been executed
+	TArray<TSharedRef<FRigVMBreakpoint>> Breakpoints;
+	TSharedPtr<FRigVMBreakpoint> TemporaryBreakpoint;
+	TMap<TSharedPtr<FRigVMBreakpoint> , uint16> BreakpointActivationOnHit; // After how many instruction executions, this breakpoint becomes active
+	TMap<TSharedPtr<FRigVMBreakpoint> , uint16> BreakpointHits; // How many times this instruction has been executed
 
-	FRigVMBreakpoint* SteppingOriginBreakpoint;
-	TArray<UObject*> SteppingOriginBreakpointCallstack;
+	TSharedPtr<FRigVMBreakpoint> CurrentActiveBreakpoint;
+	TArray<UObject*> CurrentActiveBreakpointCallstack;
 };
