ControlRig: Fix compile issue in RigVMDebugInfo

diff --git a/Engine/Source/Runtime/RigVM/Private/RigVMCore/RigVMDebugInfo.cpp b/Engine/Source/Runtime/RigVM/Private/RigVMCore/RigVMDebugInfo.cpp
--- a/Engine/Source/Runtime/RigVM/Private/RigVMCore/RigVMDebugInfo.cpp
+++ b/Engine/Source/Runtime/RigVM/Private/RigVMCore/RigVMDebugInfo.cpp
@@ -50,7 +50,7 @@ TArray<TSharedPtr<FRigVMBreakpoint>> FRigVMDebugInfo::FindBreakpointsAtInstructi
 TSharedPtr<FRigVMBreakpoint> FRigVMDebugInfo::AddBreakpoint(const uint16 InstructionIndex, UObject* InNode, const uint16 InDepth,
 	const bool bIsTemporary)
 {
-	for (const TSharedPtr<FRigVMBreakpoint> BP : FindBreakpointsAtInstruction(InstructionIndex))
+	for (const TSharedPtr<FRigVMBreakpoint>& BP : FindBreakpointsAtInstruction(InstructionIndex))
 	{
 		if (BP->Subject == InNode)
 		{
