ControlRig: Fix ordering of enums in ERigVMMemoryType causing failure loading assets created in older version of the engine

diff --git a/Engine/Source/Runtime/RigVM/Private/RigVMCore/RigVM.cpp b/Engine/Source/Runtime/RigVM/Private/RigVMCore/RigVM.cpp
--- a/Engine/Source/Runtime/RigVM/Private/RigVMCore/RigVM.cpp
+++ b/Engine/Source/Runtime/RigVM/Private/RigVMCore/RigVM.cpp
@@ -1167,6 +1167,7 @@ bool URigVM::Execute(FRigVMMemoryContainerPtrArray Memory, FRigVMFixedArray<void
 		return true;
 	}
 
+	// changes to the layout of memory array should be reflected in GetContainerIndex()
 	FRigVMMemoryContainer* LocalMemory[] = { WorkMemoryPtr, LiteralMemoryPtr, DebugMemoryPtr }; 
 	if (Memory.Num() == 0)
 	{
@@ -1993,7 +1994,7 @@ FString URigVM::GetOperandLabel(const FRigVMOperand& InOperand, TFunction<FStrin
 void URigVM::ClearDebugMemory()
 {
 #if WITH_EDITOR
-	FRigVMMemoryContainer* DebugMemory = CachedMemory[(int32)ERigVMMemoryType::Debug];
+	FRigVMMemoryContainer* DebugMemory = CachedMemory[GetContainerIndex(ERigVMMemoryType::Debug)];
 	if (DebugMemory)
 	{ 
 		for (int32 RegisterIndex = 0; RegisterIndex < DebugMemory->Num(); RegisterIndex++)
@@ -2142,6 +2143,7 @@ void URigVM::CopyOperandForDebuggingImpl(const FRigVMOperand& InArg, const FRigV
 		default:
 		{
 			// the default pass for any complex memory
+			// changes to the layout of memory array should be reflected in GetContainerIndex()
 			FRigVMMemoryContainer* LocalMemory[] = { WorkMemoryPtr, LiteralMemoryPtr, DebugMemoryPtr };
 			DebugMemoryPtr->Copy(InArg, InDebugOperand, LocalMemory[InArg.GetContainerIndex()]);
 			break;
@@ -2170,6 +2172,7 @@ TPair<ERigVMRegisterType, uint16> URigVM::GetCopyInfoForOperand(const FRigVMOper
 {
 	if(CachedMemory.IsEmpty())
 	{
+		// changes to the layout of memory array should be reflected in GetContainerIndex()
 		CachedMemory.Add(WorkMemoryPtr);
 		CachedMemory.Add(LiteralMemoryPtr);
 		CachedMemory.Add(DebugMemoryPtr);
diff --git a/Engine/Source/Runtime/RigVM/Public/RigVMCore/RigVM.h b/Engine/Source/Runtime/RigVM/Public/RigVMCore/RigVM.h
--- a/Engine/Source/Runtime/RigVM/Public/RigVMCore/RigVM.h
+++ b/Engine/Source/Runtime/RigVM/Public/RigVMCore/RigVM.h
@@ -588,9 +588,25 @@ private:
 
 	TArray<uint32> FirstHandleForInstruction;
 	TArray<FRigVMMemoryHandle> CachedMemoryHandles;
+	// changes to the layout of cached memory array should be reflected in GetContainerIndex()
 	TArray<FRigVMMemoryContainer*> CachedMemory;
 	TArray<FRigVMExternalVariable> ExternalVariables;
 
+	// this function should be kept in sync with FRigVMOperand::GetContainerIndex()
+	static int32 GetContainerIndex(ERigVMMemoryType InType)
+	{
+		if(InType == ERigVMMemoryType::External)
+		{
+			return (int32)ERigVMMemoryType::Work;
+		}
+		
+		if(InType == ERigVMMemoryType::Debug)
+		{
+			return 2;
+		}
+		return (int32)InType;
+	}
+	
 #if WITH_EDITOR
 	// stores the number of times each instruction was visited
 	TArray<int32> InstructionVisitedDuringLastRun;
diff --git a/Engine/Source/Runtime/RigVM/Public/RigVMCore/RigVMMemory.h b/Engine/Source/Runtime/RigVM/Public/RigVMCore/RigVMMemory.h
--- a/Engine/Source/Runtime/RigVM/Public/RigVMCore/RigVMMemory.h
+++ b/Engine/Source/Runtime/RigVM/Public/RigVMCore/RigVMMemory.h
@@ -23,6 +23,8 @@
 	#define UE_LOG_RIGVMMEMORY(Format, ...)
 #endif
 
+// The ERigVMMemoryType maps to memory container index in RigVM through
+// FRigVMOperand::GetContainerIndex() or URigVM::GetContainerIndex(...) 
 /**
  * The type of memory used. Typically we differentiate between
  * Work (Mutable) and Literal (Constant) memory.
@@ -30,10 +32,10 @@
 UENUM()
 enum class ERigVMMemoryType: uint8
 {
-	Work, // Mutable state
-	Literal, // Const / fixed state
-	Debug, // Owned memory used for debug watches
-	External, // Unowned external memory
+	Work = 0, // Mutable state
+	Literal = 1, // Const / fixed state
+	External = 2, // Unowned external memory
+	Debug = 3, // Owned memory used for debug watches
 	Invalid
 };
 
@@ -87,11 +89,17 @@ public:
 	FORCEINLINE_DEBUGGABLE ERigVMMemoryType GetMemoryType() const { return MemoryType; }
 
 	// returns the index of the container of this argument
+	// this function should be kept in sync with URigVM::GetContainerIndex()
 	FORCEINLINE_DEBUGGABLE int32 GetContainerIndex() const
 	{
 		if(MemoryType == ERigVMMemoryType::External)
 		{
-			return (int)ERigVMMemoryType::Work;
+			return (int32)ERigVMMemoryType::Work;
+		}
+		
+		if(MemoryType == ERigVMMemoryType::Debug)
+		{
+			return 2;
 		}
 		return (int32)MemoryType;
 	}
