Control Rig: RigVM: Validation of operands during load

diff --git a/Engine/Source/Runtime/RigVM/Private/RigVMCore/RigVM.cpp b/Engine/Source/Runtime/RigVM/Private/RigVMCore/RigVM.cpp
--- a/Engine/Source/Runtime/RigVM/Private/RigVMCore/RigVM.cpp
+++ b/Engine/Source/Runtime/RigVM/Private/RigVMCore/RigVM.cpp
@@ -149,7 +149,8 @@ void URigVM::Load(FArchive& Ar)
 	}
 
 	if (WorkMemoryStorage.bEncounteredErrorDuringLoad ||
-		LiteralMemoryStorage.bEncounteredErrorDuringLoad)
+		LiteralMemoryStorage.bEncounteredErrorDuringLoad ||
+		!ValidateAllOperandsDuringLoad())
 	{
 		Reset();
 	}
@@ -165,13 +166,185 @@ void URigVM::Load(FArchive& Ar)
 		}
 
 		// rebuild the bytecode to adjust for byte shifts in shipping
-		// @fixme: commented as workaround to loading issues, fix this asap
-		// RebuildByteCodeOnLoad();
+		RebuildByteCodeOnLoad();
 
 		InvalidateCachedMemory();
 	}
 }
 
+bool URigVM::ValidateAllOperandsDuringLoad()
+{
+	// check all operands on all ops for validity
+	bool bAllOperandsValid = true;
+
+	FRigVMMemoryContainer* LocalMemory[] = { &WorkMemoryStorage, &LiteralMemoryStorage, &DebugMemoryStorage };
+	auto CheckOperandValidity = [LocalMemory, &bAllOperandsValid](const FRigVMOperand& InOperand) -> bool
+	{
+		if(InOperand.GetContainerIndex() < 0 || InOperand.GetContainerIndex() >= (int32)ERigVMMemoryType::Invalid)
+		{
+			bAllOperandsValid = false;
+			return false;
+		}
+
+
+		const FRigVMMemoryContainer* MemoryForOperand = LocalMemory[InOperand.GetContainerIndex()];
+
+		if(InOperand.GetMemoryType() != ERigVMMemoryType::External)
+		{
+			if(!MemoryForOperand->Registers.IsValidIndex(InOperand.GetRegisterIndex()))
+			{
+				bAllOperandsValid = false;
+				return false;
+			}
+		}
+
+		if(InOperand.GetRegisterOffset() != INDEX_NONE)
+		{
+			if(!MemoryForOperand->RegisterOffsets.IsValidIndex(InOperand.GetRegisterOffset()))
+			{
+				bAllOperandsValid = false;
+				return false;
+			}
+		}
+
+		return true;
+	};
+	
+	const FRigVMInstructionArray ByteCodeInstructions = ByteCodeStorage.GetInstructions();
+	for(const FRigVMInstruction& ByteCodeInstruction : ByteCodeInstructions)
+	{
+		switch (ByteCodeInstruction.OpCode)
+		{
+			case ERigVMOpCode::Execute_0_Operands:
+			case ERigVMOpCode::Execute_1_Operands:
+			case ERigVMOpCode::Execute_2_Operands:
+			case ERigVMOpCode::Execute_3_Operands:
+			case ERigVMOpCode::Execute_4_Operands:
+			case ERigVMOpCode::Execute_5_Operands:
+			case ERigVMOpCode::Execute_6_Operands:
+			case ERigVMOpCode::Execute_7_Operands:
+			case ERigVMOpCode::Execute_8_Operands:
+			case ERigVMOpCode::Execute_9_Operands:
+			case ERigVMOpCode::Execute_10_Operands:
+			case ERigVMOpCode::Execute_11_Operands:
+			case ERigVMOpCode::Execute_12_Operands:
+			case ERigVMOpCode::Execute_13_Operands:
+			case ERigVMOpCode::Execute_14_Operands:
+			case ERigVMOpCode::Execute_15_Operands:
+			case ERigVMOpCode::Execute_16_Operands:
+			case ERigVMOpCode::Execute_17_Operands:
+			case ERigVMOpCode::Execute_18_Operands:
+			case ERigVMOpCode::Execute_19_Operands:
+			case ERigVMOpCode::Execute_20_Operands:
+			case ERigVMOpCode::Execute_21_Operands:
+			case ERigVMOpCode::Execute_22_Operands:
+			case ERigVMOpCode::Execute_23_Operands:
+			case ERigVMOpCode::Execute_24_Operands:
+			case ERigVMOpCode::Execute_25_Operands:
+			case ERigVMOpCode::Execute_26_Operands:
+			case ERigVMOpCode::Execute_27_Operands:
+			case ERigVMOpCode::Execute_28_Operands:
+			case ERigVMOpCode::Execute_29_Operands:
+			case ERigVMOpCode::Execute_30_Operands:
+			case ERigVMOpCode::Execute_31_Operands:
+			case ERigVMOpCode::Execute_32_Operands:
+			case ERigVMOpCode::Execute_33_Operands:
+			case ERigVMOpCode::Execute_34_Operands:
+			case ERigVMOpCode::Execute_35_Operands:
+			case ERigVMOpCode::Execute_36_Operands:
+			case ERigVMOpCode::Execute_37_Operands:
+			case ERigVMOpCode::Execute_38_Operands:
+			case ERigVMOpCode::Execute_39_Operands:
+			case ERigVMOpCode::Execute_40_Operands:
+			case ERigVMOpCode::Execute_41_Operands:
+			case ERigVMOpCode::Execute_42_Operands:
+			case ERigVMOpCode::Execute_43_Operands:
+			case ERigVMOpCode::Execute_44_Operands:
+			case ERigVMOpCode::Execute_45_Operands:
+			case ERigVMOpCode::Execute_46_Operands:
+			case ERigVMOpCode::Execute_47_Operands:
+			case ERigVMOpCode::Execute_48_Operands:
+			case ERigVMOpCode::Execute_49_Operands:
+			case ERigVMOpCode::Execute_50_Operands:
+			case ERigVMOpCode::Execute_51_Operands:
+			case ERigVMOpCode::Execute_52_Operands:
+			case ERigVMOpCode::Execute_53_Operands:
+			case ERigVMOpCode::Execute_54_Operands:
+			case ERigVMOpCode::Execute_55_Operands:
+			case ERigVMOpCode::Execute_56_Operands:
+			case ERigVMOpCode::Execute_57_Operands:
+			case ERigVMOpCode::Execute_58_Operands:
+			case ERigVMOpCode::Execute_59_Operands:
+			case ERigVMOpCode::Execute_60_Operands:
+			case ERigVMOpCode::Execute_61_Operands:
+			case ERigVMOpCode::Execute_62_Operands:
+			case ERigVMOpCode::Execute_63_Operands:
+			case ERigVMOpCode::Execute_64_Operands:
+			{
+				const FRigVMExecuteOp& Op = ByteCodeStorage.GetOpAt<FRigVMExecuteOp>(ByteCodeInstruction);
+				FRigVMOperandArray Operands = ByteCodeStorage.GetOperandsForExecuteOp(ByteCodeInstruction);
+				for (const FRigVMOperand& Arg : Operands)
+				{
+					CheckOperandValidity(Arg);
+				}
+				break;
+			}
+			case ERigVMOpCode::Zero:
+			case ERigVMOpCode::BoolFalse:
+			case ERigVMOpCode::BoolTrue:
+			case ERigVMOpCode::Increment:
+			case ERigVMOpCode::Decrement:
+			{
+				const FRigVMUnaryOp& Op = ByteCodeStorage.GetOpAt<FRigVMUnaryOp>(ByteCodeInstruction);
+				CheckOperandValidity(Op.Arg);
+				break;
+			}
+			case ERigVMOpCode::Copy:
+			{
+				const FRigVMCopyOp& Op = ByteCodeStorage.GetOpAt<FRigVMCopyOp>(ByteCodeInstruction);
+				CheckOperandValidity(Op.Source);
+				CheckOperandValidity(Op.Target);
+				break;
+			}
+			case ERigVMOpCode::Equals:
+			case ERigVMOpCode::NotEquals:
+			{
+				const FRigVMComparisonOp& Op = ByteCodeStorage.GetOpAt<FRigVMComparisonOp>(ByteCodeInstruction);
+				CheckOperandValidity(Op.A);
+				CheckOperandValidity(Op.B);
+				CheckOperandValidity(Op.Result);
+				break;
+			}
+			case ERigVMOpCode::JumpAbsoluteIf:
+			case ERigVMOpCode::JumpForwardIf:
+			case ERigVMOpCode::JumpBackwardIf:
+			{
+				const FRigVMJumpIfOp& Op = ByteCodeStorage.GetOpAt<FRigVMJumpIfOp>(ByteCodeInstruction);
+				CheckOperandValidity(Op.Arg);
+				break;
+			}
+			case ERigVMOpCode::BeginBlock:
+			{
+				const FRigVMBinaryOp& Op = ByteCodeStorage.GetOpAt<FRigVMBinaryOp>(ByteCodeInstruction);
+				CheckOperandValidity(Op.ArgA);
+				CheckOperandValidity(Op.ArgB);
+				break;
+			}
+			case ERigVMOpCode::Invalid:
+			{
+				ensure(false);
+				break;
+			}
+			default:
+			{
+				break;
+			}
+		}
+	}
+
+	return bAllOperandsValid;
+}
+
 void URigVM::Reset()
 {
 	WorkMemoryStorage.Reset();
@@ -2160,11 +2333,13 @@ FRigVMCopyOp URigVM::GetCopyOpForOperands(const FRigVMOperand& InSource, const F
 	TPair<ERigVMRegisterType, uint16> SourceCopyInfo = GetCopyInfoForOperand(InSource);
 	TPair<ERigVMRegisterType, uint16> TargetCopyInfo = GetCopyInfoForOperand(InTarget);
 
+#if !WITH_EDITOR
 	check(SourceCopyInfo.Key != ERigVMRegisterType::Invalid);
 	check(SourceCopyInfo.Value > 0);
 	check(TargetCopyInfo.Key != ERigVMRegisterType::Invalid);
 	check(TargetCopyInfo.Value > 0);
 	//check(SourceCopyInfo.Value == TargetCopyInfo.Value);
+#endif
 
 	return FRigVMCopyOp(InSource, InTarget, TargetCopyInfo.Value, TargetCopyInfo.Key);
 }
@@ -2190,23 +2365,25 @@ TPair<ERigVMRegisterType, uint16> URigVM::GetCopyInfoForOperand(const FRigVMOper
 	}
 	else if (InOperand.GetMemoryType() == ERigVMMemoryType::External)
 	{
-		ensure(ExternalVariables.IsValidIndex(InOperand.GetRegisterIndex()));
-		const FRigVMExternalVariable& ExternalVariable = ExternalVariables[InOperand.GetRegisterIndex()];
-
-		NumBytesToCopy = ExternalVariable.Size;
-		RegisterType = ERigVMRegisterType::Plain;
-		
-		if (UScriptStruct* ExternalScriptStruct = Cast<UScriptStruct>(ExternalVariable.TypeObject))	
+		if(ExternalVariables.IsValidIndex(InOperand.GetRegisterIndex()))
 		{
-			RegisterType = ERigVMRegisterType::Struct;
-		}
-		else if (ExternalVariable.TypeName == TEXT("FString"))
-		{
-			RegisterType = ERigVMRegisterType::String;
-		}
-		else if (ExternalVariable.TypeName == TEXT("FName"))
-		{
-			RegisterType = ERigVMRegisterType::Name;
+			const FRigVMExternalVariable& ExternalVariable = ExternalVariables[InOperand.GetRegisterIndex()];
+
+			NumBytesToCopy = ExternalVariable.Size;
+			RegisterType = ERigVMRegisterType::Plain;
+			
+			if (UScriptStruct* ExternalScriptStruct = Cast<UScriptStruct>(ExternalVariable.TypeObject))	
+			{
+				RegisterType = ERigVMRegisterType::Struct;
+			}
+			else if (ExternalVariable.TypeName == TEXT("FString"))
+			{
+				RegisterType = ERigVMRegisterType::String;
+			}
+			else if (ExternalVariable.TypeName == TEXT("FName"))
+			{
+				RegisterType = ERigVMRegisterType::Name;
+			}
 		}
 	}
 	else
diff --git a/Engine/Source/Runtime/RigVM/Public/RigVMCore/RigVM.h b/Engine/Source/Runtime/RigVM/Public/RigVMCore/RigVM.h
--- a/Engine/Source/Runtime/RigVM/Public/RigVMCore/RigVM.h
+++ b/Engine/Source/Runtime/RigVM/Public/RigVMCore/RigVM.h
@@ -65,7 +65,7 @@ public:
 		P.Serialize(Ar);
 		return Ar;
 	}
-
+	
 	// returns true if the parameter is valid
 	bool IsValid() const { return Type != ERigVMParameterType::Invalid; }
 
@@ -147,7 +147,7 @@ public:
 	void Serialize(FArchive& Ar);
 	void Save(FArchive& Ar);
 	void Load(FArchive& Ar);
-
+	
 	// resets the container and maintains all memory
 	void Reset();
 
@@ -637,6 +637,8 @@ private:
 #endif
 	}
 
+	bool ValidateAllOperandsDuringLoad();
+
 	void CopyOperandForDebuggingImpl(const FRigVMOperand& InArg, const FRigVMMemoryHandle& InHandle, const FRigVMOperand& InDebugOperand);
 
 	FRigVMCopyOp GetCopyOpForOperands(const FRigVMOperand& InSource, const FRigVMOperand& InTarget);
diff --git a/Engine/Source/Runtime/RigVM/Public/RigVMCore/RigVMByteCode.h b/Engine/Source/Runtime/RigVM/Public/RigVMCore/RigVMByteCode.h
--- a/Engine/Source/Runtime/RigVM/Public/RigVMCore/RigVMByteCode.h
+++ b/Engine/Source/Runtime/RigVM/Public/RigVMCore/RigVMByteCode.h
@@ -462,6 +462,9 @@ public:
 	// const accessor for an instruction given its index
 	FORCEINLINE const FRigVMInstruction& operator[](int32 InIndex) const { return Instructions[InIndex]; }
 
+	FORCEINLINE_DEBUGGABLE TArray<FRigVMInstruction>::RangedForConstIteratorType begin() const { return Instructions.begin(); }
+	FORCEINLINE_DEBUGGABLE TArray<FRigVMInstruction>::RangedForConstIteratorType end() const { return Instructions.end(); }
+
 private:
 
 	// hide utility constructor
