Control Rig: Fix watching pins from get variable node of external variables not showing

diff --git a/Engine/Source/Runtime/RigVM/Private/RigVMCore/RigVM.cpp b/Engine/Source/Runtime/RigVM/Private/RigVMCore/RigVM.cpp
--- a/Engine/Source/Runtime/RigVM/Private/RigVMCore/RigVM.cpp
+++ b/Engine/Source/Runtime/RigVM/Private/RigVMCore/RigVM.cpp
@@ -1602,6 +1602,13 @@ bool URigVM::Execute(FRigVMMemoryContainerPtrArray Memory, FRigVMFixedArray<void
 					}
 				}
 
+#if WITH_EDITOR
+				if(DebugMemoryPtr->Num() > 0)
+				{
+					CopyOperandForDebuggingIfNeeded(Op.Source, SourceHandle);
+				}
+#endif
+
 				Context.InstructionIndex++;
 				break;
 			}
