Control Rig / RigVM: Support for LWC

diff --git a/Engine/Plugins/Animation/ControlRig/Source/ControlRig/Private/ControlRig.cpp b/Engine/Plugins/Animation/ControlRig/Source/ControlRig/Private/ControlRig.cpp
--- a/Engine/Plugins/Animation/ControlRig/Source/ControlRig/Private/ControlRig.cpp
+++ b/Engine/Plugins/Animation/ControlRig/Source/ControlRig/Private/ControlRig.cpp
@@ -2345,6 +2345,10 @@ FEdGraphPinType UControlRig::GetPinTypeFromExternalVariable(const FRigVMExternal
 	{
 		PinType.PinCategory = UEdGraphSchema_K2::PC_Float;
 	}
+	else if (InExternalVariable.TypeName == TEXT("double"))
+	{
+		PinType.PinCategory = UEdGraphSchema_K2::PC_Double;
+	}
 	else if (InExternalVariable.TypeName == TEXT("FName"))
 	{
 		PinType.PinCategory = UEdGraphSchema_K2::PC_Name;
@@ -2425,6 +2429,11 @@ FRigVMExternalVariable UControlRig::GetExternalVariableFromPinType(const FName&
 		ExternalVariable.TypeName = TEXT("float");
 		ExternalVariable.Size = sizeof(float);
 	}
+	else if (InPinType.PinCategory == UEdGraphSchema_K2::PC_Double)
+	{
+		ExternalVariable.TypeName = TEXT("double");
+		ExternalVariable.Size = sizeof(double);
+	}
 	else if (InPinType.PinCategory == UEdGraphSchema_K2::PC_Name)
 	{
 		ExternalVariable.TypeName = TEXT("FName");
diff --git a/Engine/Plugins/Animation/ControlRig/Source/ControlRig/Private/Units/Math/RigUnit_MathVectorTests.cpp b/Engine/Plugins/Animation/ControlRig/Source/ControlRig/Private/Units/Math/RigUnit_MathVectorTests.cpp
--- a/Engine/Plugins/Animation/ControlRig/Source/ControlRig/Private/Units/Math/RigUnit_MathVectorTests.cpp
+++ b/Engine/Plugins/Animation/ControlRig/Source/ControlRig/Private/Units/Math/RigUnit_MathVectorTests.cpp
@@ -248,6 +248,8 @@ IMPLEMENT_RIGUNIT_AUTOMATION_TEST(FRigUnit_MathVectorSelectBool)
 	return true;
 }
 
+#if UE_LARGE_WORLD_COORDINATES_DISABLED
+
 IMPLEMENT_RIGUNIT_AUTOMATION_TEST(FRigUnit_MathVectorDeg)
 {
 	Unit.Value = FVector(-PI, PI * 2.f, PI * 0.5f);
@@ -264,6 +266,8 @@ IMPLEMENT_RIGUNIT_AUTOMATION_TEST(FRigUnit_MathVectorRad)
 	return true;
 }
 
+#endif
+
 IMPLEMENT_RIGUNIT_AUTOMATION_TEST(FRigUnit_MathVectorLengthSquared)
 {
 	Unit.Value = FVector(2.f, 3.f, 4.f);
diff --git a/Engine/Plugins/Animation/ControlRig/Source/ControlRig/Public/Rigs/RigHierarchyDefines.h b/Engine/Plugins/Animation/ControlRig/Source/ControlRig/Public/Rigs/RigHierarchyDefines.h
--- a/Engine/Plugins/Animation/ControlRig/Source/ControlRig/Public/Rigs/RigHierarchyDefines.h
+++ b/Engine/Plugins/Animation/ControlRig/Source/ControlRig/Public/Rigs/RigHierarchyDefines.h
@@ -167,6 +167,54 @@ public:
 	UPROPERTY()
 	float Float33;
 
+	UPROPERTY()
+	float Float00_2;
+
+	UPROPERTY()
+	float Float01_2;
+
+	UPROPERTY()
+	float Float02_2;
+
+	UPROPERTY()
+	float Float03_2;
+
+	UPROPERTY()
+	float Float10_2;
+
+	UPROPERTY()
+	float Float11_2;
+
+	UPROPERTY()
+	float Float12_2;
+
+	UPROPERTY()
+	float Float13_2;
+
+	UPROPERTY()
+	float Float20_2;
+
+	UPROPERTY()
+	float Float21_2;
+
+	UPROPERTY()
+	float Float22_2;
+
+	UPROPERTY()
+	float Float23_2;
+	
+	UPROPERTY()
+	float Float30_2;
+
+	UPROPERTY()
+	float Float31_2;
+
+	UPROPERTY()
+	float Float32_2;
+
+	UPROPERTY()
+	float Float33_2;
+	
 	UPROPERTY()
 	bool bValid;
 };
diff --git a/Engine/Plugins/Animation/ControlRig/Source/ControlRigDeveloper/Private/ControlRigBlueprint.cpp b/Engine/Plugins/Animation/ControlRig/Source/ControlRigDeveloper/Private/ControlRigBlueprint.cpp
--- a/Engine/Plugins/Animation/ControlRig/Source/ControlRigDeveloper/Private/ControlRigBlueprint.cpp
+++ b/Engine/Plugins/Animation/ControlRig/Source/ControlRigDeveloper/Private/ControlRigBlueprint.cpp
@@ -1863,6 +1863,11 @@ FName UControlRigBlueprint::AddMemberVariable(const FName& InName, const FString
 		Variable.TypeName = *CPPType;
 		Variable.Size = sizeof(float);
 	}
+	else if (CPPType == TEXT("double"))
+	{
+		Variable.TypeName = *CPPType;
+		Variable.Size = sizeof(double);
+	}
 	else if (CPPType == TEXT("int32"))
 	{
 		Variable.TypeName = *CPPType;
diff --git a/Engine/Plugins/Animation/ControlRig/Source/ControlRigDeveloper/Private/Graph/ControlRigGraphNode.cpp b/Engine/Plugins/Animation/ControlRig/Source/ControlRigDeveloper/Private/Graph/ControlRigGraphNode.cpp
--- a/Engine/Plugins/Animation/ControlRig/Source/ControlRigDeveloper/Private/Graph/ControlRigGraphNode.cpp
+++ b/Engine/Plugins/Animation/ControlRig/Source/ControlRigDeveloper/Private/Graph/ControlRigGraphNode.cpp
@@ -1027,6 +1027,10 @@ FEdGraphPinType UControlRigGraphNode::GetPinTypeForModelPin(URigVMPin* InModelPi
 	{
 		PinType.PinCategory = UEdGraphSchema_K2::PC_Float;
 	}
+	else if (ModelPinCPPType == TEXT("double"))
+	{
+		PinType.PinCategory = UEdGraphSchema_K2::PC_Double;
+	}
 	else if (ModelPinCPPType == TEXT("FName"))
 	{
 		PinType.PinCategory = UEdGraphSchema_K2::PC_Name;
diff --git a/Engine/Plugins/Animation/ControlRig/Source/ControlRigDeveloper/Private/Graph/ControlRigGraphSchema.cpp b/Engine/Plugins/Animation/ControlRig/Source/ControlRigDeveloper/Private/Graph/ControlRigGraphSchema.cpp
--- a/Engine/Plugins/Animation/ControlRig/Source/ControlRigDeveloper/Private/Graph/ControlRigGraphSchema.cpp
+++ b/Engine/Plugins/Animation/ControlRig/Source/ControlRigDeveloper/Private/Graph/ControlRigGraphSchema.cpp
@@ -515,6 +515,14 @@ FLinearColor UControlRigGraphSchema::GetPinTypeColor(const FEdGraphPinType& PinT
 			}
 		}
 	}
+
+	if(TypeName == UEdGraphSchema_K2::PC_Double)
+	{
+		FEdGraphPinType LocalPinType = PinType;
+		LocalPinType.PinCategory = UEdGraphSchema_K2::PC_Float;
+		return GetPinTypeColor(LocalPinType);
+	}
+	
 	return GetDefault<UEdGraphSchema_K2>()->GetPinTypeColor(PinType);
 }
 
@@ -538,6 +546,7 @@ bool UControlRigGraphSchema::SupportsPinType(UScriptStruct* ScriptStruct) const
 		FString CPPType = Property->GetCPPType();
 		if (CPPType == TEXT("bool") ||
 			CPPType == TEXT("float") ||
+			CPPType == TEXT("double") ||
 			CPPType == TEXT("int32") ||
 			CPPType == TEXT("FString") ||
 			CPPType == TEXT("FName") ||
@@ -583,6 +592,7 @@ bool UControlRigGraphSchema::SupportsPinType(TWeakPtr<const FEdGraphSchemaAction
 	if (TypeName == UEdGraphSchema_K2::PC_Boolean ||
 		TypeName == UEdGraphSchema_K2::PC_Int ||
 		TypeName == UEdGraphSchema_K2::PC_Float ||
+		TypeName == UEdGraphSchema_K2::PC_Double ||
 		TypeName == UEdGraphSchema_K2::PC_Name ||
 		TypeName == UEdGraphSchema_K2::PC_String ||
 		TypeName == UEdGraphSchema_K2::PC_Enum)
@@ -1237,6 +1247,7 @@ void UControlRigGraphSchema::GetVariablePinTypes(TArray<FEdGraphPinType>& PinTyp
 {
 	PinTypes.Add(FEdGraphPinType(UEdGraphSchema_K2::PC_Boolean, FName(NAME_None), nullptr, EPinContainerType::None, false, FEdGraphTerminalType()));
 	PinTypes.Add(FEdGraphPinType(UEdGraphSchema_K2::PC_Float, FName(NAME_None), nullptr, EPinContainerType::None, false, FEdGraphTerminalType()));
+	PinTypes.Add(FEdGraphPinType(UEdGraphSchema_K2::PC_Double, FName(NAME_None), nullptr, EPinContainerType::None, false, FEdGraphTerminalType()));
 	PinTypes.Add(FEdGraphPinType(UEdGraphSchema_K2::PC_Int, FName(NAME_None), nullptr, EPinContainerType::None, false, FEdGraphTerminalType()));
 	PinTypes.Add(FEdGraphPinType(UEdGraphSchema_K2::PC_Struct, FName(NAME_None), TBaseStructure<FVector>::Get(), EPinContainerType::None, false, FEdGraphTerminalType()));
 	PinTypes.Add(FEdGraphPinType(UEdGraphSchema_K2::PC_Struct, FName(NAME_None), TBaseStructure<FVector2D>::Get(), EPinContainerType::None, false, FEdGraphTerminalType()));
diff --git a/Engine/Plugins/Animation/ControlRig/Source/ControlRigEditor/Private/ControlRigEditorModule.cpp b/Engine/Plugins/Animation/ControlRig/Source/ControlRigEditor/Private/ControlRigEditorModule.cpp
--- a/Engine/Plugins/Animation/ControlRig/Source/ControlRigEditor/Private/ControlRigEditorModule.cpp
+++ b/Engine/Plugins/Animation/ControlRig/Source/ControlRigEditor/Private/ControlRigEditorModule.cpp
@@ -1401,6 +1401,7 @@ void FControlRigEditorModule::GetContextMenuActions(const UControlRigGraphSchema
 						}
 
 						if (ModelPin->GetCPPType() == TEXT("float") ||
+							ModelPin->GetCPPType() == TEXT("double") ||
 							ModelPin->GetCPPType() == TEXT("FVector"))
 						{
 							FToolMenuSection& Section = Menu->AddSection("EdGraphSchemaNodeInjectionInterp", LOCTEXT("NodeInjectionInterp", "Interpolate"));
@@ -1422,7 +1423,7 @@ void FControlRigEditorModule::GetContextMenuActions(const UControlRigGraphSchema
 							{
 								UScriptStruct* ScriptStruct = nullptr;
 
-								if (ModelPin->GetCPPType() == TEXT("float"))
+								if ((ModelPin->GetCPPType() == TEXT("float")) || (ModelPin->GetCPPType() == TEXT("double")))
 								{
 									ScriptStruct = FRigUnit_AlphaInterp::StaticStruct();
 								}
diff --git a/Engine/Source/Developer/RigVMDeveloper/Private/RigVMCompiler/RigVMAST.cpp b/Engine/Source/Developer/RigVMDeveloper/Private/RigVMCompiler/RigVMAST.cpp
--- a/Engine/Source/Developer/RigVMDeveloper/Private/RigVMCompiler/RigVMAST.cpp
+++ b/Engine/Source/Developer/RigVMDeveloper/Private/RigVMCompiler/RigVMAST.cpp
@@ -1660,15 +1660,32 @@ void FRigVMParserAST::FoldAssignments()
 		ensure(AssignExpr->Parents.Num() == 1);
 		ensure(AssignExpr->Children.Num() == 1);
 
+		URigVMPin* SourcePin = AssignExpr->GetSourcePin();
+		URigVMPin* TargetPin = AssignExpr->GetTargetPin();
+
+		// in case the assign has different types for left and right - we need to avoid folding
+		// since this assign represents a cast operation
+		if(SourcePin->GetCPPTypeObject() != TargetPin->GetCPPTypeObject())
+		{
+			continue;
+		}
+		else if(SourcePin->GetCPPTypeObject() == nullptr)
+		{
+			if(SourcePin->GetCPPType() != TargetPin->GetCPPType())
+			{
+				continue;
+			}
+		}
+		
 		// non-input pins on anything but a reroute node should be skipped
-		if (AssignExpr->GetTargetPin()->GetDirection() != ERigVMPinDirection::Input &&
-			Cast<URigVMRerouteNode>(AssignExpr->GetTargetPin()->GetNode()) == nullptr)
+		if (TargetPin->GetDirection() != ERigVMPinDirection::Input &&
+			Cast<URigVMRerouteNode>(TargetPin->GetNode()) == nullptr)
 		{
 			continue;
 		}
 
 		// if this node is a loop node - let's skip the folding
-		if (URigVMUnitNode* UnitNode = Cast<URigVMUnitNode>(AssignExpr->GetTargetPin()->GetNode()))
+		if (URigVMUnitNode* UnitNode = Cast<URigVMUnitNode>(TargetPin->GetNode()))
 		{
 			if (UnitNode->IsLoopNode())
 			{
@@ -1677,9 +1694,9 @@ void FRigVMParserAST::FoldAssignments()
 		}
 
 		// if this node is a variable node and the pin requires a watch... skip this
-		if (Cast<URigVMVariableNode>(AssignExpr->GetSourcePin()->GetNode()))
+		if (Cast<URigVMVariableNode>(SourcePin->GetNode()))
 		{
-			if(AssignExpr->GetSourcePin()->RequiresWatch(true))
+			if(SourcePin->RequiresWatch(true))
 			{
 				continue;
 			}
@@ -2165,6 +2182,10 @@ void FRigVMParserAST::FoldLiterals()
 				{
 					DefaultValue = TEXT("0.000000");
 				}
+				else if (LiteralExpr->GetCPPType() == TEXT("double"))
+				{
+					DefaultValue = TEXT("0.000000");
+				}
 				else if (LiteralExpr->GetCPPType() == TEXT("int32"))
 				{
 					DefaultValue = TEXT("0");
diff --git a/Engine/Source/Developer/RigVMDeveloper/Private/RigVMCompiler/RigVMCompiler.cpp b/Engine/Source/Developer/RigVMDeveloper/Private/RigVMCompiler/RigVMCompiler.cpp
--- a/Engine/Source/Developer/RigVMDeveloper/Private/RigVMCompiler/RigVMCompiler.cpp
+++ b/Engine/Source/Developer/RigVMDeveloper/Private/RigVMCompiler/RigVMCompiler.cpp
@@ -1605,6 +1605,10 @@ uint16 URigVMCompiler::GetElementSizeFromCPPType(const FString& InCPPType, UScri
 	{
 		return sizeof(float);
 	}
+	if (InCPPType == TEXT("double"))
+	{
+		return sizeof(double);
+	}
 	if (InCPPType == TEXT("FName"))
 	{
 		return sizeof(FName);
@@ -1880,6 +1884,13 @@ FRigVMOperand URigVMCompiler::FindOrAddRegister(const FRigVMVarExprAST* InVarExp
 							continue;
 						}
 					}
+					else if (Pin->GetCPPType() == TEXT("double"))
+					{
+						if (Memory[ExistingRegisterIndex].ElementSize != sizeof(double))
+						{
+							continue;
+						}
+					}
 					else if (UEnum* Enum = Pin->GetEnum())
 					{
 						if (Memory[ExistingRegisterIndex].ElementSize != sizeof(uint8))
@@ -2054,6 +2065,29 @@ FRigVMOperand URigVMCompiler::FindOrAddRegister(const FRigVMVarExprAST* InVarExp
 				Register = Memory.AddRegisterArray<float>(!Pin->IsDynamicArray() && !bIsDebugValue, RegisterName, Values.Num(), Pin->IsArray(), (uint8*)Values.GetData(), NumSlices, ERigVMRegisterType::Plain, nullptr);
 				check(Register != INDEX_NONE);
 				
+#if WITH_EDITORONLY_DATA
+				Memory.Registers[Register].BaseCPPType = *BaseCPPType;
+				Memory.Registers[Register].BaseCPPTypeObject = nullptr;
+#endif
+				Operand = Memory.GetOperand(Register);
+			}
+			else if (BaseCPPType == TEXT("double"))
+			{
+				TArray<double> Values;
+				for (FString DefaultValue : DefaultValues)
+				{
+					if (DefaultValue.IsEmpty())
+					{
+						Values.Add(0.0);
+					}
+					else
+					{
+						Values.Add(FCString::Atod(*DefaultValue));
+					}
+				}
+				Register = Memory.AddRegisterArray<double>(!Pin->IsDynamicArray() && !bIsDebugValue, RegisterName, Values.Num(), Pin->IsArray(), (uint8*)Values.GetData(), NumSlices, ERigVMRegisterType::Plain, nullptr);
+				check(Register != INDEX_NONE);
+				
 #if WITH_EDITORONLY_DATA
 				Memory.Registers[Register].BaseCPPType = *BaseCPPType;
 				Memory.Registers[Register].BaseCPPTypeObject = nullptr;
diff --git a/Engine/Source/Developer/RigVMDeveloper/Private/RigVMModel/RigVMPin.cpp b/Engine/Source/Developer/RigVMDeveloper/Private/RigVMModel/RigVMPin.cpp
--- a/Engine/Source/Developer/RigVMDeveloper/Private/RigVMModel/RigVMPin.cpp
+++ b/Engine/Source/Developer/RigVMDeveloper/Private/RigVMModel/RigVMPin.cpp
@@ -625,6 +625,13 @@ bool URigVMPin::IsValidDefaultValue(const FString& InDefaultValue) const
 				return false;
 			}
 		}
+		else if (BaseCPPType == TEXT("double"))
+		{ 
+			if (!FDefaultValueHelper::IsStringValidFloat(Value))
+			{
+				return false;
+			}
+		}
 		else if (BaseCPPType == TEXT("int32"))
 		{ 
 			if (!FDefaultValueHelper::IsStringValidInteger(Value))
@@ -1210,21 +1217,38 @@ bool URigVMPin::CanLink(URigVMPin* InSourcePin, URigVMPin* InTargetPin, FString*
 
 	if (InSourcePin->CPPType != InTargetPin->CPPType)
 	{
-		if (OutFailureReason)
+		bool bCPPTypesDiffer = true;
+		static const FString Float = TEXT("float");
+		static const FString Double = TEXT("double");
+
+		if((InSourcePin->CPPType == Float && InTargetPin->CPPType == Double) ||
+			(InSourcePin->CPPType == Double && InTargetPin->CPPType == Float))
 		{
-			*OutFailureReason = TEXT("Source and target pin types are not compatible.");
+			bCPPTypesDiffer = false;
 		}
 
-		// check if this might be a prototype node
-		if (InSourcePin->CPPType.IsEmpty())
+		if (bCPPTypesDiffer)
 		{
-			if (URigVMPrototypeNode* PrototypeNode = Cast<URigVMPrototypeNode>(SourceNode))
+			if (OutFailureReason)
+			{
+				*OutFailureReason = TEXT("Source and target pin types are not compatible.");
+			}
+
+			// check if this might be a prototype node
+			if (InSourcePin->CPPType.IsEmpty())
 			{
-				if (PrototypeNode->SupportsType(InSourcePin, InTargetPin->CPPType))
+				if (URigVMPrototypeNode* PrototypeNode = Cast<URigVMPrototypeNode>(SourceNode))
 				{
-					if (OutFailureReason)
+					if (PrototypeNode->SupportsType(InSourcePin, InTargetPin->CPPType))
+					{
+						if (OutFailureReason)
+						{
+							*OutFailureReason = FString();
+						}
+					}
+					else
 					{
-						*OutFailureReason = FString();
+						return false;
 					}
 				}
 				else
@@ -1232,20 +1256,20 @@ bool URigVMPin::CanLink(URigVMPin* InSourcePin, URigVMPin* InTargetPin, FString*
 					return false;
 				}
 			}
-			else
-			{
-				return false;
-			}
-		}
-		else if (InTargetPin->CPPType.IsEmpty())
-		{
-			if (URigVMPrototypeNode* PrototypeNode = Cast<URigVMPrototypeNode>(TargetNode))
+			else if (InTargetPin->CPPType.IsEmpty())
 			{
-				if (PrototypeNode->SupportsType(InTargetPin, InSourcePin->CPPType))
+				if (URigVMPrototypeNode* PrototypeNode = Cast<URigVMPrototypeNode>(TargetNode))
 				{
-					if (OutFailureReason)
+					if (PrototypeNode->SupportsType(InTargetPin, InSourcePin->CPPType))
+					{
+						if (OutFailureReason)
+						{
+							*OutFailureReason = FString();
+						}
+					}
+					else
 					{
-						*OutFailureReason = FString();
+						return false;
 					}
 				}
 				else
@@ -1258,10 +1282,6 @@ bool URigVMPin::CanLink(URigVMPin* InSourcePin, URigVMPin* InTargetPin, FString*
 				return false;
 			}
 		}
-		else
-		{
-			return false;
-		}
 	}
 
 	if(!SourceNode->AllowsLinksOn(InSourcePin))
diff --git a/Engine/Source/Developer/RigVMDeveloper/Public/RigVMModel/RigVMVariableDescription.h b/Engine/Source/Developer/RigVMDeveloper/Public/RigVMModel/RigVMVariableDescription.h
--- a/Engine/Source/Developer/RigVMDeveloper/Public/RigVMModel/RigVMVariableDescription.h
+++ b/Engine/Source/Developer/RigVMDeveloper/Public/RigVMModel/RigVMVariableDescription.h
@@ -84,6 +84,10 @@ public:
 		{
 			PinType.PinCategory = UEdGraphSchema_K2::PC_Float;
 		}
+		else if (CPPType == TEXT("double"))
+		{
+			PinType.PinCategory = UEdGraphSchema_K2::PC_Double;
+		}
 		else if (CPPType == TEXT("FName"))
 		{
 			PinType.PinCategory = UEdGraphSchema_K2::PC_Name;
@@ -138,6 +142,10 @@ public:
 		{
 			OutCPPType = Prefix + TEXT("float") + Suffix;
 		}
+		else if (InPinType.PinCategory == UEdGraphSchema_K2::PC_Double)
+		{
+			OutCPPType = Prefix + TEXT("double") + Suffix;
+		}
 		else if (InPinType.PinCategory == UEdGraphSchema_K2::PC_Name)
 		{
 			OutCPPType = Prefix + TEXT("FName") + Suffix;
diff --git a/Engine/Source/Runtime/RigVM/Private/RigVMCore/RigVM.cpp b/Engine/Source/Runtime/RigVM/Private/RigVMCore/RigVM.cpp
--- a/Engine/Source/Runtime/RigVM/Private/RigVMCore/RigVM.cpp
+++ b/Engine/Source/Runtime/RigVM/Private/RigVMCore/RigVM.cpp
@@ -1251,7 +1251,39 @@ bool URigVM::Initialize(FRigVMMemoryContainerPtrArray Memory, FRigVMFixedArray<v
 				{
 					case ERigVMRegisterType::Plain:
 					{
-						FMemory::Memcpy(TargetPtr, SourcePtr, NumBytes);
+						switch(Op.CopyType)
+						{
+							case ERigVMCopyType::FloatToDouble:
+							{
+								const float* Floats = (float*)SourcePtr;
+								double* Doubles = (double*)TargetPtr;
+								const int32 NumElementsToCopy = NumBytes / sizeof(double);
+
+								for(int32 ElementIndex = 0; ElementIndex < NumElementsToCopy; ElementIndex++)
+								{
+									Doubles[ElementIndex] = (double)Floats[ElementIndex];
+								}
+								break;
+							}
+							case ERigVMCopyType::DoubleToFloat:
+							{
+								const double* Doubles = (double*)SourcePtr;
+								float* Floats = (float*)TargetPtr;
+								const int32 NumElementsToCopy = NumBytes / sizeof(float); 
+
+								for(int32 ElementIndex = 0; ElementIndex < NumElementsToCopy; ElementIndex++)
+								{
+									Floats[ElementIndex] = (float)Doubles[ElementIndex];
+								}
+								break;
+							}
+							case ERigVMCopyType::Default:
+							default:
+							{
+								FMemory::Memcpy(TargetPtr, SourcePtr, NumBytes);
+								break;
+							}
+						}
 						break;
 					}
 					case ERigVMRegisterType::Name:
@@ -1559,7 +1591,39 @@ bool URigVM::Execute(FRigVMMemoryContainerPtrArray Memory, FRigVMFixedArray<void
 				{
 					case ERigVMRegisterType::Plain:
 					{
-						FMemory::Memcpy(TargetPtr, SourcePtr, NumBytes);
+						switch(Op.CopyType)
+						{
+							case ERigVMCopyType::FloatToDouble:
+							{
+								const float* Floats = (float*)SourcePtr;
+								double* Doubles = (double*)TargetPtr;
+								const int32 NumElementsToCopy = NumBytes / sizeof(double);
+
+								for(int32 ElementIndex = 0; ElementIndex < NumElementsToCopy; ElementIndex++)
+								{
+									Doubles[ElementIndex] = (double)Floats[ElementIndex];
+								}
+								break;
+							}
+							case ERigVMCopyType::DoubleToFloat:
+							{
+								const double* Doubles = (double*)SourcePtr;
+								float* Floats = (float*)TargetPtr;
+								const int32 NumElementsToCopy = NumBytes / sizeof(float); 
+
+								for(int32 ElementIndex = 0; ElementIndex < NumElementsToCopy; ElementIndex++)
+								{
+									Floats[ElementIndex] = (float)Doubles[ElementIndex];
+								}
+								break;
+							}
+							case ERigVMCopyType::Default:
+							default:
+							{
+								FMemory::Memcpy(TargetPtr, SourcePtr, NumBytes);
+								break;
+							}
+						}
 						break;
 					}
 					case ERigVMRegisterType::Name:
@@ -2337,21 +2401,34 @@ void URigVM::CopyOperandForDebuggingImpl(const FRigVMOperand& InArg, const FRigV
 
 FRigVMCopyOp URigVM::GetCopyOpForOperands(const FRigVMOperand& InSource, const FRigVMOperand& InTarget)
 {
-	TPair<ERigVMRegisterType, uint16> SourceCopyInfo = GetCopyInfoForOperand(InSource);
-	TPair<ERigVMRegisterType, uint16> TargetCopyInfo = GetCopyInfoForOperand(InTarget);
+	FCopyInfoForOperand SourceCopyInfo = GetCopyInfoForOperand(InSource);
+	FCopyInfoForOperand TargetCopyInfo = GetCopyInfoForOperand(InTarget);
 
 #if !WITH_EDITOR
 	check(SourceCopyInfo.Key != ERigVMRegisterType::Invalid);
 	check(SourceCopyInfo.Value > 0);
 	check(TargetCopyInfo.Key != ERigVMRegisterType::Invalid);
 	check(TargetCopyInfo.Value > 0);
-	//check(SourceCopyInfo.Value == TargetCopyInfo.Value);
 #endif
 
-	return FRigVMCopyOp(InSource, InTarget, TargetCopyInfo.Value, TargetCopyInfo.Key);
+	ERigVMCopyType CopyType = ERigVMCopyType::Default;
+
+	if(SourceCopyInfo.RegisterType == ERigVMRegisterType::Plain && TargetCopyInfo.RegisterType == ERigVMRegisterType::Plain)
+	{
+		if((SourceCopyInfo.ElementSize == sizeof(float)) && (TargetCopyInfo.ElementSize == sizeof(double)))
+		{
+			CopyType = ERigVMCopyType::FloatToDouble;
+		}
+		else if((SourceCopyInfo.ElementSize == sizeof(double)) && (TargetCopyInfo.ElementSize == sizeof(float)))
+		{
+			CopyType = ERigVMCopyType::DoubleToFloat;
+		}
+	}
+
+	return FRigVMCopyOp(InSource, InTarget, TargetCopyInfo.NumBytesToCopy, TargetCopyInfo.RegisterType, CopyType);
 }
 
-TPair<ERigVMRegisterType, uint16> URigVM::GetCopyInfoForOperand(const FRigVMOperand& InOperand)
+URigVM::FCopyInfoForOperand URigVM::GetCopyInfoForOperand(const FRigVMOperand& InOperand)
 {
 	if(CachedMemory.IsEmpty())
 	{
@@ -2363,12 +2440,14 @@ TPair<ERigVMRegisterType, uint16> URigVM::GetCopyInfoForOperand(const FRigVMOper
 	
 	ERigVMRegisterType RegisterType = ERigVMRegisterType::Invalid;
 	uint16 NumBytesToCopy = 0;
+	uint16 ElementSize = 0;
 
 	if (InOperand.GetRegisterOffset() != INDEX_NONE)
 	{
 		const FRigVMRegisterOffset& RegisterOffset = CachedMemory[InOperand.GetContainerIndex()]->RegisterOffsets[InOperand.GetRegisterOffset()];
 		RegisterType = RegisterOffset.GetType();
 		NumBytesToCopy = RegisterOffset.GetElementSize();
+		ElementSize = RegisterOffset.GetElementSize();
 	}
 	else if (InOperand.GetMemoryType() == ERigVMMemoryType::External)
 	{
@@ -2377,6 +2456,7 @@ TPair<ERigVMRegisterType, uint16> URigVM::GetCopyInfoForOperand(const FRigVMOper
 			const FRigVMExternalVariable& ExternalVariable = ExternalVariables[InOperand.GetRegisterIndex()];
 
 			NumBytesToCopy = ExternalVariable.Size;
+			ElementSize = ExternalVariable.Size;
 			RegisterType = ERigVMRegisterType::Plain;
 			
 			if (UScriptStruct* ExternalScriptStruct = Cast<UScriptStruct>(ExternalVariable.TypeObject))	
@@ -2399,9 +2479,10 @@ TPair<ERigVMRegisterType, uint16> URigVM::GetCopyInfoForOperand(const FRigVMOper
 
 		RegisterType = Register.Type;
 		NumBytesToCopy = Register.GetNumBytesPerSlice();
+		ElementSize = Register.ElementSize;
 	}
 
-	return TPair<ERigVMRegisterType, uint16>(RegisterType, NumBytesToCopy); 
+	return FCopyInfoForOperand(RegisterType, NumBytesToCopy, ElementSize); 
 }
 
 UScriptStruct* URigVM::GetScriptStructForCopyOp(const FRigVMCopyOp& InCopyOp) const
diff --git a/Engine/Source/Runtime/RigVM/Private/RigVMCore/RigVMByteCode.cpp b/Engine/Source/Runtime/RigVM/Private/RigVMCore/RigVMByteCode.cpp
--- a/Engine/Source/Runtime/RigVM/Private/RigVMCore/RigVMByteCode.cpp
+++ b/Engine/Source/Runtime/RigVM/Private/RigVMCore/RigVMByteCode.cpp
@@ -774,19 +774,19 @@ uint64 FRigVMByteCode::AddTrueOp(const FRigVMOperand& InArg)
 	return AddOp(Op);
 }
 
-uint64 FRigVMByteCode::AddCopyOp(const FRigVMOperand& InSource, const FRigVMOperand& InTarget, uint16 InNumBytes, ERigVMRegisterType InTargetType)
+uint64 FRigVMByteCode::AddCopyOp(const FRigVMOperand& InSource, const FRigVMOperand& InTarget, uint16 InNumBytes, ERigVMRegisterType InTargetType, ERigVMCopyType InCopyType)
 {
 	check(InTarget.GetMemoryType() != ERigVMMemoryType::Literal);
 	check(InNumBytes > 0);
 	check(InTargetType != ERigVMRegisterType::Invalid);
 	
-	FRigVMCopyOp Op(InSource, InTarget, InNumBytes, InTargetType);
+	FRigVMCopyOp Op(InSource, InTarget, InNumBytes, InTargetType, InCopyType);
 	return AddOp(Op);
 }
 
 uint64 FRigVMByteCode::AddCopyOp(const FRigVMCopyOp& InCopyOp)
 {
-	return AddCopyOp(InCopyOp.Source, InCopyOp.Target, InCopyOp.NumBytes, InCopyOp.RegisterType);
+	return AddCopyOp(InCopyOp.Source, InCopyOp.Target, InCopyOp.NumBytes, InCopyOp.RegisterType, InCopyOp.CopyType);
 }
 
 uint64 FRigVMByteCode::AddIncrementOp(const FRigVMOperand& InArg)
diff --git a/Engine/Source/Runtime/RigVM/Private/RigVMCore/RigVMMemory.cpp b/Engine/Source/Runtime/RigVM/Private/RigVMCore/RigVMMemory.cpp
--- a/Engine/Source/Runtime/RigVM/Private/RigVMCore/RigVMMemory.cpp
+++ b/Engine/Source/Runtime/RigVM/Private/RigVMCore/RigVMMemory.cpp
@@ -1779,6 +1779,10 @@ void FRigVMMemoryContainer::SetRegisterValueFromString(const FRigVMOperand& InOp
 		{
 			GetFixedArray<float>(Register)[Index] = FCString::Atof(*DefaultValue);
 		}
+		else if (CPPType == TEXT("double") && Register.Type == ERigVMRegisterType::Plain && Register.ElementSize == sizeof(double))
+		{
+			GetFixedArray<double>(Register)[Index] = FCString::Atod(*DefaultValue);
+		}
 		else if (CPPType == TEXT("FName") && Register.Type == ERigVMRegisterType::Name)
 		{
 			GetFixedArray<FName>(Register)[Index] = *DefaultValue;
@@ -1856,6 +1860,11 @@ TArray<FString> FRigVMMemoryContainer::GetRegisterValueAsString(const FRigVMOper
 				float Value = GetFixedArray<float>(Register, INDEX_NONE, SliceIndex)[ElementIndex];
 				DefaultValue = FString::Printf(TEXT("%f"), Value);
 			}
+			else if (InCPPType == TEXT("double") && Register.GetNumBytesPerSlice() == sizeof(double))
+			{
+				double Value = GetFixedArray<double>(Register, INDEX_NONE, SliceIndex)[ElementIndex];
+				DefaultValue = FString::Printf(TEXT("%f"), (float)Value);
+			}
 			else if (InCPPType == TEXT("FName") && Register.GetNumBytesPerSlice() == sizeof(FName))
 			{
 				DefaultValue = GetFixedArray<FName>(Register, INDEX_NONE, SliceIndex)[ElementIndex].ToString();
diff --git a/Engine/Source/Runtime/RigVM/Public/RigVMCore/RigVM.h b/Engine/Source/Runtime/RigVM/Public/RigVMCore/RigVM.h
--- a/Engine/Source/Runtime/RigVM/Public/RigVMCore/RigVM.h
+++ b/Engine/Source/Runtime/RigVM/Public/RigVMCore/RigVM.h
@@ -374,6 +374,12 @@ public:
 		return GetParameterValue<float>(InParameterName, InArrayIndex);
 	}
 
+	UFUNCTION(BlueprintCallable, Category = RigVM)
+	double GetParameterValueDouble(const FName& InParameterName, int32 InArrayIndex = 0)
+	{
+		return GetParameterValue<double>(InParameterName, InArrayIndex);
+	}
+
 	UFUNCTION(BlueprintCallable, Category = RigVM)
 	int32 GetParameterValueInt(const FName& InParameterName, int32 InArrayIndex = 0)
 	{
@@ -401,7 +407,7 @@ public:
 	UFUNCTION(BlueprintCallable, Category = RigVM)
 	FVector GetParameterValueVector(const FName& InParameterName, int32 InArrayIndex = 0)
 	{
-		return GetParameterValue<FVector3f>(InParameterName, InArrayIndex, FVector3f::ZeroVector);	// LWC_TODO: Store double FVector
+		return GetParameterValue<FVector>(InParameterName, InArrayIndex, FVector::ZeroVector);	// LWC_TODO: Store double FVector
 	}
 
 	UFUNCTION(BlueprintCallable, Category = RigVM)
@@ -428,6 +434,12 @@ public:
 		SetParameterValue<float>(InParameterName, InValue, InArrayIndex);
 	}
 
+	UFUNCTION(BlueprintCallable, Category = RigVM)
+	void SetParameterValueDouble(const FName& InParameterName, double InValue, int32 InArrayIndex = 0)
+	{
+		SetParameterValue<double>(InParameterName, InValue, InArrayIndex);
+	}
+
 	UFUNCTION(BlueprintCallable, Category = RigVM)
 	void SetParameterValueInt(const FName& InParameterName, int32 InValue, int32 InArrayIndex = 0)
 	{
@@ -455,7 +467,7 @@ public:
 	UFUNCTION(BlueprintCallable, Category = RigVM)
 	void SetParameterValueVector(const FName& InParameterName, const FVector& InValue, int32 InArrayIndex = 0)
 	{
-		SetParameterValue<FVector3f>(InParameterName, InValue, InArrayIndex);	// LWC_TODO: Store double FVector
+		SetParameterValue<FVector>(InParameterName, InValue, InArrayIndex);	// LWC_TODO: Store double FVector
 	}
 
 	UFUNCTION(BlueprintCallable, Category = RigVM)
@@ -641,8 +653,27 @@ private:
 
 	void CopyOperandForDebuggingImpl(const FRigVMOperand& InArg, const FRigVMMemoryHandle& InHandle, const FRigVMOperand& InDebugOperand);
 
+	struct FCopyInfoForOperand
+	{
+		FCopyInfoForOperand()
+			: RegisterType(ERigVMRegisterType::Invalid)
+			, NumBytesToCopy(0)
+			, ElementSize(0)
+		{}
+
+		FCopyInfoForOperand(ERigVMRegisterType InRegisterType, uint16 InNumBytesToCopy, uint16 InElementSize)
+			: RegisterType(InRegisterType)
+			, NumBytesToCopy(InNumBytesToCopy)
+			, ElementSize(InElementSize)
+		{}
+		
+		ERigVMRegisterType RegisterType;
+		uint16 NumBytesToCopy;
+		uint16 ElementSize;
+	};
+
 	FRigVMCopyOp GetCopyOpForOperands(const FRigVMOperand& InSource, const FRigVMOperand& InTarget);
-	TPair<ERigVMRegisterType, uint16> GetCopyInfoForOperand(const FRigVMOperand& InOperand);
+	FCopyInfoForOperand GetCopyInfoForOperand(const FRigVMOperand& InOperand);
 	UScriptStruct* GetScriptStructForCopyOp(const FRigVMCopyOp& InCopyOp) const;
 	UScriptStruct* GetScripStructForOperand(const FRigVMOperand& InOperand) const;
 	
diff --git a/Engine/Source/Runtime/RigVM/Public/RigVMCore/RigVMByteCode.h b/Engine/Source/Runtime/RigVM/Public/RigVMCore/RigVMByteCode.h
--- a/Engine/Source/Runtime/RigVM/Public/RigVMCore/RigVMByteCode.h
+++ b/Engine/Source/Runtime/RigVM/Public/RigVMCore/RigVMByteCode.h
@@ -218,6 +218,15 @@ struct RIGVM_API FRigVMBinaryOp : public FRigVMBaseOp
 	}
 };
 
+// The kind of copy operation to perform
+UENUM()
+enum class ERigVMCopyType : uint8
+{
+	Default,
+	FloatToDouble,
+	DoubleToFloat
+};
+
 // copy the content of one register to another
 USTRUCT()
 struct RIGVM_API FRigVMCopyOp : public FRigVMBaseOp
@@ -230,6 +239,7 @@ struct RIGVM_API FRigVMCopyOp : public FRigVMBaseOp
 	, Target()
 	, NumBytes(0)
 	, RegisterType(ERigVMRegisterType::Invalid)
+	, CopyType(ERigVMCopyType::Default)
 	{
 	}
 
@@ -237,13 +247,15 @@ struct RIGVM_API FRigVMCopyOp : public FRigVMBaseOp
 		FRigVMOperand InSource,
 		FRigVMOperand InTarget,
 		uint16 InNumBytes,
-		ERigVMRegisterType InRegisterType
+		ERigVMRegisterType InRegisterType,
+		ERigVMCopyType InCopyType
 	)
 		: FRigVMBaseOp(ERigVMOpCode::Copy)
 		, Source(InSource)
 		, Target(InTarget)
 		, NumBytes(InNumBytes)
 		, RegisterType(InRegisterType)
+		, CopyType(InCopyType)
 	{
 	}
 
@@ -251,6 +263,7 @@ struct RIGVM_API FRigVMCopyOp : public FRigVMBaseOp
 	FRigVMOperand Target;
 	uint16 NumBytes;
 	ERigVMRegisterType RegisterType;
+	ERigVMCopyType CopyType;
 
 	void Serialize(FArchive& Ar);
 	FORCEINLINE friend FArchive& operator<<(FArchive& Ar, FRigVMCopyOp& P)
@@ -551,7 +564,7 @@ public:
 	uint64 AddTrueOp(const FRigVMOperand& InArg);
 
 	// adds a copy operator to copy the content of a source argument to a target argument
-	uint64 AddCopyOp(const FRigVMOperand& InSource, const FRigVMOperand& InTarget, uint16 InNumBytes, ERigVMRegisterType InTargetType);
+	uint64 AddCopyOp(const FRigVMOperand& InSource, const FRigVMOperand& InTarget, uint16 InNumBytes, ERigVMRegisterType InTargetType, ERigVMCopyType InCopyType);
 
 	// adds a copy operator to copy the content of a source argument to a target argument
 	uint64 AddCopyOp(const FRigVMCopyOp& InCopyOp);
diff --git a/Engine/Source/Runtime/RigVM/Public/RigVMCore/RigVMExternalVariable.h b/Engine/Source/Runtime/RigVM/Public/RigVMCore/RigVMExternalVariable.h
--- a/Engine/Source/Runtime/RigVM/Public/RigVMCore/RigVMExternalVariable.h
+++ b/Engine/Source/Runtime/RigVM/Public/RigVMCore/RigVMExternalVariable.h
@@ -64,6 +64,10 @@ struct RIGVM_API FRigVMExternalVariable
 		{
 			ExternalVariable.TypeName = TEXT("float");
 		}
+		else if (CastField<FDoubleProperty>(Property))
+		{
+			ExternalVariable.TypeName = TEXT("double");
+		}
 		else if (CastField<FStrProperty>(Property))
 		{
 			ExternalVariable.TypeName = TEXT("FString");
@@ -195,6 +199,30 @@ struct RIGVM_API FRigVMExternalVariable
 		return Variable;
 	}
 
+	FORCEINLINE static FRigVMExternalVariable Make(const FName& InName, double& InValue)
+	{
+		FRigVMExternalVariable Variable;
+		Variable.Name = InName;
+		Variable.TypeName = TEXT("double");
+		Variable.TypeObject = nullptr;
+		Variable.bIsArray = false;
+		Variable.Size = sizeof(double);
+		Variable.Memory = (uint8*)&InValue;
+		return Variable;
+	}
+
+	FORCEINLINE static FRigVMExternalVariable Make(const FName& InName, TArray<double>& InValue)
+	{
+		FRigVMExternalVariable Variable;
+		Variable.Name = InName;
+		Variable.TypeName = TEXT("double");
+		Variable.TypeObject = nullptr;
+		Variable.bIsArray = true;
+		Variable.Size = sizeof(double);
+		Variable.Memory = (uint8*)&InValue;
+		return Variable;
+	}
+
 	FORCEINLINE static FRigVMExternalVariable Make(const FName& InName, FString& InValue)
 	{
 		FRigVMExternalVariable Variable;
