Quick fix for non-editor builds

#rnx

diff --git a/Engine/Source/Runtime/RigVM/Private/RigVMCore/RigVM.cpp b/Engine/Source/Runtime/RigVM/Private/RigVMCore/RigVM.cpp
--- a/Engine/Source/Runtime/RigVM/Private/RigVMCore/RigVM.cpp
+++ b/Engine/Source/Runtime/RigVM/Private/RigVMCore/RigVM.cpp
@@ -2405,10 +2405,10 @@ FRigVMCopyOp URigVM::GetCopyOpForOperands(const FRigVMOperand& InSource, const F
 	FCopyInfoForOperand TargetCopyInfo = GetCopyInfoForOperand(InTarget);
 
 #if !WITH_EDITOR
-	check(SourceCopyInfo.Key != ERigVMRegisterType::Invalid);
-	check(SourceCopyInfo.Value > 0);
-	check(TargetCopyInfo.Key != ERigVMRegisterType::Invalid);
-	check(TargetCopyInfo.Value > 0);
+	check(SourceCopyInfo.RegisterType != ERigVMRegisterType::Invalid);
+	check(SourceCopyInfo.NumBytesToCopy > 0);
+	check(TargetCopyInfo.RegisterType != ERigVMRegisterType::Invalid);
+	check(TargetCopyInfo.NumBytesToCopy > 0);
 #endif
 
 	ERigVMCopyType CopyType = ERigVMCopyType::Default;
