ControlRig: Edigrate CL 16861268, revert back to use non deterministic hash but now we no longer serializes dynamic registers.

diff --git a/Engine/Plugins/Animation/ControlRig/Source/ControlRig/Private/Units/Execution/RigUnit_Collection.cpp b/Engine/Plugins/Animation/ControlRig/Source/ControlRig/Private/Units/Execution/RigUnit_Collection.cpp
--- a/Engine/Plugins/Animation/ControlRig/Source/ControlRig/Private/Units/Execution/RigUnit_Collection.cpp
+++ b/Engine/Plugins/Animation/ControlRig/Source/ControlRig/Private/Units/Execution/RigUnit_Collection.cpp
@@ -59,7 +59,7 @@ FRigUnit_CollectionNameSearch_Execute()
 	}
 
 	int32 CurrentHierarchyHash = Context.Hierarchy->GetTopologyVersion() * 17;
-	CurrentHierarchyHash += GetTypeHash(PartialName.ToString());
+	CurrentHierarchyHash += GetTypeHash(PartialName);
 	CurrentHierarchyHash += (int32)TypeToSearch * 8;
 
 	if (CachedHierarchyHash != CurrentHierarchyHash || CachedCollection.IsEmpty())
@@ -152,8 +152,8 @@ FRigUnit_CollectionReplaceItems_Execute()
 
 	int32 CurrentHierarchyHash = Context.Hierarchy->GetTopologyVersion() * 17;
 	CurrentHierarchyHash += GetTypeHash(Items);
-	CurrentHierarchyHash += 12 * GetTypeHash(Old.ToString());
-	CurrentHierarchyHash += 13 * GetTypeHash(New.ToString());
+	CurrentHierarchyHash += 12 * GetTypeHash(Old);
+	CurrentHierarchyHash += 13 * GetTypeHash(New);
 	CurrentHierarchyHash += RemoveInvalidItems ? 14 : 0;
 
 	if (CachedHierarchyHash != CurrentHierarchyHash || CachedCollection.IsEmpty())
diff --git a/Engine/Plugins/Animation/ControlRig/Source/ControlRig/Public/Rigs/RigHierarchyDefines.h b/Engine/Plugins/Animation/ControlRig/Source/ControlRig/Public/Rigs/RigHierarchyDefines.h
--- a/Engine/Plugins/Animation/ControlRig/Source/ControlRig/Public/Rigs/RigHierarchyDefines.h
+++ b/Engine/Plugins/Animation/ControlRig/Source/ControlRig/Public/Rigs/RigHierarchyDefines.h
@@ -975,7 +975,7 @@ public:
 
 	friend FORCEINLINE uint32 GetTypeHash(const FRigElementKey& Key)
 	{
-		return GetTypeHash(Key.Name.ToString()) * 10 + (uint32)Key.Type;
+		return GetTypeHash(Key.Name) * 10 + (uint32)Key.Type;
 	}
 
 	FORCEINLINE bool operator ==(const FRigElementKey& Other) const
diff --git a/Engine/Source/Runtime/RigVM/Private/RigVMCore/RigVMMemory.cpp b/Engine/Source/Runtime/RigVM/Private/RigVMCore/RigVMMemory.cpp
--- a/Engine/Source/Runtime/RigVM/Private/RigVMCore/RigVMMemory.cpp
+++ b/Engine/Source/Runtime/RigVM/Private/RigVMCore/RigVMMemory.cpp
@@ -606,6 +606,14 @@ void FRigVMMemoryContainer::Save(FArchive& Ar)
 			continue;
 		}
 
+		// literal mem container should not have any dynamic registers,
+		// if we ever added them, we might need to re-consider if we need to serialize them
+		// by default dynamic registers' memory are not serialized
+		if (MemoryType == ERigVMMemoryType::Literal)
+		{
+			check(!Register.IsDynamic());
+		}
+
 		if (!Register.IsDynamic())
 		{
 			switch (Register.Type)
@@ -659,45 +667,36 @@ void FRigVMMemoryContainer::Save(FArchive& Ar)
 		else if (!Register.IsNestedDynamic())
 		{
 			FRigVMByteArray* ArrayStorage = (FRigVMByteArray*)&Data[Register.GetWorkByteIndex()];
-
+			
+			// notice that all dynamic registers' memory is not really serialized
+			// serialize empty arrays here just to avoid having to increment serialization version
+			// it does not need to be saved since it will be recomputed during execution
 			switch (Register.Type)
 			{
-			case ERigVMRegisterType::Plain:
-			{
-				Ar << *ArrayStorage;
-				break;
-			}
-			case ERigVMRegisterType::Name:
-			{
-				TArray<FName> View = FRigVMFixedArray<FName>(*ArrayStorage);
-				Ar << View;
-				break;
-			}
-			case ERigVMRegisterType::String:
-			{
-				TArray<FString> View = FRigVMFixedArray<FString>(*ArrayStorage);
-				Ar << View;
-				break;
-			}
-			case ERigVMRegisterType::Struct:
-			{
-				uint8* DataPtr = ArrayStorage->GetData();
-				uint16 NumElements = (uint16)(ArrayStorage->Num() / Register.ElementSize);
-				ensure(NumElements * Register.ElementSize == ArrayStorage->Num());
-				UScriptStruct* ScriptStruct = GetScriptStruct(Register);
-
-				TArray<FString> View;
-				for (uint16 ElementIndex = 0; ElementIndex < NumElements; ElementIndex++)
+				case ERigVMRegisterType::Plain:
 				{
-					FString Value;
-					ScriptStruct->ExportText(Value, DataPtr, nullptr, nullptr, PPF_None, nullptr);
-					View.Add(Value);
-					DataPtr += Register.ElementSize;
+					FRigVMByteArray DummyStorage;
+					Ar << DummyStorage;
+					break;
+				}
+				case ERigVMRegisterType::Name:
+				{
+					TArray<FName> View;
+					Ar << View;
+					break;
+				}
+				case ERigVMRegisterType::String:
+				{
+					TArray<FString> View;
+					Ar << View;
+					break;
+				}
+				case ERigVMRegisterType::Struct:
+				{
+					TArray<FString> View;
+					Ar << View;
+					break;
 				}
-
-				Ar << View;
-				break;
-			}
 			}
 		}
 		else
@@ -715,41 +714,32 @@ void FRigVMMemoryContainer::Save(FArchive& Ar)
 
 				FRigVMByteArray& ArrayStorage = (*NestedArrayStorage)[SliceIndex];
 
+				// notice that all dynamic registers' memory is not really serialized
+				// serialize empty arrays here just to avoid having to increment serialization version
+				// it does not need to be saved since it will be recomputed during execution
 				switch (Register.Type)
 				{
 					case ERigVMRegisterType::Plain:
 					{
-						Ar << ArrayStorage;
+						FRigVMByteArray DummyStorage;
+						Ar << DummyStorage;
 						break;
 					}
 					case ERigVMRegisterType::Name:
 					{
-						TArray<FName> View = FRigVMFixedArray<FName>(ArrayStorage);
+						TArray<FName> View;
 						Ar << View;
 						break;
 					}
 					case ERigVMRegisterType::String:
 					{
-						TArray<FString> View = FRigVMFixedArray<FString>(ArrayStorage);
+						TArray<FString> View;
 						Ar << View;
 						break;
 					}
 					case ERigVMRegisterType::Struct:
 					{
-						uint8* DataPtr = ArrayStorage.GetData();
-						uint16 NumElements = (uint16)(ArrayStorage.Num() / Register.ElementSize);
-						ensure(NumElements * Register.ElementSize == ArrayStorage.Num());
-						UScriptStruct* ScriptStruct = GetScriptStruct(Register);
-
 						TArray<FString> View;
-						for (uint16 ElementIndex = 0; ElementIndex < NumElements; ElementIndex++)
-						{
-							FString Value;
-							ScriptStruct->ExportText(Value, DataPtr, nullptr, nullptr, PPF_None, nullptr);
-							View.Add(Value);
-							DataPtr += Register.ElementSize;
-						}
-
 						Ar << View;
 						break;
 					}
@@ -878,6 +868,14 @@ void FRigVMMemoryContainer::Load(FArchive& Ar)
 			continue;
 		}
 
+		// literal mem container should not have any dynamic registers,
+		// if we ever added them, we might need to re-consider if we need to serialize them
+		// by default dynamic registers' memory are not serialized
+		if (MemoryType == ERigVMMemoryType::Literal)
+		{
+			check(!Register.IsDynamic());
+		}
+		
 		if (!Register.IsDynamic())
 		{
 			switch (Register.Type)
@@ -954,68 +952,34 @@ void FRigVMMemoryContainer::Load(FArchive& Ar)
 			{
 				ArrayStorage = (FRigVMByteArray*)&Data[Register.GetWorkByteIndex()];
 			}
-
+			
+			// notice that all dynamic registers' memory is not really serialized
+			// serialize empty arrays here just to avoid having to increment serialization version
+			// the loaded data is not getting used at all, it will be recomputed at execution time
 			switch (Register.Type)
 			{
 				case ERigVMRegisterType::Plain:
 				{
 					FRigVMByteArray View;
 					Ar << View;
-					if (ArrayStorage)
-					{
-						*ArrayStorage = View;
-					}
 					break;
 				}
 				case ERigVMRegisterType::Name:
 				{
 					TArray<FName> View;
 					Ar << View;
-
-					if (ArrayStorage)
-					{
-						FRigVMDynamicArray<FName> ValueArray(*ArrayStorage);
-						ValueArray.CopyFrom(View);
-					}
 					break;
 				}
 				case ERigVMRegisterType::String:
 				{
 					TArray<FString> View;
 					Ar << View;
-
-					if (ArrayStorage)
-					{
-						FRigVMDynamicArray<FString> ValueArray(*ArrayStorage);
-						ValueArray.CopyFrom(View);
-					}
 					break;
 				}
 				case ERigVMRegisterType::Struct:
 				{
 					TArray<FString> View;
 					Ar << View;
-
-					UScriptStruct* ScriptStruct = GetScriptStruct(Register);
-					if (ScriptStruct && !bEncounteredErrorDuringLoad && ArrayStorage)
-					{
-						ensure(ScriptStruct->GetStructureSize() == Register.ElementSize);
-
-						ArrayStorage->SetNumZeroed(View.Num() * Register.ElementSize);
-						uint8* DataPtr = ArrayStorage->GetData();
-
-						for (uint16 ElementIndex = 0; ElementIndex < View.Num(); ElementIndex++)
-						{
-							FRigVMMemoryContainerImportErrorContext ErrorPipe;
-							ScriptStruct->ImportText(*View[ElementIndex], DataPtr, nullptr, PPF_None, &ErrorPipe, ScriptStruct->GetName());
-							if (ErrorPipe.NumErrors > 0)
-							{
-								bEncounteredErrorDuringLoad = true;
-								break;
-							}
-							DataPtr += Register.ElementSize;
-						}
-					}
 					break;
 				}
 			}
@@ -1038,67 +1002,33 @@ void FRigVMMemoryContainer::Load(FArchive& Ar)
 					ArrayStorage = &(*NestedArrayStorage)[SliceIndex];
 				}
 
+				// notice that all dynamic registers' memory is not really serialized
+				// serialize empty arrays here just to avoid having to increment serialization version
+				// the loaded data is not getting used at all, it will be recomputed at execution time
 				switch (Register.Type)
 				{
 					case ERigVMRegisterType::Plain:
 					{
 						FRigVMByteArray View;
 						Ar << View;
-						if (ArrayStorage)
-						{
-							*ArrayStorage = View;
-						}
 						break;
 					}
 					case ERigVMRegisterType::Name:
 					{
 						TArray<FName> View;
 						Ar << View;
-
-						if (ArrayStorage)
-						{
-							FRigVMDynamicArray<FName> ValueArray(*ArrayStorage);
-							ValueArray.CopyFrom(View);
-						}
 						break;
 					}
 					case ERigVMRegisterType::String:
 					{
 						TArray<FString> View;
 						Ar << View;
-
-						if (ArrayStorage)
-						{
-							FRigVMDynamicArray<FString> ValueArray(*ArrayStorage);
-							ValueArray.CopyFrom(View);
-						}
 						break;
 					}
 					case ERigVMRegisterType::Struct:
 					{
 						TArray<FString> View;
 						Ar << View;
-
-						UScriptStruct* ScriptStruct = GetScriptStruct(Register);
-						if (ScriptStruct && !bEncounteredErrorDuringLoad && ArrayStorage)
-						{
-							ensure(ScriptStruct->GetStructureSize() == Register.ElementSize);
-
-							ArrayStorage->SetNumZeroed(View.Num() * Register.ElementSize);
-							uint8* DataPtr = ArrayStorage->GetData();
-
-							for (uint16 ElementIndex = 0; ElementIndex < View.Num(); ElementIndex++)
-							{
-								FRigVMMemoryContainerImportErrorContext ErrorPipe;
-								ScriptStruct->ImportText(*View[ElementIndex], DataPtr, nullptr, PPF_None, &ErrorPipe, ScriptStruct->GetName());
-								if (ErrorPipe.NumErrors > 0)
-								{
-									bEncounteredErrorDuringLoad = true;
-									break;
-								}
-								DataPtr += Register.ElementSize;
-							}
-						}
 						break;
 					}
 				}
