RigVM Memory Refactoring - moving old code to deprecated file and introducing global define, defaulting to old behavior

diff --git a/Engine/Source/Runtime/RigVM/Private/RigVMCore/RigVMMemory.cpp b/Engine/Source/Runtime/RigVM/Private/RigVMCore/RigVMMemory.cpp
deleted file mode 100644
diff --git a/Engine/Source/Runtime/RigVM/Private/RigVMCore/RigVMMemoryCommon.cpp b/Engine/Source/Runtime/RigVM/Private/RigVMCore/RigVMMemoryCommon.cpp
new file mode 100644
--- /dev/null
+++ b/Engine/Source/Runtime/RigVM/Private/RigVMCore/RigVMMemoryCommon.cpp
@@ -0,0 +1,16 @@
+// Copyright Epic Games, Inc. All Rights Reserved.
+
+#include "RigVMCore/RigVMMemoryCommon.h"
+
+#if DEBUG_RIGVMMEMORY
+	DEFINE_LOG_CATEGORY(LogRigVMMemory);
+#endif
+
+////////////////////////////////////////////////////////////////////////////////////////////////////////////////////////
+
+void FRigVMOperand::Serialize(FArchive& Ar)
+{
+	Ar << MemoryType;
+	Ar << RegisterIndex;
+	Ar << RegisterOffset;
+}
diff --git a/Engine/Source/Runtime/RigVM/Private/RigVMCore/RigVMMemoryDeprecated.cpp b/Engine/Source/Runtime/RigVM/Private/RigVMCore/RigVMMemoryDeprecated.cpp
new file mode 100644
--- /dev/null
+++ b/Engine/Source/Runtime/RigVM/Private/RigVMCore/RigVMMemoryDeprecated.cpp
@@ -0,0 +1,1944 @@
+// Copyright Epic Games, Inc. All Rights Reserved.
+
+#include "RigVMCore/RigVMMemoryDeprecated.h"
+#include "RigVMCore/RigVMUtilities.h"
+#include "UObject/AnimObjectVersion.h"
+#include "UObject/ReleaseObjectVersion.h"
+#include "UObject/PropertyPortFlags.h"
+#include "UObject/Package.h"
+#include "RigVMModule.h"
+
+////////////////////////////////////////////////////////////////////////////////////////////////////////////////////////
+
+void FRigVMRegister::Serialize(FArchive& Ar)
+{
+	Ar.UsingCustomVersion(FAnimObjectVersion::GUID);
+
+	if (Ar.CustomVer(FAnimObjectVersion::GUID) < FAnimObjectVersion::StoreMarkerNamesOnSkeleton)
+	{
+		return;
+	}
+
+	uint16 SliceIndex = 0;
+
+	Ar << Type;
+	Ar << ByteIndex;
+	Ar << ElementSize;
+	Ar << ElementCount;
+	Ar << SliceIndex;
+	Ar << SliceCount;
+	Ar << AlignmentBytes;
+	Ar << TrailingBytes;
+	Ar << Name;
+	Ar << ScriptStructIndex;
+
+	if (Ar.IsSaving() || Ar.IsObjectReferenceCollector() || Ar.IsCountingMemory())
+	{
+		Save(Ar);
+	}
+	else if (Ar.IsLoading())
+	{
+		Load(Ar);
+	}
+	else
+	{
+		checkNoEntry();
+	}
+}
+
+void FRigVMRegister::Save(FArchive& Ar)
+{
+	Ar << bIsArray;
+	Ar << bIsDynamic;
+}
+
+void FRigVMRegister::Load(FArchive& Ar)
+{
+	if (Ar.CustomVer(FAnimObjectVersion::GUID) >= FAnimObjectVersion::SerializeRigVMRegisterArrayState)
+	{
+		Ar << bIsArray;
+	}
+	else
+	{
+		bIsArray = false;
+	}
+
+	if (Ar.CustomVer(FAnimObjectVersion::GUID) >= FAnimObjectVersion::SerializeRigVMRegisterDynamicState)
+	{
+		Ar << bIsDynamic;
+	}
+	else
+	{
+		bIsDynamic = false;
+	}
+}
+
+////////////////////////////////////////////////////////////////////////////////////////////////////////////////////////
+
+void FRigVMRegisterOffset::Serialize(FArchive& Ar)
+{
+	Ar.UsingCustomVersion(FAnimObjectVersion::GUID);
+	Ar.UsingCustomVersion(FReleaseObjectVersion::GUID);
+
+	if (Ar.CustomVer(FAnimObjectVersion::GUID) < FAnimObjectVersion::StoreMarkerNamesOnSkeleton)
+	{
+		return;
+	}
+
+	if (Ar.IsSaving() || Ar.IsObjectReferenceCollector() || Ar.IsCountingMemory())
+	{
+		Save(Ar);
+	}
+	else if (Ar.IsLoading())
+	{
+		Load(Ar);
+	}
+	else
+	{
+		checkNoEntry();
+	}
+}
+
+void FRigVMRegisterOffset::Save(FArchive& Ar)
+{
+	Ar << Segments;
+	Ar << Type;
+	Ar << CPPType;
+
+	Ar << ScriptStruct;
+	
+	Ar << ElementSize;
+	Ar << ParentScriptStruct;
+	Ar << CachedSegmentPath;
+	Ar << ArrayIndex;
+}
+
+void FRigVMRegisterOffset::Load(FArchive& Ar)
+{
+	Ar << Segments;
+	Ar << Type;
+	Ar << CPPType;
+
+	if (Ar.CustomVer(FReleaseObjectVersion::GUID) < FReleaseObjectVersion::SerializeRigVMOffsetSegmentPaths)
+	{
+		FName ScriptStructPath;
+		Ar << ScriptStructPath;
+
+		ScriptStruct = FindObject<UScriptStruct>(ANY_PACKAGE, *ScriptStructPath.ToString());
+	}
+	else
+	{
+		Ar << ScriptStruct;
+	}
+
+	Ar << ElementSize;
+
+	if (Ar.CustomVer(FReleaseObjectVersion::GUID) >= FReleaseObjectVersion::SerializeRigVMOffsetSegmentPaths)
+	{
+		FString SegmentPath;
+		Ar << ParentScriptStruct;
+		Ar << SegmentPath;
+		Ar << ArrayIndex;
+
+		if (Ar.IsTransacting())
+		{
+			CachedSegmentPath = SegmentPath;
+		}
+		else if (ParentScriptStruct != nullptr)
+		{
+			// if segment path is empty, it implies that the register offset refers to an element in a struct array
+			// so segments also need to be recalculated
+			int32 InitialOffset = ArrayIndex * ParentScriptStruct->GetStructureSize();
+			FRigVMRegisterOffset TempOffset(ParentScriptStruct, SegmentPath, InitialOffset, ElementSize);
+			if (TempOffset.GetSegments().Num() == Segments.Num())
+			{
+				Segments = TempOffset.GetSegments();
+				CachedSegmentPath = SegmentPath;
+			}
+			else
+			{
+				checkNoEntry();
+			}
+		}
+	}
+}
+
+FRigVMRegisterOffset::FRigVMRegisterOffset(UScriptStruct* InScriptStruct, const FString& InSegmentPath, int32 InInitialOffset, uint16 InElementSize, const FName& InCPPType)
+	: Segments()
+	, Type(ERigVMRegisterType::Plain)
+	, CPPType(InCPPType)
+	, ScriptStruct(nullptr)
+	, ParentScriptStruct(nullptr)
+	, ArrayIndex(0)
+	, ElementSize(InElementSize)
+	, CachedSegmentPath(InSegmentPath)
+{
+	ParentScriptStruct = InScriptStruct;
+
+	if (ParentScriptStruct)
+	{
+		ArrayIndex = InInitialOffset / InScriptStruct->GetStructureSize();
+	}
+	else
+	{
+		ArrayIndex = InInitialOffset / InElementSize;
+	}
+
+	struct FRigVMRegisterOffsetBuilder
+	{
+		static void WalkStruct(UStruct* InStruct, const FString& InPath, FRigVMRegisterOffset& Offset)
+		{
+			FString Left, Right;
+			if (!InPath.Split(TEXT("."), &Left, &Right))
+			{
+				Left = InPath;
+			}
+
+			FProperty* Property = InStruct->FindPropertyByName(*Left);
+			check(Property)
+
+			int32 SegmentIndex = Property->GetOffset_ReplaceWith_ContainerPtrToValuePtr();
+			if (Offset.Segments.Num() > 0)
+			{
+				if (Offset.Segments.Last() >= 0)
+				{
+					Offset.Segments[Offset.Segments.Num() - 1] += SegmentIndex;
+				}
+				else
+				{
+					Offset.Segments.Add(SegmentIndex);
+				}
+			}
+			else
+			{
+				Offset.Segments.Add(SegmentIndex);
+			}
+
+			if (!Right.IsEmpty())
+			{
+				if (FStructProperty* StructProperty = CastField<FStructProperty>(Property))
+				{
+					WalkStruct(StructProperty->Struct, Right, Offset);
+				}
+				else if (FArrayProperty* ArrayProperty = CastField<FArrayProperty>(Property))
+				{
+					WalkArray(ArrayProperty, Right, Offset);
+				}
+			}
+			else
+			{
+				Offset.CPPType = *Property->GetCPPType();
+				Offset.ElementSize = Property->ElementSize;
+
+				if (FArrayProperty* ArrayProperty = CastField<FArrayProperty>(Property))
+				{
+					Offset.Segments.Add(-1);
+					Property = ArrayProperty->Inner;
+
+					Offset.CPPType = *ArrayProperty->Inner->GetCPPType();
+					Offset.ElementSize = ArrayProperty->Inner->ElementSize;
+				}
+
+				if (FStructProperty* StructProperty = CastField<FStructProperty>(Property))
+				{
+					Offset.ScriptStruct = StructProperty->Struct;
+					Offset.Type = ERigVMRegisterType::Struct;
+				}
+				else
+				{
+					Offset.Type = ERigVMRegisterType::Plain;
+				}
+			}
+		}
+
+		static void WalkArray(FArrayProperty* InArrayProperty, const FString& InPath, FRigVMRegisterOffset& Offset)
+		{
+			FString Left, Right;
+			if (!InPath.Split(TEXT("."), &Left, &Right))
+			{
+				Left = InPath;
+			}
+
+			int32 ArrayIndex = FCString::Atoi(*Left);
+			int32 SegmentIndex = -1 - InArrayProperty->Inner->ElementSize * ArrayIndex;
+
+			if (Offset.Segments.Num() > 0)
+			{
+				if (Offset.Segments.Last() == 0)
+				{
+					Offset.Segments[Offset.Segments.Num() - 1] = SegmentIndex;
+				}
+				else
+				{
+					Offset.Segments.Add(SegmentIndex);
+				}
+			}
+			else
+			{
+				Offset.Segments.Add(SegmentIndex);
+			}
+
+			if (!Right.IsEmpty())
+			{
+				if (FStructProperty* StructProperty = CastField<FStructProperty>(InArrayProperty->Inner))
+				{
+					WalkStruct(StructProperty->Struct, Right, Offset);
+				}
+				else if (FArrayProperty* ArrayProperty = CastField<FArrayProperty>(InArrayProperty->Inner))
+				{
+					WalkArray(ArrayProperty, Right, Offset);
+				}
+			}
+			else
+			{
+				Offset.CPPType = *InArrayProperty->Inner->GetCPPType();
+				Offset.ElementSize = InArrayProperty->Inner->ElementSize;
+
+				if (FArrayProperty* ArrayProperty = CastField<FArrayProperty>(InArrayProperty->Inner))
+				{
+					Offset.Segments.Add(-1);
+					InArrayProperty = ArrayProperty;
+				}
+
+				if (FStructProperty* StructProperty = CastField<FStructProperty>(InArrayProperty->Inner))
+				{
+					Offset.ScriptStruct = StructProperty->Struct;
+					Offset.Type = ERigVMRegisterType::Struct;
+				}
+				else
+				{
+					Offset.Type = ERigVMRegisterType::Plain;
+				}
+			}
+
+		}
+	};
+
+	Segments.Add(InInitialOffset);
+
+	// if segment path is not empty, it implies that the register offset refers to a sub-property in a struct pin
+	if (!InSegmentPath.IsEmpty())
+	{
+		ensure(!InSegmentPath.IsEmpty());
+		check(InScriptStruct)
+		FString SegmentPath = InSegmentPath;
+		SegmentPath = SegmentPath.Replace(TEXT("["), TEXT("."));
+		SegmentPath = SegmentPath.Replace(TEXT("]"), TEXT("."));
+		FRigVMRegisterOffsetBuilder::WalkStruct(InScriptStruct, SegmentPath, *this);
+		if (Type == ERigVMRegisterType::Plain)
+		{
+			if (CPPType == TEXT("FName"))
+			{
+				Type = ERigVMRegisterType::Name;
+			}
+			else if (CPPType == TEXT("FString"))
+			{
+				Type = ERigVMRegisterType::String;
+			}
+		}
+	}
+	else
+	{
+		// if segment path is empty, it implies that the register offset refers to an element in a struct array
+		if (ParentScriptStruct)
+		{
+			ScriptStruct = ParentScriptStruct;
+			Type = ERigVMRegisterType::Struct;
+			CPPType = *ScriptStruct->GetStructCPPName();
+		}
+	}
+
+	ensure(ElementSize > 0);
+}
+
+#if UE_RIGVM_UCLASS_BASED_STORAGE_DISABLED
+
+uint8* FRigVMRegisterOffset::GetData(uint8* InContainer) const
+{
+	uint8* Data = InContainer;
+	for(int32 SegmentIndex : Segments)
+	{
+		if (SegmentIndex < 0)
+		{
+			int32 ArrayOffset = (-SegmentIndex) - 1;
+			FRigVMByteArray* ArrayPtr = (FRigVMByteArray*)Data;
+			Data = ArrayPtr->GetData() + ArrayOffset;
+		}
+		else
+		{
+			Data = Data + SegmentIndex;
+		}
+	}
+	return Data;
+}
+
+#endif
+
+bool FRigVMRegisterOffset::ContainsArraySegment() const
+{
+	for (int32 SegmentIndex : Segments)
+	{
+		if (SegmentIndex < 0) // Segement indices less than zero represent array element offsets.
+		{
+			return true;
+		}
+	}
+	return false;
+}
+
+bool FRigVMRegisterOffset::operator == (const FRigVMRegisterOffset& InOther) const
+{
+	if (Segments.Num() != InOther.Segments.Num())
+	{
+		return false;
+	}
+	if (ElementSize != InOther.ElementSize)
+	{
+		return false;
+	}
+	if (GetScriptStruct() != InOther.GetScriptStruct())
+	{
+		return false;
+	}
+	if (ParentScriptStruct != InOther.ParentScriptStruct)
+	{
+		return false;
+	}
+	if (CachedSegmentPath != InOther.CachedSegmentPath)
+	{
+		return false;
+	}
+	for (int32 Index = 0; Index < Segments.Num(); Index++)
+	{
+		if (Segments[Index] != InOther.Segments[Index])
+		{
+			return false;
+		}
+	}
+	return true;
+}
+
+UScriptStruct* FRigVMRegisterOffset::GetScriptStruct() const
+{
+	return ScriptStruct;
+}
+
+uint16 FRigVMRegisterOffset::GetElementSize() const
+{
+	return ElementSize;
+}
+
+////////////////////////////////////////////////////////////////////////////////////////////////////////////////////////
+
+class FRigVMMemoryContainerImportErrorContext : public FOutputDevice
+{
+public:
+
+	int32 NumErrors;
+
+	FRigVMMemoryContainerImportErrorContext()
+		: FOutputDevice()
+		, NumErrors(0)
+	{
+	}
+
+	FORCEINLINE_DEBUGGABLE void Serialize(const TCHAR* V, ELogVerbosity::Type Verbosity, const class FName& Category) override
+	{
+#if WITH_EDITOR
+		UE_LOG(LogRigVM, Display, TEXT("Skipping Importing To MemoryContainer: %s"), V);
+#else
+		UE_LOG(LogRigVM, Error, TEXT("Error Importing To MemoryContainer: %s"), V);
+#endif
+		NumErrors++;
+	}
+};
+
+////////////////////////////////////////////////////////////////////////////////////////////////////////////////////////
+
+#if UE_RIGVM_UCLASS_BASED_STORAGE_DISABLED
+
+FRigVMByteArray FRigVMMemoryContainer::DefaultByteArray;
+FRigVMRegisterOffset FRigVMMemoryContainer::InvalidRegisterOffset;
+
+#endif
+
+FRigVMMemoryContainer::FRigVMMemoryContainer(bool bInUseNames)
+	: bUseNameMap(bInUseNames)
+	, MemoryType(ERigVMMemoryType::Work)
+	, bEncounteredErrorDuringLoad(false)
+{
+}
+
+#if UE_RIGVM_UCLASS_BASED_STORAGE_DISABLED
+
+FRigVMMemoryContainer::FRigVMMemoryContainer(const FRigVMMemoryContainer& Other)
+{
+	*this = Other;
+}
+
+FRigVMMemoryContainer::~FRigVMMemoryContainer()
+{
+	Reset();
+}
+
+bool FRigVMMemoryContainer::CopyRegisters(const FRigVMMemoryContainer &InOther)
+{
+	ensure(Registers.Num() == InOther.Registers.Num());
+	for (int32 Index = 0; Index < Registers.Num(); ++Index)
+	{
+		if (!Copy(Index, Index, &InOther))
+		{
+			return false;
+		}
+	}
+	return true;
+}
+
+FRigVMMemoryContainer& FRigVMMemoryContainer::operator= (const FRigVMMemoryContainer &InOther)
+{
+	Reset();
+
+	MemoryType = InOther.MemoryType;
+	bUseNameMap = InOther.bUseNameMap;
+	bEncounteredErrorDuringLoad = false;
+	Data.Append(InOther.Data);
+	Registers.Append(InOther.Registers);
+	RegisterOffsets.Append(InOther.RegisterOffsets);
+	ScriptStructs.Append(InOther.ScriptStructs);
+
+	UpdateRegisters();
+
+	for (int32 Index = 0; Index < Registers.Num(); Index++)
+	{
+		const FRigVMRegister& Register = Registers[Index];
+		if (Register.IsNestedDynamic())
+		{
+			FMemory::Memzero(&Data[Register.GetWorkByteIndex()], sizeof(FRigVMNestedByteArray));
+			*(FRigVMNestedByteArray*)&Data[Register.GetWorkByteIndex()] = FRigVMNestedByteArray();
+		}
+		else if(Register.IsDynamic())
+		{
+			FMemory::Memzero(&Data[Register.GetWorkByteIndex()], sizeof(FRigVMByteArray));
+			*(FRigVMByteArray*)&Data[Register.GetWorkByteIndex()] = FRigVMByteArray();
+		}
+
+		Construct(Index);
+
+		// Why was this introduced? Are we trying to save time here?
+		// An assignment should be copying everything.
+		if (MemoryType == ERigVMMemoryType::Literal)
+		{
+			Copy(Index, Index, &InOther);
+		}
+	}
+
+	return *this;
+}
+
+#endif
+
+#if UE_RIGVM_UCLASS_BASED_STORAGE_DISABLED
+
+void FRigVMMemoryContainer::Serialize(FArchive& Ar)
+{
+	Ar.UsingCustomVersion(FAnimObjectVersion::GUID);
+
+ 	if (Ar.CustomVer(FAnimObjectVersion::GUID) < FAnimObjectVersion::StoreMarkerNamesOnSkeleton)
+ 	{
+ 		return;
+ 	}
+
+	if (Ar.IsSaving() || Ar.IsObjectReferenceCollector() || Ar.IsCountingMemory())
+	{
+		Save(Ar);
+	}
+	else if (Ar.IsLoading())
+	{
+		Load(Ar);
+	}
+	else
+	{
+		checkNoEntry();
+	}
+}
+
+#else
+
+void FRigVMMemoryContainer::Serialize(FArchive& Ar)
+{
+	return;
+}
+
+#endif
+
+#if UE_RIGVM_UCLASS_BASED_STORAGE_DISABLED
+
+void FRigVMMemoryContainer::Save(FArchive& Ar)
+{
+	for (FRigVMRegister& Register : Registers)
+	{
+		if (Register.IsNestedDynamic())
+		{
+			FRigVMNestedByteArray* NestedArrayStorage = (FRigVMNestedByteArray*)&Data[Register.GetWorkByteIndex()];
+			Register.SliceCount = NestedArrayStorage->Num();
+		}
+		else if (Register.IsDynamic())
+		{
+			FRigVMByteArray* ArrayStorage = (FRigVMByteArray*)&Data[Register.GetWorkByteIndex()];
+			Register.SliceCount = ArrayStorage->Num();
+		}
+	}
+
+	Ar << bUseNameMap;
+	Ar << MemoryType;
+	Ar << Registers;
+	Ar << RegisterOffsets;
+	
+	TArray<FString> ScriptStructPaths;
+	for (UScriptStruct* ScriptStruct : ScriptStructs)
+	{
+		ScriptStructPaths.Add(ScriptStruct->GetPathName());
+	}
+	Ar << ScriptStructPaths;
+
+	uint64 TotalBytes = Data.Num();
+	Ar << TotalBytes;
+
+	// if we are running a memory count
+	// there's no need to do a deep serialize
+	// since we know the memory needed
+	if(Ar.IsCountingMemory())
+	{
+		Ar << Data;
+		return;
+	}
+
+	for (FRigVMRegister& Register : Registers)
+	{
+		if (Register.ElementCount == 0 && !Register.IsDynamic())
+		{
+			continue;
+		}
+
+		// literal mem container should not have any dynamic registers,
+		// if we ever added them, we might need to re-consider if we need to serialize them
+		// by default dynamic registers' memory are not serialized
+		if (MemoryType == ERigVMMemoryType::Literal)
+		{
+			check(!Register.IsDynamic());
+		}
+
+		if (!Register.IsDynamic())
+		{
+			switch (Register.Type)
+			{
+				case ERigVMRegisterType::Plain:
+				{
+					FRigVMByteArray View;
+					View.Append(&Data[Register.GetWorkByteIndex()], Register.GetAllocatedBytes() - Register.GetAlignmentBytes());
+					Ar << View;
+					break;
+				}
+				case ERigVMRegisterType::Name:
+				{
+					TArray<FName> View;
+					View.Append((FName*)&Data[Register.GetWorkByteIndex()], Register.GetTotalElementCount());
+					Ar << View;
+					break;
+				}
+				case ERigVMRegisterType::String:
+				{
+					TArray<FString> View;
+					View.Append((FString*)&Data[Register.GetWorkByteIndex()], Register.GetTotalElementCount());
+					Ar << View;
+					break;
+				}
+				case ERigVMRegisterType::Struct:
+				{
+					uint8* DataPtr = &Data[Register.GetWorkByteIndex()];
+					UScriptStruct* ScriptStruct = GetScriptStruct(Register);
+
+					TArray<uint8, TAlignedHeapAllocator<16>> DefaultStructData;
+					DefaultStructData.AddZeroed(ScriptStruct->GetStructureSize());
+					ScriptStruct->InitializeDefaultValue(DefaultStructData.GetData());
+
+					TArray<FString> View;
+					for (uint16 ElementIndex = 0; ElementIndex < Register.GetTotalElementCount(); ElementIndex++)
+					{
+						FString Value;
+						ScriptStruct->ExportText(Value, DataPtr, DefaultStructData.GetData(), nullptr, PPF_None, nullptr);
+						View.Add(Value);
+						DataPtr += Register.ElementSize;
+					}
+
+					ScriptStruct->DestroyStruct(DefaultStructData.GetData(), 1);
+
+					Ar << View;
+					break;
+				}
+			}
+		}
+		else if (!Register.IsNestedDynamic())
+		{
+			FRigVMByteArray* ArrayStorage = (FRigVMByteArray*)&Data[Register.GetWorkByteIndex()];
+			
+			// notice that all dynamic registers' memory is not really serialized
+			// serialize empty arrays here just to avoid having to increment serialization version
+			// it does not need to be saved since it will be recomputed during execution
+			switch (Register.Type)
+			{
+				case ERigVMRegisterType::Plain:
+				{
+					FRigVMByteArray DummyStorage;
+					Ar << DummyStorage;
+					break;
+				}
+				case ERigVMRegisterType::Name:
+				{
+					TArray<FName> View;
+					Ar << View;
+					break;
+				}
+				case ERigVMRegisterType::String:
+				{
+					TArray<FString> View;
+					Ar << View;
+					break;
+				}
+				case ERigVMRegisterType::Struct:
+				{
+					TArray<FString> View;
+					Ar << View;
+					break;
+				}
+			}
+		}
+		else
+		{
+			FRigVMNestedByteArray* NestedArrayStorage = (FRigVMNestedByteArray*)&Data[Register.GetWorkByteIndex()];
+
+			for (int32 SliceIndex = 0; SliceIndex < Register.SliceCount; SliceIndex++)
+			{
+				if (NestedArrayStorage->Num() <= SliceIndex)
+				{
+					FRigVMByteArray EmptyStorage;
+					Ar << EmptyStorage;
+					continue;
+				}
+
+				FRigVMByteArray& ArrayStorage = (*NestedArrayStorage)[SliceIndex];
+
+				// notice that all dynamic registers' memory is not really serialized
+				// serialize empty arrays here just to avoid having to increment serialization version
+				// it does not need to be saved since it will be recomputed during execution
+				switch (Register.Type)
+				{
+					case ERigVMRegisterType::Plain:
+					{
+						FRigVMByteArray DummyStorage;
+						Ar << DummyStorage;
+						break;
+					}
+					case ERigVMRegisterType::Name:
+					{
+						TArray<FName> View;
+						Ar << View;
+						break;
+					}
+					case ERigVMRegisterType::String:
+					{
+						TArray<FString> View;
+						Ar << View;
+						break;
+					}
+					case ERigVMRegisterType::Struct:
+					{
+						TArray<FString> View;
+						Ar << View;
+						break;
+					}
+				}
+			}
+		}
+	}
+}
+
+void FRigVMMemoryContainer::Load(FArchive& Ar)
+{
+	for (int32 RegisterIndex = 0; RegisterIndex < Registers.Num(); RegisterIndex++)
+	{
+		Destroy(RegisterIndex);
+	}
+	
+	Ar << bUseNameMap;
+	Ar << MemoryType;
+	Ar << Registers;
+	Ar << RegisterOffsets;
+
+
+#if DEBUG_RIGVMMEMORY
+	UE_LOG_RIGVMMEMORY(TEXT("%d Memory - Begin Loading..."), (int32)GetMemoryType());
+#endif
+
+	bEncounteredErrorDuringLoad = false;
+
+	ScriptStructs.Reset();
+	TArray<FString> ScriptStructPaths;
+	Ar << ScriptStructPaths;
+
+	for (const FString& ScriptStructPath : ScriptStructPaths)
+	{
+		UScriptStruct* ScriptStruct = FindObject<UScriptStruct>(nullptr, *ScriptStructPath);
+
+		// this might have happened if a given script struct no longer 
+		// exists or cannot be loaded.
+		if (ScriptStruct == nullptr)
+		{
+			FString PackagePath = Ar.GetArchiveName();
+			UE_LOG(LogRigVM, Error, TEXT("Struct '%s' cannot be found. Asset '%s' no longer functional."), *ScriptStructPath, *PackagePath);
+			bEncounteredErrorDuringLoad = true;
+		}
+
+		ScriptStructs.Add(ScriptStruct);
+	}
+
+	uint64 TotalBytes = 0;
+	Ar << TotalBytes;
+
+	Data.Empty();
+
+	if (!bEncounteredErrorDuringLoad)
+	{
+		// during load we'll recreate the memory for all registers.
+		// the size for structs might have changed, so we need to reallocate.
+		for (int32 RegisterIndex = 0; RegisterIndex < Registers.Num(); RegisterIndex++)
+		{
+			FRigVMRegister& Register = Registers[RegisterIndex];
+
+			UScriptStruct* ScriptStruct = GetScriptStruct(Register);
+			if (ScriptStruct)
+			{
+				Register.ElementSize = ScriptStruct->GetStructureSize();
+			}
+			else if (Register.Type == ERigVMRegisterType::Name)
+			{
+				Register.ElementSize = sizeof(FName);
+			}
+			else if (Register.Type == ERigVMRegisterType::String)
+			{
+				Register.ElementSize = sizeof(FString);
+			}
+
+			Register.AlignmentBytes = 0;
+			Register.TrailingBytes = 0;
+
+			if (Registers[RegisterIndex].IsDynamic())
+			{
+				Register.ByteIndex = Data.AddZeroed(sizeof(FRigVMByteArray));
+			}
+			else
+			{
+				Register.ByteIndex = Data.AddZeroed(Register.GetNumBytesAllSlices());
+			}
+		}
+		UpdateRegisters();
+
+		for (int32 RegisterOffsetIndex = 0; RegisterOffsetIndex < RegisterOffsets.Num(); RegisterOffsetIndex++)
+		{
+			FRigVMRegisterOffset& RegisterOffset = RegisterOffsets[RegisterOffsetIndex];
+
+			UScriptStruct* ScriptStruct = RegisterOffset.GetScriptStruct();
+			if (ScriptStruct)
+			{
+				RegisterOffset.SetElementSize(ScriptStruct->GetStructureSize());
+			}
+			if (RegisterOffset.GetType() == ERigVMRegisterType::Name)
+			{
+				RegisterOffset.SetElementSize(sizeof(FName));
+			}
+			else if (RegisterOffset.GetType() == ERigVMRegisterType::String)
+			{
+				RegisterOffset.SetElementSize(sizeof(FString));
+			}
+		}
+
+		// once the register memory is allocated we can construt its contents.
+		for (int32 RegisterIndex = 0; RegisterIndex < Registers.Num(); RegisterIndex++)
+		{
+			const FRigVMRegister& Register = Registers[RegisterIndex];
+			if (!Registers[RegisterIndex].IsDynamic())
+			{
+				Construct(RegisterIndex);
+			}
+		}
+	}
+
+	for (int32 RegisterIndex = 0; RegisterIndex < Registers.Num(); RegisterIndex++)
+	{
+		FRigVMRegister& Register = Registers[RegisterIndex];
+
+		if (Register.ElementCount == 0 && !Register.IsDynamic())
+		{
+			continue;
+		}
+
+		// literal mem container should not have any dynamic registers,
+		// if we ever added them, we might need to re-consider if we need to serialize them
+		// by default dynamic registers' memory are not serialized
+		if (MemoryType == ERigVMMemoryType::Literal)
+		{
+			check(!Register.IsDynamic());
+		}
+		
+		if (!Register.IsDynamic())
+		{
+			switch (Register.Type)
+			{
+				case ERigVMRegisterType::Plain:
+				{
+					FRigVMByteArray View;
+					Ar << View;
+					if (!bEncounteredErrorDuringLoad)
+					{
+						ensure(View.Num() <= Register.GetAllocatedBytes());
+						FMemory::Memcpy(&Data[Register.GetWorkByteIndex()], View.GetData(), View.Num());
+					}
+					break;
+				}
+				case ERigVMRegisterType::Name:
+				{
+					TArray<FName> View;
+					Ar << View;
+					if (!bEncounteredErrorDuringLoad)
+					{
+						ensure(View.Num() == Register.GetTotalElementCount());
+						RigVMCopy<FName>(&Data[Register.GetWorkByteIndex()], View.GetData(), View.Num());
+					}
+					break;
+				}
+				case ERigVMRegisterType::String:
+				{
+					TArray<FString> View;
+					Ar << View;
+					if (!bEncounteredErrorDuringLoad)
+					{
+						ensure(View.Num() == Register.GetTotalElementCount());
+						RigVMCopy<FString>(&Data[Register.GetWorkByteIndex()], View.GetData(), View.Num());
+					}
+					break;
+				}
+				case ERigVMRegisterType::Struct:
+				{
+					TArray<FString> View;
+					Ar << View;
+					if (!bEncounteredErrorDuringLoad)
+					{
+						ensure(View.Num() == Register.GetTotalElementCount());
+
+						uint8* DataPtr = &Data[Register.GetWorkByteIndex()];
+						UScriptStruct* ScriptStruct = GetScriptStruct(Register);
+						if (ScriptStruct && !bEncounteredErrorDuringLoad)
+						{
+							ensure(ScriptStruct->GetStructureSize() == Register.ElementSize);
+
+							for (uint16 ElementIndex = 0; ElementIndex < Register.GetTotalElementCount(); ElementIndex++)
+							{
+								FRigVMMemoryContainerImportErrorContext ErrorPipe;
+								ScriptStruct->ImportText(*View[ElementIndex], DataPtr, nullptr, PPF_None, &ErrorPipe, ScriptStruct->GetName());
+								if (ErrorPipe.NumErrors > 0)
+								{
+									bEncounteredErrorDuringLoad = true;
+									break;
+								}
+								DataPtr += Register.ElementSize;
+							}
+						}
+					}
+					break;
+				}
+			}
+		}
+		else if (!Register.IsNestedDynamic())
+		{
+			FRigVMByteArray* ArrayStorage = nullptr;
+
+			if (!bEncounteredErrorDuringLoad)
+			{
+				ArrayStorage = (FRigVMByteArray*)&Data[Register.GetWorkByteIndex()];
+			}
+			
+			// notice that all dynamic registers' memory is not really serialized
+			// serialize empty arrays here just to avoid having to increment serialization version
+			// the loaded data is not getting used at all, it will be recomputed at execution time
+			switch (Register.Type)
+			{
+				case ERigVMRegisterType::Plain:
+				{
+					FRigVMByteArray View;
+					Ar << View;
+					break;
+				}
+				case ERigVMRegisterType::Name:
+				{
+					TArray<FName> View;
+					Ar << View;
+					break;
+				}
+				case ERigVMRegisterType::String:
+				{
+					TArray<FString> View;
+					Ar << View;
+					break;
+				}
+				case ERigVMRegisterType::Struct:
+				{
+					TArray<FString> View;
+					Ar << View;
+					break;
+				}
+			}
+		}
+		else
+		{
+			FRigVMNestedByteArray* NestedArrayStorage = nullptr;
+			if (!bEncounteredErrorDuringLoad)
+			{
+				NestedArrayStorage = (FRigVMNestedByteArray*)&Data[Register.GetWorkByteIndex()];
+				NestedArrayStorage->Reset();
+				NestedArrayStorage->SetNumZeroed(Register.SliceCount);
+			}
+
+			for (int32 SliceIndex = 0; SliceIndex < Register.SliceCount; SliceIndex++)
+			{
+				FRigVMByteArray* ArrayStorage = nullptr;
+				if (NestedArrayStorage)
+				{
+					ArrayStorage = &(*NestedArrayStorage)[SliceIndex];
+				}
+
+				// notice that all dynamic registers' memory is not really serialized
+				// serialize empty arrays here just to avoid having to increment serialization version
+				// the loaded data is not getting used at all, it will be recomputed at execution time
+				switch (Register.Type)
+				{
+					case ERigVMRegisterType::Plain:
+					{
+						FRigVMByteArray View;
+						Ar << View;
+						break;
+					}
+					case ERigVMRegisterType::Name:
+					{
+						TArray<FName> View;
+						Ar << View;
+						break;
+					}
+					case ERigVMRegisterType::String:
+					{
+						TArray<FString> View;
+						Ar << View;
+						break;
+					}
+					case ERigVMRegisterType::Struct:
+					{
+						TArray<FString> View;
+						Ar << View;
+						break;
+					}
+				}
+			}
+		}
+	}
+
+	if (bEncounteredErrorDuringLoad)
+	{
+#if DEBUG_RIGVMMEMORY
+		UE_LOG_RIGVMMEMORY(TEXT("%d Memory - Encountered errors during load."), (int32)GetMemoryType());
+#endif
+		Reset();
+	}
+	else
+	{
+		UpdateRegisters();
+	}
+
+#if DEBUG_RIGVMMEMORY
+	UE_LOG_RIGVMMEMORY(TEXT("%d Memory - Finished Loading."), (int32)GetMemoryType());
+#endif
+	
+}
+
+void FRigVMMemoryContainer::Reset()
+{
+	if (Data.Num() > 0)
+	{
+		for (int32 Index = 0; Index < Registers.Num(); Index++)
+		{
+			Destroy(Index);
+		}
+	}
+
+	Data.Reset();
+	Registers.Reset();
+	RegisterOffsets.Reset();
+	ScriptStructs.Reset();
+	NameMap.Reset();
+}
+
+void FRigVMMemoryContainer::Empty()
+{
+	Data.Empty();
+	Registers.Empty();
+	RegisterOffsets.Empty();
+	ScriptStructs.Empty();
+	NameMap.Empty();
+}
+
+bool FRigVMMemoryContainer::Copy(
+	int32 InTargetRegisterIndex,
+	int32 InTargetRegisterOffset,
+	ERigVMRegisterType InTargetType,
+	const uint8* InSourcePtr,
+	uint8* InTargetPtr,
+	uint16 InNumBytes)
+{
+	switch (InTargetType)
+	{
+		case ERigVMRegisterType::Plain:
+		{
+			FMemory::Memcpy(InTargetPtr, InSourcePtr, InNumBytes);
+			break;
+		}
+		case ERigVMRegisterType::Struct:
+		{
+			UScriptStruct* ScriptStruct = GetScriptStruct(InTargetRegisterIndex, InTargetRegisterOffset);
+			int32 NumStructs = InNumBytes / ScriptStruct->GetStructureSize();
+			ensure(NumStructs * ScriptStruct->GetStructureSize() == InNumBytes);
+			if (NumStructs > 0 && InTargetPtr)
+			{
+				ScriptStruct->CopyScriptStruct(InTargetPtr, InSourcePtr, NumStructs);
+			}
+			break;
+		}
+		case ERigVMRegisterType::Name:
+		{
+			int32 NumNames = InNumBytes / sizeof(FName);
+			ensure(NumNames * sizeof(FName) == InNumBytes);
+			RigVMCopy<FName>(InTargetPtr, InSourcePtr, NumNames);
+			break;
+		}
+		case ERigVMRegisterType::String:
+		{
+			int32 NumStrings = InNumBytes / sizeof(FString);
+			ensure(NumStrings * sizeof(FString) == InNumBytes);
+			RigVMCopy<FString>(InTargetPtr, InSourcePtr, NumStrings);
+			break;
+		}
+		case ERigVMRegisterType::Invalid:
+		{
+			return false;
+		}
+	}
+
+	return true;
+}
+
+bool FRigVMMemoryContainer::Copy(
+	int32 InSourceRegisterIndex,
+	int32 InTargetRegisterIndex,
+	const FRigVMMemoryContainer* InSourceMemory,
+	int32 InSourceRegisterOffset,
+	int32 InTargetRegisterOffset,
+	int32 InSourceSliceIndex,
+	int32 InTargetSliceIndex)
+{
+	if (InSourceMemory == nullptr)
+	{
+		InSourceMemory = this;
+	}
+
+	ensure(InSourceMemory->Registers.IsValidIndex(InSourceRegisterIndex));
+	ensure(Registers.IsValidIndex(InTargetRegisterIndex));
+
+	if (InSourceRegisterIndex == InTargetRegisterIndex && InSourceRegisterOffset == InTargetRegisterOffset && this == InSourceMemory)
+	{
+		return false;
+	}
+
+	const FRigVMRegister& Source = InSourceMemory->Registers[InSourceRegisterIndex];
+	const FRigVMRegister& Target = Registers[InTargetRegisterIndex];
+
+	// todo: copy between slices!
+	int32 SourceElementCount = Source.ElementCount;
+	const uint8* SourcePtr = InSourceMemory->GetDataPtr(Source, InSourceRegisterOffset, InSourceSliceIndex, false);
+	uint8* TargetPtr = GetDataPtr(Target, InTargetRegisterOffset, InTargetSliceIndex, false);
+	uint16 NumBytes = Target.GetNumBytesPerSlice();
+
+	// TODO: preventing crash when SourcePtr is null
+	if (SourcePtr == nullptr)
+	{
+		return false;
+	}
+	
+	if (Source.IsNestedDynamic())
+	{
+#if WITH_EDITOR
+		ensure(InSourceRegisterOffset == INDEX_NONE);
+#else
+		if (InSourceRegisterOffset != INDEX_NONE)
+		{
+			return true;
+		}
+#endif
+		FRigVMByteArray* ArrayStorage = (FRigVMByteArray*)SourcePtr;
+		SourceElementCount = ArrayStorage->Num() / Source.ElementSize;
+		SourcePtr = ArrayStorage->GetData();
+	}
+	else if (Source.IsDynamic())
+	{
+		FRigVMByteArray* ArrayStorage = (FRigVMByteArray*)SourcePtr;
+		SourcePtr = ArrayStorage->GetData();
+	}
+
+	// TODO: preventing crash when TargetPtr has not been allocated 
+	if (TargetPtr == nullptr && Target.IsNestedDynamic())
+	{
+		uint8* Ptr = (uint8*)&Data[Target.GetWorkByteIndex()];
+		FRigVMNestedByteArray& Storage = *(FRigVMNestedByteArray*)Ptr;
+		Storage.SetNum(Target.SliceCount);
+		Storage[InTargetSliceIndex].SetNumZeroed(SourceElementCount * Source.ElementSize);
+		Copy(InTargetRegisterIndex, INDEX_NONE, Source.Type, SourcePtr, Storage[InTargetSliceIndex].GetData(), Source.ElementSize * SourceElementCount);
+		return true;
+	}
+
+	ERigVMRegisterType TargetType = Target.Type;
+	if (Target.IsDynamic() && !Target.bIsArray)
+	{
+		FRigVMByteArray* ArrayStorage = (FRigVMByteArray*)TargetPtr;
+		TargetPtr = ArrayStorage->GetData();
+	}
+	else if (TargetPtr && Target.IsNestedDynamic())
+	{
+#if WITH_EDITOR
+		ensure(InTargetRegisterOffset == INDEX_NONE);
+#else
+		if (InTargetRegisterOffset != INDEX_NONE)
+		{
+			return true;
+		}
+		return true;
+#endif
+
+		NumBytes = Source.ElementSize * SourceElementCount;
+
+		FRigVMByteArray* ArrayStorage = (FRigVMByteArray*)TargetPtr;
+		if (ArrayStorage->Num() != NumBytes)
+		{
+			Destroy(InTargetRegisterIndex, INDEX_NONE, InTargetSliceIndex);
+			ArrayStorage->SetNumZeroed(NumBytes);
+			Construct(InTargetRegisterIndex, INDEX_NONE, InTargetSliceIndex);
+
+			TargetPtr = ArrayStorage->GetData();
+		}
+	}
+	else if (InTargetRegisterOffset != INDEX_NONE)
+	{
+		if (Target.GetNumBytesPerSlice() == 0)
+		{
+			return true;
+		}
+
+		TargetType = RegisterOffsets[InTargetRegisterOffset].GetType();
+		NumBytes = RegisterOffsets[InTargetRegisterOffset].GetElementSize();
+	}
+	else if (Target.GetNumBytesPerSlice() == 0)
+	{
+		return true;
+	}
+
+	return Copy(InTargetRegisterIndex, InTargetRegisterOffset, TargetType, SourcePtr, TargetPtr, NumBytes);
+}
+
+bool FRigVMMemoryContainer::Copy(
+	const FName& InSourceName,
+	const FName& InTargetName,
+	const FRigVMMemoryContainer* InSourceMemory,
+	int32 InSourceRegisterOffset,
+	int32 InTargetRegisterOffset,
+	int32 InSourceSliceIndex,
+	int32 InTargetSliceIndex)
+{
+	ensure(bUseNameMap);
+
+	int32 SourceRegisterIndex = GetIndex(InSourceName);
+	int32 TargetRegisterIndex = GetIndex(InTargetName);
+
+	if(SourceRegisterIndex == INDEX_NONE || TargetRegisterIndex == INDEX_NONE)
+	{
+		return false;
+	}
+
+	return Copy(SourceRegisterIndex, TargetRegisterIndex, InSourceMemory, InSourceRegisterOffset, InTargetRegisterOffset);
+}
+
+bool FRigVMMemoryContainer::Copy(
+	const FRigVMOperand& InSourceOperand,
+	const FRigVMOperand& InTargetOperand,
+	const FRigVMMemoryContainer* InSourceMemory,
+	int32 InSourceSliceIndex,
+	int32 InTargetSliceIndex)
+{
+	return Copy(InSourceOperand.GetRegisterIndex(), InTargetOperand.GetRegisterIndex(), InSourceMemory, InSourceOperand.GetRegisterOffset(), InTargetOperand.GetRegisterOffset(), InSourceSliceIndex, InTargetSliceIndex);
+}
+
+int32 FRigVMMemoryContainer::Allocate(const FName& InNewName, int32 InElementSize, int32 InElementCount, int32 InSliceCount, const uint8* InDataPtr, bool bUpdateRegisters)
+{
+	FName Name = InNewName;
+	if (bUseNameMap && InNewName == NAME_None)
+	{
+		const TCHAR* RegisterPrefix = TEXT("Register");
+		int32 RegisterSuffix = 0;
+		do
+		{
+			Name = FName(*FString::Printf(TEXT("%s_%d"), RegisterPrefix, RegisterSuffix++));
+		} while (!IsNameAvailable(Name));
+	}
+
+	ensure(InElementSize > 0 && InElementCount >= 0 && InSliceCount > 0);
+
+	if (bUseNameMap)
+	{
+		if (!IsNameAvailable(Name))
+		{
+			return INDEX_NONE;
+		}
+	}
+
+	FRigVMRegister NewRegister;
+	NewRegister.ByteIndex = Data.Num();
+	if (bUseNameMap)
+	{
+		NewRegister.Name = Name;
+	}
+	NewRegister.ElementSize = (uint16)InElementSize;
+	NewRegister.ElementCount = (uint16)InElementCount;
+	NewRegister.SliceCount = (uint16)InSliceCount;
+	NewRegister.Type = ERigVMRegisterType::Plain;
+
+	Data.AddZeroed(NewRegister.GetAllocatedBytes());
+
+#if DEBUG_RIGVMMEMORY
+	if (NewRegister.ElementCount > 0)
+	{
+		UE_LOG_RIGVMMEMORY(TEXT("%d.%04d: Allocated %04d bytes at %04d (%s)."), (int32)GetMemoryType(), Registers.Num(), NewRegister.GetAllocatedBytes(), (int32)reinterpret_cast<long long>(&Data[NewRegister.ByteIndex]), *NewRegister.Name.ToString());
+	}
+#endif
+
+	if (InDataPtr != nullptr)
+	{
+		for (uint16 SliceIndex = 0; SliceIndex < NewRegister.SliceCount; SliceIndex++)
+		{
+			FMemory::Memcpy(&Data[NewRegister.GetWorkByteIndex(SliceIndex)], InDataPtr, NewRegister.GetNumBytesPerSlice());
+		}
+	}
+
+	int32 RegisterIndex = Registers.Num();
+	Registers.Add(NewRegister);
+
+	if (bUpdateRegisters)
+	{
+		UpdateRegisters();
+	}
+	return RegisterIndex;
+}
+
+int32 FRigVMMemoryContainer::Allocate(int32 InElementSize, int32 InElementCount, int32 InSliceCount, const uint8* InDataPtr, bool bUpdateRegisters)
+{
+	return Allocate(NAME_None, InElementSize, InElementCount, InSliceCount, InDataPtr, bUpdateRegisters);
+}
+
+bool FRigVMMemoryContainer::Construct(int32 InRegisterIndex, int32 InElementIndex, int32 InSliceIndex)
+{
+	ensure(Registers.IsValidIndex(InRegisterIndex));
+
+	const FRigVMRegister& Register = Registers[InRegisterIndex];
+
+	if (Register.ElementCount == 0 || (Register.Type == ERigVMRegisterType::Plain && !Register.IsDynamic()))
+	{
+		return true;
+	}
+
+	int32 ElementIndex = InElementIndex == INDEX_NONE ? 0 : InElementIndex;
+	int32 ElementCount = InElementIndex == INDEX_NONE ? Register.GetTotalElementCount() : 1;
+	int32 NumSlices = 1;
+
+	uint8* TopDataPtr = &Data[Register.GetWorkByteIndex()];
+	if (Register.IsDynamic())
+	{
+		if (Register.IsNestedDynamic())
+		{
+			FRigVMNestedByteArray& Storage = *(FRigVMNestedByteArray*)TopDataPtr;
+			NumSlices = Storage.Num();
+			TopDataPtr = (uint8*)Storage.GetData();
+		}
+		else
+		{
+			FRigVMByteArray& Storage = *(FRigVMByteArray*)TopDataPtr;
+			ElementCount = Storage.Num() / Register.ElementSize;
+			if (ElementCount == 0)
+			{
+				Storage.AddZeroed(Register.ElementSize);
+				ElementCount = 1;
+			}
+			ensure(ElementCount * Register.ElementSize == Storage.Num());
+			TopDataPtr = Storage.GetData();
+		}
+	}
+
+	for (int32 SliceIndex = 0; SliceIndex < NumSlices; SliceIndex++)
+	{
+		if (InSliceIndex != INDEX_NONE)
+		{
+			SliceIndex = InSliceIndex;
+
+			if (Register.IsNestedDynamic())
+			{
+				TopDataPtr = TopDataPtr + SliceIndex * sizeof(FRigVMByteArray);
+			}
+		}
+
+		uint8* NestedDataPtr = TopDataPtr;
+		if (Register.IsNestedDynamic())
+		{
+			FRigVMByteArray& Storage = *(FRigVMByteArray*)TopDataPtr;
+			ElementCount = Storage.Num() / Register.ElementSize;
+			ensure(ElementCount * Register.ElementSize == Storage.Num());
+			NestedDataPtr = Storage.GetData();
+		}
+
+		switch (Register.Type)
+		{
+			case ERigVMRegisterType::Struct:
+			{
+				uint8* DataPtr = NestedDataPtr + ElementIndex * Register.ElementSize;
+				UScriptStruct* ScriptStruct = GetScriptStruct(InRegisterIndex);
+
+				if (ScriptStruct)
+				{
+					if (Register.ElementSize == ScriptStruct->GetStructureSize())
+					{
+						ScriptStruct->InitializeStruct(DataPtr, ElementCount);
+#if DEBUG_RIGVMMEMORY
+						UE_LOG_RIGVMMEMORY(TEXT("%d.%04d: Initialized struct, %04d bytes at %04d (%s)."), (int32)GetMemoryType(), InRegisterIndex, ElementCount * ScriptStruct->GetStructureSize(), (int32)reinterpret_cast<long long>(DataPtr), *Register.Name.ToString());
+#endif
+					}
+				}
+				break;
+			}
+			case ERigVMRegisterType::String:
+			{
+				FString* DataPtr = (FString*)(NestedDataPtr + ElementIndex * Register.ElementSize);
+				RigVMInitialize<FString>(DataPtr, ElementCount);
+#if DEBUG_RIGVMMEMORY
+				UE_LOG_RIGVMMEMORY(TEXT("%d.%04d: Initialized string, %04d bytes at %04d (%s)."), (int32)GetMemoryType(), InRegisterIndex, ElementCount * Register.ElementSize, (int32)reinterpret_cast<long long>(DataPtr), *Register.Name.ToString());
+#endif
+				break;
+			}
+			case ERigVMRegisterType::Name:
+			{
+				FName* DataPtr = (FName*)(NestedDataPtr + ElementIndex * Register.ElementSize);
+				RigVMInitialize<FName>(DataPtr, ElementCount);
+#if DEBUG_RIGVMMEMORY
+				UE_LOG_RIGVMMEMORY(TEXT("%d.%04d: Initialized name, %04d bytes at %04d (%s)."), (int32)GetMemoryType(), InRegisterIndex, ElementCount * Register.ElementSize, (int32)reinterpret_cast<long long>(DataPtr), *Register.Name.ToString());
+#endif
+				break;
+			}
+			default:
+			{
+				return false;
+			}
+		}
+
+		if (Register.IsNestedDynamic())
+		{
+			TopDataPtr = TopDataPtr + sizeof(FRigVMByteArray);
+		}
+
+		if (InSliceIndex != INDEX_NONE)
+		{
+			break;
+		}
+	}
+
+	return true;
+}
+
+bool FRigVMMemoryContainer::Destroy(int32 InRegisterIndex, int32 InElementIndex, int32 InSliceIndex)
+{
+	FRigVMRegister& Register = Registers[InRegisterIndex];
+
+	if (Register.ElementCount == 0 || (Register.Type == ERigVMRegisterType::Plain && !Register.IsDynamic()))
+	{
+		return true;
+	}
+
+	int32 ElementIndex = InElementIndex == INDEX_NONE ? 0 : InElementIndex;
+	int32 ElementCount = InElementIndex == INDEX_NONE ? Register.GetTotalElementCount() : 1;
+	int32 NumSlices = 1;
+
+	uint8* TopDataPtr = &Data[Register.GetWorkByteIndex()];
+	if (Register.IsDynamic())
+	{
+		if (Register.IsNestedDynamic())
+		{
+			FRigVMNestedByteArray& Storage = *(FRigVMNestedByteArray*)TopDataPtr;
+			NumSlices = Storage.Num();
+			TopDataPtr = (uint8*)Storage.GetData();
+		}
+		else
+		{
+			FRigVMByteArray& Storage = *(FRigVMByteArray*)TopDataPtr;
+			ElementCount = Storage.Num() / Register.ElementSize;
+			ensure(ElementCount * Register.ElementSize == Storage.Num());
+			TopDataPtr = Storage.GetData();
+		}
+	}
+
+	for (int32 SliceIndex = 0; SliceIndex < NumSlices; SliceIndex++)
+	{
+		if (InSliceIndex != INDEX_NONE) 
+		{
+			SliceIndex = InSliceIndex;
+
+			if (Register.IsNestedDynamic())
+			{
+				TopDataPtr = TopDataPtr + SliceIndex * sizeof(FRigVMByteArray);
+			}
+		}
+
+		uint8* NestedDataPtr = TopDataPtr;
+		if (Register.IsNestedDynamic())
+		{
+			FRigVMByteArray& Storage = *(FRigVMByteArray*)TopDataPtr;
+			ElementCount = Storage.Num() / Register.ElementSize;
+			ensure(ElementCount * Register.ElementSize == Storage.Num());
+			NestedDataPtr = Storage.GetData();
+		}
+
+		switch (Register.Type)
+		{
+			case ERigVMRegisterType::Struct:
+			{
+				uint8* DataPtr = NestedDataPtr + ElementIndex * Register.ElementSize;
+
+				UScriptStruct* ScriptStruct = GetScriptStruct(InRegisterIndex);
+				if (ScriptStruct)
+				{
+					if (Register.ElementSize != ScriptStruct->GetStructureSize())
+					{
+						FMemory::Memzero(DataPtr, Register.ElementSize * ElementCount);
+#if DEBUG_RIGVMMEMORY
+						UE_LOG_RIGVMMEMORY(TEXT("%d.%04d: Zeroed struct, %04d bytes at %04d (%s)."), (int32)GetMemoryType(), InRegisterIndex, Register.ElementSize * ElementCount, (int32)reinterpret_cast<long long>(DataPtr), *Register.Name.ToString());
+#endif
+					}
+					else
+					{
+						ScriptStruct->DestroyStruct(DataPtr, ElementCount);
+#if DEBUG_RIGVMMEMORY
+						UE_LOG_RIGVMMEMORY(TEXT("%d.%04d: Destroyed struct, %04d bytes at %04d (%s)."), (int32)GetMemoryType(), InRegisterIndex, ScriptStruct->GetStructureSize() * ElementCount, (int32)reinterpret_cast<long long>(DataPtr), *Register.Name.ToString());
+#endif
+					}
+				}
+				else
+				{
+					FMemory::Memzero(DataPtr, Register.ElementSize * ElementCount);
+#if DEBUG_RIGVMMEMORY
+					UE_LOG_RIGVMMEMORY(TEXT("%d.%04d: Zeroed struct, %04d bytes at %04d (%s)."), (int32)GetMemoryType(), InRegisterIndex, Register.ElementSize * ElementCount, (int32)reinterpret_cast<long long>(DataPtr), *Register.Name.ToString());
+#endif
+				}
+				break;
+			}
+			case ERigVMRegisterType::String:
+			{
+				FString* DataPtr = (FString*)(NestedDataPtr + ElementIndex * Register.ElementSize);
+				RigVMDestroy<FString>(DataPtr, ElementCount);
+#if DEBUG_RIGVMMEMORY
+				UE_LOG_RIGVMMEMORY(TEXT("%d.%04d: Destroyed string, %04d bytes at %04d (%s)."), (int32)GetMemoryType(), InRegisterIndex, ElementCount * Register.ElementSize, (int32)reinterpret_cast<long long>(DataPtr), *Register.Name.ToString());
+#endif
+				break;
+			}
+			case ERigVMRegisterType::Name:
+			{
+				FName* DataPtr = (FName*)(NestedDataPtr + ElementIndex * Register.ElementSize);
+				RigVMDestroy<FName>(DataPtr, ElementCount);
+#if DEBUG_RIGVMMEMORY
+				UE_LOG_RIGVMMEMORY(TEXT("%d.%04d: Destroyed name, %04d bytes at %04d (%s)."), (int32)GetMemoryType(), InRegisterIndex, ElementCount * Register.ElementSize, (int32)reinterpret_cast<long long>(DataPtr), *Register.Name.ToString());
+#endif
+				break;
+			}
+			default:
+			{
+				if (!Register.IsDynamic())
+				{
+					return false;
+				}
+			}
+		}
+
+		if (Register.IsNestedDynamic())
+		{
+			((FRigVMByteArray*)TopDataPtr)->Empty();
+			TopDataPtr = TopDataPtr + sizeof(FRigVMByteArray);
+		}
+
+		if (InSliceIndex != INDEX_NONE)
+		{
+			break;
+		}
+	}
+
+	if (Register.IsDynamic())
+	{
+		TopDataPtr = &Data[Register.GetWorkByteIndex()];
+
+		if (Register.IsNestedDynamic())
+		{
+			if (InSliceIndex != INDEX_NONE)
+			{
+				((FRigVMNestedByteArray*)TopDataPtr)->Empty();
+			}
+		}
+		else
+		{
+			((FRigVMByteArray*)TopDataPtr)->Empty();
+		}
+	}
+
+	return true;
+}
+
+int32 FRigVMMemoryContainer::GetOrAddRegisterOffset(int32 InRegisterIndex, int32 InArrayElement)
+{
+	return GetOrAddRegisterOffset(InRegisterIndex, FString(), InArrayElement);
+}
+
+int32 FRigVMMemoryContainer::GetOrAddRegisterOffset(int32 InRegisterIndex, const FString& InSegmentPath, int32 InArrayElement)
+{
+	ensure(Registers.IsValidIndex(InRegisterIndex));
+
+	UScriptStruct* ScriptStruct = GetScriptStruct(InRegisterIndex);
+	int32 ElementSize = (int32)Registers[InRegisterIndex].ElementSize;
+	int32 InitialOffset = InArrayElement * ElementSize;
+	return GetOrAddRegisterOffset(InRegisterIndex, ScriptStruct, InSegmentPath, InitialOffset);
+}
+
+int32 FRigVMMemoryContainer::GetOrAddRegisterOffset(int32 InRegisterIndex, UScriptStruct* InScriptStruct, const FString& InSegmentPath, int32 InInitialOffset, int32 InElementSize)
+{
+	if ((InScriptStruct == nullptr || InSegmentPath.IsEmpty()) && InInitialOffset == 0)
+	{
+		return INDEX_NONE;
+	}
+
+	// if this is a register offset for a external variable
+	// the register index is expected to be INDEX_NONE
+	// and it is also expected that InElementSize != 0
+	if (InElementSize == 0)
+	{
+		ensure(Registers.IsValidIndex(InRegisterIndex)); 
+		InElementSize = (int32)Registers[InRegisterIndex].ElementSize;
+	}
+	
+	FName BaseCPPType = NAME_None;
+	if(Registers.IsValidIndex(InRegisterIndex))
+	{
+		if(InElementSize == 0)
+		{
+			InElementSize = (int32)Registers[InRegisterIndex].ElementSize;
+		}
+#if WITH_EDITORONLY_DATA
+		BaseCPPType = Registers[InRegisterIndex].BaseCPPType;
+#endif
+	}
+
+	FRigVMRegisterOffset Offset(InScriptStruct, InSegmentPath, InInitialOffset, InElementSize, BaseCPPType);
+	int32 ExistingIndex = RegisterOffsets.Find(Offset);
+	if (ExistingIndex == INDEX_NONE)
+	{
+		return RegisterOffsets.Add(Offset);
+	}
+	return ExistingIndex;
+}
+
+void FRigVMMemoryContainer::SetRegisterValueFromString(const FRigVMOperand& InOperand, const FString& InCPPType, const UObject* InCPPTypeObject, const TArray<FString>& InDefaultValues)
+{
+	if (InOperand.GetRegisterIndex() < 0 || InOperand.GetRegisterIndex() >= Registers.Num())
+	{
+		return;
+	}
+
+	FRigVMRegister& Register = Registers[InOperand.GetRegisterIndex()];
+
+	if (Register.ElementCount != InDefaultValues.Num())
+	{
+		return;
+	}
+
+	FString CPPType = InCPPType;
+	if (RigVMUtilities::IsArrayType(CPPType))
+	{
+		CPPType = CPPType.Mid(7, CPPType.Len() - 8);
+	}
+
+	for (int32 Index = 0; Index < InDefaultValues.Num(); Index++)
+	{
+		FString DefaultValue = InDefaultValues[Index];
+
+		if (const UScriptStruct* ScriptStruct = Cast<UScriptStruct>(InCPPTypeObject))
+		{
+			uint8* DataPtr = (uint8*)GetData(Register);
+			DataPtr += Index * ScriptStruct->GetStructureSize();
+			FRigVMMemoryContainerImportErrorContext ErrorPipe;
+			((UScriptStruct*)ScriptStruct)->ImportText(*DefaultValue, DataPtr, nullptr, PPF_None, &ErrorPipe, ScriptStruct->GetName());
+		}
+		else if (const UEnum* Enum = Cast<const UEnum>(InCPPTypeObject))
+		{
+			if (FCString::IsNumeric(*DefaultValue))
+			{
+				GetFixedArray<uint8>(Register)[Index] = (uint8)FCString::Atoi(*DefaultValue);
+			}
+			else
+			{
+				GetFixedArray<uint8>(Register)[Index] = (uint8)Enum->GetValueByNameString(DefaultValue);
+			}
+		}
+		else if (CPPType == TEXT("bool") && Register.Type == ERigVMRegisterType::Plain && Register.ElementSize == sizeof(bool))
+		{
+			GetFixedArray<bool>(Register)[Index] = (DefaultValue == TEXT("True")) || (DefaultValue == TEXT("true")) || (DefaultValue == TEXT("1"));
+		}
+		else if (CPPType == TEXT("int32") && Register.Type == ERigVMRegisterType::Plain && Register.ElementSize == sizeof(int32))
+		{
+			GetFixedArray<int32>(Register)[Index] = FCString::Atoi(*DefaultValue);
+		}
+		else if (CPPType == TEXT("float") && Register.Type == ERigVMRegisterType::Plain && Register.ElementSize == sizeof(float))
+		{
+			GetFixedArray<float>(Register)[Index] = FCString::Atof(*DefaultValue);
+		}
+		else if (CPPType == TEXT("double") && Register.Type == ERigVMRegisterType::Plain && Register.ElementSize == sizeof(double))
+		{
+			GetFixedArray<double>(Register)[Index] = FCString::Atod(*DefaultValue);
+		}
+		else if (CPPType == TEXT("FName") && Register.Type == ERigVMRegisterType::Name)
+		{
+			GetFixedArray<FName>(Register)[Index] = *DefaultValue;
+		}
+		else if (CPPType == TEXT("FString") && Register.Type == ERigVMRegisterType::String)
+		{
+			GetFixedArray<FString>(Register)[Index] = *DefaultValue;
+		}
+	}
+}
+
+TArray<FString> FRigVMMemoryContainer::GetRegisterValueAsString(const FRigVMOperand& InOperand, const FString& InCPPType, const UObject* InCPPTypeObject)
+{
+	TArray<FString> DefaultValues;
+
+	if (InOperand.GetRegisterIndex() < 0 || InOperand.GetRegisterIndex() >= Registers.Num())
+	{
+		return DefaultValues;
+	}
+
+	FRigVMRegister& Register = Registers[InOperand.GetRegisterIndex()];
+
+	int32 SliceCount = Register.SliceCount;
+	
+	if (Register.IsNestedDynamic())
+	{
+		FRigVMNestedByteArray* Storage = (FRigVMNestedByteArray*)&Data[Register.GetWorkByteIndex()];
+		SliceCount = Storage->Num();
+	}
+	else if (Register.IsDynamic())
+	{
+		FRigVMByteArray* Storage = (FRigVMByteArray*)&Data[Register.GetWorkByteIndex()];
+		SliceCount = Storage->Num() / Register.ElementSize;
+	}
+
+	for (int32 SliceIndex = 0; SliceIndex < SliceCount; SliceIndex++)
+	{
+		int32 ElementCount = Register.ElementCount;
+		if (Register.IsNestedDynamic())
+		{
+			FRigVMNestedByteArray* Storage = (FRigVMNestedByteArray*)&Data[Register.GetWorkByteIndex()];
+			ElementCount = (*Storage)[SliceIndex].Num() / Register.ElementSize;
+		}
+
+		for (int32 ElementIndex = 0; ElementIndex < ElementCount; ElementIndex++)
+		{
+			FString DefaultValue;
+
+			if (Register.ScriptStructIndex != INDEX_NONE)
+			{
+				UScriptStruct* ScriptStruct = GetScriptStruct(Register);
+				if (ScriptStruct == InCPPTypeObject)
+				{
+					uint8* DataPtr = (uint8*)GetData(Register, INDEX_NONE, SliceIndex);
+					DataPtr += ElementIndex * ScriptStruct->GetStructureSize(); 
+
+					//use DataPtr for both value and default value to force exporting complete string
+					ScriptStruct->ExportText(DefaultValue, DataPtr, DataPtr, nullptr, PPF_None, nullptr);
+				}
+			}
+			else if (const UEnum* Enum = Cast<const UEnum>(InCPPTypeObject))
+			{
+				DefaultValue = Enum->GetNameStringByValue((int64)GetFixedArray<uint8>(Register, INDEX_NONE, SliceIndex)[ElementIndex]);
+			}
+			else if (InCPPType == TEXT("bool") && Register.GetNumBytesPerSlice() == sizeof(bool))
+			{
+				DefaultValue = GetFixedArray<bool>(Register, INDEX_NONE, SliceIndex)[ElementIndex] ? TEXT("True") : TEXT("False");
+			}
+			else if (InCPPType == TEXT("int32") && Register.GetNumBytesPerSlice() == sizeof(int32))
+			{
+				DefaultValue = FString::FromInt(GetFixedArray<int32>(Register, INDEX_NONE, SliceIndex)[ElementIndex]);
+			}
+			else if (InCPPType == TEXT("float") && Register.GetNumBytesPerSlice() == sizeof(float))
+			{
+				float Value = GetFixedArray<float>(Register, INDEX_NONE, SliceIndex)[ElementIndex];
+				DefaultValue = FString::Printf(TEXT("%f"), Value);
+			}
+			else if (InCPPType == TEXT("double") && Register.GetNumBytesPerSlice() == sizeof(double))
+			{
+				double Value = GetFixedArray<double>(Register, INDEX_NONE, SliceIndex)[ElementIndex];
+				DefaultValue = FString::Printf(TEXT("%f"), (float)Value);
+			}
+			else if (InCPPType == TEXT("FName") && Register.GetNumBytesPerSlice() == sizeof(FName))
+			{
+				DefaultValue = GetFixedArray<FName>(Register, INDEX_NONE, SliceIndex)[ElementIndex].ToString();
+			}
+			else if (InCPPType == TEXT("FString") && Register.GetNumBytesPerSlice() == sizeof(FString))
+			{
+				DefaultValue = GetFixedArray<FString>(Register, INDEX_NONE, SliceIndex)[ElementIndex];
+			}
+			else
+			{
+				continue;
+			}
+
+			DefaultValues.Add(DefaultValue);
+		}
+	}
+
+	return DefaultValues;
+}
+
+const FRigVMRegisterOffset& FRigVMMemoryContainer::GetRegisterOffset(int32 InRegisterOffsetIndex) const
+{
+	if (RegisterOffsets.IsValidIndex(InRegisterOffsetIndex))
+	{
+		return RegisterOffsets[InRegisterOffsetIndex];
+	}
+
+	return InvalidRegisterOffset;
+}
+
+const FRigVMRegisterOffset& FRigVMMemoryContainer::GetRegisterOffsetForOperand(const FRigVMOperand& InOperand) const
+{
+	ensure(InOperand.GetMemoryType() == MemoryType || 
+		(InOperand.GetMemoryType() == ERigVMMemoryType::External && MemoryType == ERigVMMemoryType::Work));
+	
+	return GetRegisterOffset(InOperand.GetRegisterOffset());
+}
+
+void FRigVMMemoryContainer::UpdateRegisters()
+{
+	int32 AlignmentShift = 0;
+	for (int32 RegisterIndex = 0; RegisterIndex < Registers.Num(); RegisterIndex++)
+	{
+		FRigVMRegister& Register = Registers[RegisterIndex];
+		Register.ByteIndex += AlignmentShift;
+
+		int32 Alignment = 4;
+
+		if (Register.IsDynamic() ||
+			Register.Type == ERigVMRegisterType::Name ||
+			Register.Type == ERigVMRegisterType::String)
+		{
+			Alignment = 8;
+		}
+		else if (UScriptStruct* ScriptStruct = GetScriptStruct(RegisterIndex))
+		{
+			if (UScriptStruct::ICppStructOps* TheCppStructOps = ScriptStruct->GetCppStructOps())
+			{
+				Alignment = TheCppStructOps->GetAlignment();
+			}
+		}
+
+		if (Alignment != 0)
+		{
+			// no need to adjust for alignment if nothing is allocated
+			if (!Register.IsDynamic() && Register.ElementCount == 0)
+			{
+				continue;
+			}
+
+			ensure(Data.IsValidIndex(Register.GetWorkByteIndex()));
+
+
+			if (ensure(Data.IsValidIndex(Register.GetWorkByteIndex())))
+			{
+				uint8* Pointer = (uint8*)(&(Data[Register.GetWorkByteIndex()]));
+
+				if (Register.AlignmentBytes > 0)
+				{
+					if (!IsAligned(Pointer, Alignment))
+					{
+						Data.RemoveAt(Register.GetFirstAllocatedByte(), Register.AlignmentBytes);
+						AlignmentShift -= Register.AlignmentBytes;
+						Register.ByteIndex -= Register.AlignmentBytes;
+						Register.AlignmentBytes = 0;
+						Pointer = (uint8*)(&(Data[Register.GetWorkByteIndex()]));
+					}
+				}
+
+				while (!IsAligned(Pointer, Alignment))
+				{
+					Data.InsertZeroed(Register.GetFirstAllocatedByte(), 1);
+					Register.AlignmentBytes++;
+					Register.ByteIndex++;
+					AlignmentShift++;
+					Pointer = (uint8*)(&(Data[Register.GetWorkByteIndex()]));
+				}
+			}
+		}
+	}
+
+	if (bUseNameMap)
+	{
+		NameMap.Reset();
+		for (int32 Index = 0; Index < Registers.Num(); Index++)
+		{
+			NameMap.Add(Registers[Index].Name, Index);
+		}
+	}
+}
+
+void FRigVMMemoryContainer::FillWithZeroes(int32 InRegisterIndex)
+{
+	ensure(Registers.IsValidIndex(InRegisterIndex));
+	FMemory::Memzero(GetData(InRegisterIndex), Registers[InRegisterIndex].GetNumBytesAllSlices());
+}
+
+int32 FRigVMMemoryContainer::FindOrAddScriptStruct(UScriptStruct* InScriptStruct)
+{
+	if (InScriptStruct == nullptr)
+	{
+		return INDEX_NONE;
+	}
+
+	int32 StructIndex = INDEX_NONE;
+	if (ScriptStructs.Find(InScriptStruct, StructIndex))
+	{
+		return StructIndex;
+	}
+	return ScriptStructs.Add(InScriptStruct);
+}
+
+#endif
\ No newline at end of file
diff --git a/Engine/Source/Runtime/RigVM/Private/RigVMCore/RigVMMemoryStorage.cpp b/Engine/Source/Runtime/RigVM/Private/RigVMCore/RigVMMemoryStorage.cpp
new file mode 100644
--- /dev/null
+++ b/Engine/Source/Runtime/RigVM/Private/RigVMCore/RigVMMemoryStorage.cpp
@@ -0,0 +1,34 @@
+// Copyright Epic Games, Inc. All Rights Reserved.
+
+#include "RigVMCore/RigVMMemoryStorage.h"
+
+////////////////////////////////////////////////////////////////////////////////////////////////////////////////////////
+
+URigVMMemoryStorage::FPropertyDescription::FPropertyDescription(const FProperty* InProperty, const FString& InDefaultValue, const FName& InName)
+	: Name(InName)
+	, Property(InProperty)
+	, PinType()
+	, DefaultValue(InDefaultValue)
+{
+}
+
+URigVMMemoryStorage::FPropertyDescription::FPropertyDescription(const FName& InName, const FEdGraphPinType& InPinType, const FString& InDefaultValue)
+	: Name(InName)
+	, Property(nullptr)
+	, PinType(InPinType)
+	, DefaultValue(InDefaultValue)
+{
+}
+
+////////////////////////////////////////////////////////////////////////////////////////////////////////////////////////
+
+UClass* URigVMMemoryStorage::CreateStorageClass(UObject* InOuter, const TArray<FPropertyDescription>& InProperties)
+{
+	return nullptr;
+}
+
+URigVMMemoryStorage* URigVMMemoryStorage::CreateStorage(UObject* InOuter,
+	const TArray<FPropertyDescription>& InProperties)
+{
+	return nullptr;
+}
\ No newline at end of file
diff --git a/Engine/Source/Runtime/RigVM/Public/RigVMCore/RigVMArray.h b/Engine/Source/Runtime/RigVM/Public/RigVMCore/RigVMArray.h
--- a/Engine/Source/Runtime/RigVM/Public/RigVMCore/RigVMArray.h
+++ b/Engine/Source/Runtime/RigVM/Public/RigVMCore/RigVMArray.h
@@ -5,6 +5,8 @@
 #include "CoreMinimal.h"
 #include "RigVMTraits.h"
 
+#if UE_RIGVM_UCLASS_BASED_STORAGE_DISABLED
+
 typedef TArray<uint8> FRigVMByteArray;
 typedef TArray<FRigVMByteArray> FRigVMNestedByteArray;
 
@@ -457,3 +459,4 @@ private:
 	int32 Size;
 };
 
+#endif
\ No newline at end of file
diff --git a/Engine/Source/Runtime/RigVM/Public/RigVMCore/RigVMByteCode.h b/Engine/Source/Runtime/RigVM/Public/RigVMCore/RigVMByteCode.h
--- a/Engine/Source/Runtime/RigVM/Public/RigVMCore/RigVMByteCode.h
+++ b/Engine/Source/Runtime/RigVM/Public/RigVMCore/RigVMByteCode.h
@@ -663,6 +663,8 @@ public:
 		return GetOperandsAt(ByteCodeIndex, ExecuteOp.GetOperandCount());
 	}
 
+#if UE_RIGVM_UCLASS_BASED_STORAGE_DISABLED
+
 	// returns the raw data of the byte code
 	FORCEINLINE const FRigVMFixedArray<uint8> GetByteCode() const
 	{
@@ -670,6 +672,17 @@ public:
 		return FRigVMFixedArray<uint8>((uint8*)Data, ByteCode.Num());
 	}
 
+#else
+
+	// returns the raw data of the byte code
+	FORCEINLINE const TArrayView<const uint8> GetByteCode() const
+	{
+		const uint8* Data = ByteCode.GetData();
+		return TArrayView<const uint8>((uint8*)Data, ByteCode.Num());
+	}
+
+#endif
+
 	// returns the statistics information
 	FRigVMByteCodeStatistics GetStatistics() const
 	{
diff --git a/Engine/Source/Runtime/RigVM/Public/RigVMCore/RigVMExecuteContext.h b/Engine/Source/Runtime/RigVM/Public/RigVMCore/RigVMExecuteContext.h
--- a/Engine/Source/Runtime/RigVM/Public/RigVMCore/RigVMExecuteContext.h
+++ b/Engine/Source/Runtime/RigVM/Public/RigVMCore/RigVMExecuteContext.h
@@ -245,7 +245,11 @@ struct RIGVM_API FRigVMExecuteContext
 	FName EventName;
 	FName FunctionName;
 	uint16 InstructionIndex;
+#if UE_RIGVM_UCLASS_BASED_STORAGE_DISABLED
 	FRigVMFixedArray<void*> OpaqueArguments;
+#else
+	TArrayView<void*> OpaqueArguments;
+#endif
 	TArray<FRigVMExternalVariable> ExternalVariables;
 	TArray<FRigVMSlice> Slices;
 	TArray<uint16> SliceOffsets;
diff --git a/Engine/Source/Runtime/RigVM/Public/RigVMCore/RigVMExternalVariable.h b/Engine/Source/Runtime/RigVM/Public/RigVMCore/RigVMExternalVariable.h
--- a/Engine/Source/Runtime/RigVM/Public/RigVMCore/RigVMExternalVariable.h
+++ b/Engine/Source/Runtime/RigVM/Public/RigVMCore/RigVMExternalVariable.h
@@ -450,6 +450,8 @@ struct RIGVM_API FRigVMExternalVariable
 			(bAllowNullPtr || Memory != nullptr);
 	}
 
+#if UE_RIGVM_UCLASS_BASED_STORAGE_DISABLED
+
 	FORCEINLINE FRigVMMemoryHandle GetHandle() const
 	{
 		return FRigVMMemoryHandle(
@@ -459,6 +461,8 @@ struct RIGVM_API FRigVMExternalVariable
 		);
 	}
 
+#endif
+
 	FORCEINLINE static void MergeExternalVariable(TArray<FRigVMExternalVariable>& OutVariables, const FRigVMExternalVariable& InVariable)
 	{
 		if(!InVariable.IsValid(true))
diff --git a/Engine/Source/Runtime/RigVM/Public/RigVMCore/RigVMFunction.h b/Engine/Source/Runtime/RigVM/Public/RigVMCore/RigVMFunction.h
--- a/Engine/Source/Runtime/RigVM/Public/RigVMCore/RigVMFunction.h
+++ b/Engine/Source/Runtime/RigVM/Public/RigVMCore/RigVMFunction.h
@@ -8,8 +8,18 @@
 #include "RigVMMemory.h"
 #include "Blueprint/BlueprintSupport.h"
 
+#if UE_RIGVM_UCLASS_BASED_STORAGE_DISABLED
+
 typedef FRigVMFixedArray<FRigVMMemoryHandle> FRigVMMemoryHandleArray;
 typedef FRigVMFixedArray<void*> FRigVMUserDataArray;
+
+#else
+
+typedef TArrayView<FRigVMMemoryHandle> FRigVMMemoryHandleArray;
+typedef TArrayView<void*> FRigVMUserDataArray;
+
+#endif
+
 typedef void (*FRigVMFunctionPtr)(FRigVMExecuteContext& RigVMExecuteContext, FRigVMMemoryHandleArray RigVMMemoryHandles);
 
 /**
diff --git a/Engine/Source/Runtime/RigVM/Public/RigVMCore/RigVMMemory.h b/Engine/Source/Runtime/RigVM/Public/RigVMCore/RigVMMemory.h
--- a/Engine/Source/Runtime/RigVM/Public/RigVMCore/RigVMMemory.h
+++ b/Engine/Source/Runtime/RigVM/Public/RigVMCore/RigVMMemory.h
@@ -2,1484 +2,14 @@
 
 #pragma once
 
-#include "CoreMinimal.h"
-#include "UObject/ObjectMacros.h"
-#include "RigVMTraits.h"
-#include "RigVMStatistics.h"
-#include "RigVMArray.h"
-#include "RigVMMemory.generated.h"
+#include "RigVMDefines.h"
 
-#ifdef UE_BUILD_DEBUG
-	#define DEBUG_RIGVMMEMORY 0
-	//#define DEBUG_RIGVMMEMORY WITH_EDITOR
-#else
-	#define DEBUG_RIGVMMEMORY 0
-#endif
-
-#if DEBUG_RIGVMMEMORY
-	RIGVM_API DECLARE_LOG_CATEGORY_EXTERN(LogRigVMMemory, Log, All);
-	#define UE_LOG_RIGVMMEMORY(Format, ...) UE_LOG(LogRigVMMemory, Display, (Format), ##__VA_ARGS__)
-#else
-	#define UE_LOG_RIGVMMEMORY(Format, ...)
-#endif
-
-// The ERigVMMemoryType maps to memory container index in RigVM through
-// FRigVMOperand::GetContainerIndex() or URigVM::GetContainerIndex(...) 
-/**
- * The type of memory used. Typically we differentiate between
- * Work (Mutable) and Literal (Constant) memory.
- */
-UENUM()
-enum class ERigVMMemoryType: uint8
-{
-	Work = 0, // Mutable state
-	Literal = 1, // Const / fixed state
-	External = 2, // Unowned external memory
-	Debug = 3, // Owned memory used for debug watches
-	Invalid
-};
-
-/**
- * The FRigVMOperand represents an argument used for an operator
- * within the virtual machine. Operands provide information about
- * which memory needs to be referred to, which register within the
- * memory all the way to the actual byte address in memory.
- * The FRigVMOperand is a light weight address for a register in
- * a FRigVMMemoryContainer.
- * For external variables the register index represents the
- * index of the external variable within the running VM.
- */
-USTRUCT()
-struct RIGVM_API FRigVMOperand
-{
-	GENERATED_BODY()
-
-public:
-
-	FRigVMOperand()
-		: MemoryType(ERigVMMemoryType::Work)
-		, RegisterIndex(UINT16_MAX)
-		, RegisterOffset(UINT16_MAX)
-	{
-	}
-
-	FRigVMOperand(ERigVMMemoryType InMemoryType, int32 InRegisterIndex, int32 InRegisterOffset = INDEX_NONE)
-		: MemoryType(InMemoryType)
-		, RegisterIndex(InRegisterIndex < 0 ? UINT16_MAX : (uint16)InRegisterIndex)
-		, RegisterOffset(InRegisterOffset < 0 ? UINT16_MAX : (uint16)InRegisterOffset)
-	{
-	}
-
-	FORCEINLINE_DEBUGGABLE bool operator == (const FRigVMOperand& InOther) const
-	{
-		return MemoryType == InOther.MemoryType &&
-			RegisterIndex == InOther.RegisterIndex &&
-			RegisterOffset == InOther.RegisterOffset;
-	}
-
-	FORCEINLINE_DEBUGGABLE bool operator != (const FRigVMOperand& InOther) const
-	{
-		return !(*this == InOther);
-	}
-
-	// returns the memory type of this argument
-	FORCEINLINE_DEBUGGABLE bool IsValid() const { return RegisterIndex != UINT16_MAX; }
-
-	// returns the memory type of this argument
-	FORCEINLINE_DEBUGGABLE ERigVMMemoryType GetMemoryType() const { return MemoryType; }
-
-	// returns the index of the container of this argument
-	// this function should be kept in sync with URigVM::GetContainerIndex()
-	FORCEINLINE_DEBUGGABLE int32 GetContainerIndex() const
-	{
-		if(MemoryType == ERigVMMemoryType::External)
-		{
-			return (int32)ERigVMMemoryType::Work;
-		}
-		
-		if(MemoryType == ERigVMMemoryType::Debug)
-		{
-			return 2;
-		}
-		return (int32)MemoryType;
-	}
-
-	// returns the index of the register of this argument
-	FORCEINLINE_DEBUGGABLE int32 GetRegisterIndex() const { return RegisterIndex == UINT16_MAX ? INDEX_NONE : (int32)RegisterIndex; }
-
-	// returns the register offset of this argument
-	FORCEINLINE_DEBUGGABLE int32 GetRegisterOffset() const { return RegisterOffset == UINT16_MAX ? INDEX_NONE : (int32)RegisterOffset; }
-
-	friend FORCEINLINE uint32 GetTypeHash(const FRigVMOperand& Operand)
-	{
-  		return HashCombine(HashCombine((uint32)Operand.MemoryType, (uint32)Operand.RegisterIndex), (uint32)Operand.RegisterOffset);
-  	}
-	
-	void Serialize(FArchive& Ar);
-	void Save(FArchive& Ar) const;
-	void Load(FArchive& Ar);
-	FORCEINLINE friend FArchive& operator<<(FArchive& Ar, FRigVMOperand& P)
-	{
-		P.Serialize(Ar);
-		return Ar;
-	}
-
-private:
-
-	UPROPERTY()
-	ERigVMMemoryType MemoryType;
-
-	/**
-	 * The index of the register inside of the specific type of memory (work, literal etc).
-	 * For external variables the register index represents the index of the external variable within the running VM.
-	 */
-	UPROPERTY()
-	uint16 RegisterIndex;
-	
-	UPROPERTY()
-	uint16 RegisterOffset;
-
-	friend class URigVMCompiler;
-};
-
-typedef FRigVMFixedArray<FRigVMOperand> FRigVMOperandArray;
-
-/**
- * The type of register within the memory.
- */
-UENUM()
-enum class ERigVMRegisterType : uint8
-{
-	Plain, // bool, int32, float, FVector etc. (also structs that do NOT require a constructor / destructor to be valid)
-	String, // FString
-	Name, // FName
-	Struct, // Any USTRUCT (structs which require a constructor / destructor to be valid. structs which have indirection (arrays or strings in them, for example).
-	Invalid
-};
-
-// The register represents an address with the VM's memory. Within a register
-// we can store arbitrary data, so it provides a series of properties to
-// describe the memory location.
-// Registers also support the notion of slices. A slice is a complete copy of
-// the memory - so for example if your register stores 4 Vectors, then a slice 
-// would contain 48 bytes (4 * 3 * 4). The register can however store multiple
-// slices / copies of that if needed. Slices can be used to provide 
-// per-invocation memory to functions within the same register.
-// An integrator for example that needs to store a simulated position
-// might want access to a separate memory per loop iteration.
-USTRUCT()
-struct RIGVM_API FRigVMRegister
-{
-	GENERATED_BODY()
-
-	FRigVMRegister()
-		: Type(ERigVMRegisterType::Invalid)
-		, ByteIndex(INDEX_NONE)
-		, ElementSize(0)
-		, ElementCount(0)
-		, SliceCount(1)
-		, AlignmentBytes(0)
-		, TrailingBytes(0)
-		, Name(NAME_None)
-		, ScriptStructIndex(INDEX_NONE)
-		, bIsArray(false)
-		, bIsDynamic(false)
-#if WITH_EDITORONLY_DATA
-		, BaseCPPType(NAME_None)
-		, BaseCPPTypeObject(nullptr)
-#endif
-	{
-	}
-
-	// The type of register (plain, name, string, etc.)
-	UPROPERTY()
-	ERigVMRegisterType Type;
-
-	// The index of the first work byte
-	UPROPERTY()
-	uint32 ByteIndex;
-
-	// The size of each store element (for example 4 for a float)
-	UPROPERTY()
-	uint16 ElementSize;
+#if UE_RIGVM_UCLASS_BASED_STORAGE_DISABLED
 
-	// The number of elements in this register (for example the number of elements in an array)
-	UPROPERTY()
-	uint16 ElementCount;
+#include "RigVMMemoryDeprecated.h"
 
-	// The number of slices (complete copies) (for example the number of iterations on a fixed loop)
-	// Potentially redundant state - can be removed.
-	UPROPERTY()
-	uint16 SliceCount;
-
-	// The number of leading bytes for alignment.
-	// Bytes that had to be introduced before the register
-	// to align the registers memory as per platform specification.
-	UPROPERTY()
-	uint8 AlignmentBytes;
-
-	// The number of trailing bytes.
-	// These originate after shrinking a register.
-	// Potentially unused state - can be removed.
-	UPROPERTY()
-	uint16 TrailingBytes;
-
-	// The name of the register (can be None)
-	UPROPERTY()
-	FName Name;
-
-	// For struct registers this is the index of the
-	// struct used - otherwise INDEX_NONE
-	UPROPERTY()
-	int32 ScriptStructIndex;
-
-	// If true defines this register as an array
-	UPROPERTY()
-	bool bIsArray;
-
-	// If true defines this register to use dynamic storage.
-	// Is this an array or singleton value with multiple slices
-	// which potentially requires changing count at runtime. 
-	UPROPERTY()
-	bool bIsDynamic;
-
-#if WITH_EDITORONLY_DATA
-
-	// Defines the CPP type used for the register
-	// This is only used for debugging purposes in editor.
-	UPROPERTY(transient)
-	FName BaseCPPType;
+#else
 
-	// The resolved UScriptStruct / UObject (in the future)
-	// used for debugging.
-	UPROPERTY(transient)
-	TObjectPtr<UObject> BaseCPPTypeObject;
+#include "RigVMMemoryStorage.h"
 
 #endif
-
-	void Serialize(FArchive& Ar);
-	void Save(FArchive& Ar);
-	void Load(FArchive& Ar);
-	FORCEINLINE_DEBUGGABLE friend FArchive& operator<<(FArchive& Ar, FRigVMRegister& P)
-	{
-		P.Serialize(Ar);
-		return Ar;
-	}
-
-	// returns true if this register is using a dynamic array for storage
-	FORCEINLINE_DEBUGGABLE bool IsDynamic() const { return bIsDynamic; }
-
-	// returns true if this register is using a dynamic array for storage
-	FORCEINLINE_DEBUGGABLE bool IsNestedDynamic() const { return bIsDynamic && bIsArray; }
-
-	// returns the current address of the register within the data byte array.
-	// this can change over time - as the register is moving through slices.
-	// use GetFirstAllocatedByte to get the fixed first byte.
-	FORCEINLINE_DEBUGGABLE uint64 GetWorkByteIndex(int32 InSliceIndex = 0) const
-	{
-		ensure(InSliceIndex >= 0);
-		return ByteIndex + ((uint64)InSliceIndex * GetNumBytesPerSlice());
-	}
-
-	// returns the first allocated byte in the data byte array
-	FORCEINLINE_DEBUGGABLE uint64 GetFirstAllocatedByte() const
-	{ 
-		return ByteIndex - (uint64)AlignmentBytes;
-	}
-
-	// Returns the leading alignment bytes
-	FORCEINLINE_DEBUGGABLE uint8 GetAlignmentBytes() const { return AlignmentBytes; }
-
-	// Returns true if the register stores more than one element
-	FORCEINLINE_DEBUGGABLE bool IsArray() const { return bIsArray || (ElementCount > 1); }
-
-	// Returns the number of allocated bytes (including alignment + trailing bytes)
-	FORCEINLINE_DEBUGGABLE uint16 GetAllocatedBytes() const { return ElementCount * ElementSize * SliceCount + (uint16)AlignmentBytes + TrailingBytes; }
-
-	// Returns the number of bytes for a complete slice
-	FORCEINLINE_DEBUGGABLE uint16 GetNumBytesPerSlice() const { return ElementCount * ElementSize; }
-
-	// Returns the number of bytes for all slices
-	FORCEINLINE_DEBUGGABLE uint16 GetNumBytesAllSlices() const { return ElementCount * ElementSize * SliceCount; }
-
-	// Returns the total number of elements (elementcount * slicecount) in the register
-	FORCEINLINE_DEBUGGABLE uint32 GetTotalElementCount() const { return (uint32)ElementCount * (uint32)SliceCount; }
-
-};
-
-typedef FRigVMFixedArray<FRigVMRegister> FRigVMRegisterArray;
-
-// The register offset represents a memory offset within a register's memory.
-// This can be used to represent memory addresses of array elements within
-// a struct, for example.
-// A register offset's path can look like MyTransformStruct.Transforms[3].Translation.X
-USTRUCT()
-struct RIGVM_API FRigVMRegisterOffset
-{
-	GENERATED_BODY()
-
-public:
-
-	// default constructor
-	FRigVMRegisterOffset()
-		: Segments()
-		, Type(ERigVMRegisterType::Invalid)
-		, CPPType(NAME_None)
-		, ScriptStruct(nullptr)
-		, ParentScriptStruct(nullptr)
-		, ArrayIndex(0)
-		, ElementSize(0)
-		, CachedSegmentPath() 
-	{
-	}
-
-	void Serialize(FArchive& Ar);
-	void Save(FArchive& Ar);
-	void Load(FArchive& Ar);
-	FORCEINLINE_DEBUGGABLE friend FArchive& operator<<(FArchive& Ar, FRigVMRegisterOffset& P)
-	{
-		P.Serialize(Ar);
-		return Ar;
-	}
-
-	// constructs a path given a struct and a segment path
-	FRigVMRegisterOffset(UScriptStruct* InScriptStruct, const FString& InSegmentPath, int32 InInitialOffset = 0, uint16 InElementSize = 0, const FName& InCPPType = NAME_None);
-
-	// returns the data pointer within a container
-	uint8* GetData(uint8* InContainer) const;
-
-	// returns the segments of this path
-	const TArray<int32>& GetSegments() const { return Segments; }
-
-	// returns true if this register offset contains an array segment.
-	// this means the register offset has to traverse a pointer and
-	// cannot be cached / collapsed easily
-	bool ContainsArraySegment() const;
-
-	bool operator == (const FRigVMRegisterOffset& InOther) const;
-
-	FORCEINLINE_DEBUGGABLE bool IsValid() const { return Type != ERigVMRegisterType::Invalid; }
-	FORCEINLINE_DEBUGGABLE ERigVMRegisterType GetType() const { return Type; }
-	FORCEINLINE_DEBUGGABLE FName GetCPPType() const { return CPPType; }
-	FORCEINLINE_DEBUGGABLE FString GetCachedSegmentPath() const { return CachedSegmentPath; }
-	FORCEINLINE_DEBUGGABLE int32 GetArrayIndex() const { return ArrayIndex; }
-	uint16 GetElementSize() const;
-	void SetElementSize(uint16 InElementSize) { ElementSize = InElementSize; };
-	UScriptStruct* GetScriptStruct() const;
-
-private:
-
-	// Segments represent the memory offset(s) to use when accessing the target memory.
-	// In case of indirection in the source memory each memory offset / jump is stored.
-	// So for example: If you are accessing the third Rotation.X of an array of transforms,
-	// the segments would read as: [-2, 12] (second array element, and the fourth float)
-	// Segment indices less than zero represent array element offsets, while positive numbers
-	// represents jumps within a struct.
-	UPROPERTY()
-	TArray<int32> Segments;
-
-	// Type of resulting register (for example Plain for Transform.Translation.X)
-	UPROPERTY()
-	ERigVMRegisterType Type;
-
-	// The c++ type of the resulting memory (for example float for Transform.Translation.X)
-	UPROPERTY()
-	FName CPPType;
-
-	// The c++ script struct of the resulting memory (for example FVector for Transform.Translation)
-	UPROPERTY()
-	TObjectPtr<UScriptStruct> ScriptStruct;
-
-	// The c++ script struct of the source memory (for example FTransform for Transform.Translation)
-	UPROPERTY()
-	TObjectPtr<UScriptStruct> ParentScriptStruct;
-
-	// The index of the element within an array (for example 3 for Transform[3])
-	UPROPERTY()
-	int32 ArrayIndex;
-
-	// The number of bytes of the resulting memory (for example 4 (float) for Transform.Translation.X)
-	UPROPERTY()
-	uint16 ElementSize;
-
-	// The cached path of the segments within this register, for example FTransform.Translation.X
-	UPROPERTY()
-	FString CachedSegmentPath;
-
-	friend struct FRigVMRegisterOffsetBuilder;
-	friend class URigVM;
-};
-
-/**
- * The FRigVMMemoryHandle is used to access the memory used within a FRigMemoryContainer.
- */
-struct FRigVMMemoryHandle
-{
-public:
-
-	// The type of memory being accessed.
-	enum FType
-	{
-		Plain, // Shallow memory with no indirection (for example FTransform.Translation.X)
-		Dynamic, // Memory representing an array
-		NestedDynamic, // Memory representing an array or single element with slices (for example within loop)
-		ArraySize // a count represented as a pointer address
-	};
-
-	FORCEINLINE_DEBUGGABLE FRigVMMemoryHandle()
-		: Ptr(nullptr)
-		, Type(FType::Plain)
-		, Size(1)
-		, RegisterOffset(nullptr)
-	{}
-
-	FORCEINLINE_DEBUGGABLE FRigVMMemoryHandle(uint8* InPtr, uint16 InSize = 1, FType InType = FType::Plain, const FRigVMRegisterOffset* InRegisterOffset = nullptr)
-		: Ptr(InPtr)
-		, Type(InType)
-		, Size(InSize)
-		, RegisterOffset(InRegisterOffset)
-	{}
-
-	FORCEINLINE_DEBUGGABLE FRigVMMemoryHandle(uint8* InPtr, const FRigVMRegister& InRegister, const FRigVMRegisterOffset* InRegisterOffset = nullptr)
-		: Ptr(InPtr)
-		, Type(FType::Plain)
-		, Size(InRegister.ElementSize)
-		, RegisterOffset(InRegisterOffset)
-	{
-		if (InRegister.IsNestedDynamic())
-		{
-			Type = FType::NestedDynamic;
-		}
-		else if (InRegister.IsDynamic())
-		{
-			Type = FType::Dynamic;
-		}
-		else
-		{
-			Size = (uint16)InRegister.GetNumBytesPerSlice();
-		}
-	}
-
-	FORCEINLINE_DEBUGGABLE FRigVMMemoryHandle(FRigVMByteArray* InPtr, uint16 InSize = 1, const FRigVMRegisterOffset* InRegisterOffset = nullptr)
-		: Ptr((uint8*)InPtr)
-		, Type(FType::Dynamic)
-		, Size(InSize)
-		, RegisterOffset(InRegisterOffset)
-	{}
-
-	FORCEINLINE_DEBUGGABLE FRigVMMemoryHandle(FRigVMNestedByteArray* InPtr, uint16 InSize = 1, const FRigVMRegisterOffset* InRegisterOffset = nullptr)
-		: Ptr((uint8*)InPtr)
-		, Type(FType::NestedDynamic)
-		, Size(InSize)
-		, RegisterOffset(InRegisterOffset)
-	{}
-
-	// getter for the contained const memory
-	FORCEINLINE_DEBUGGABLE operator const uint8*() const
-	{
-		return GetData();
-	}
-
-	// getter for the contained mutable memory
-	FORCEINLINE_DEBUGGABLE operator uint8*()
-	{
-		return GetData();
-	}
-
-	// returns the contained const memory for a given slice index
-	FORCEINLINE_DEBUGGABLE const uint8* GetData(int32 SliceIndex = 0, bool bGetArrayData = false) const
-	{
-		return GetData_Internal(SliceIndex, bGetArrayData);
-	}
-
-	// returns the contained mutable memory for a given slice index
-	FORCEINLINE_DEBUGGABLE uint8* GetData(int32 SliceIndex = 0, bool bGetArrayData = false)
-	{
-		return GetData_Internal(SliceIndex, bGetArrayData);
-	}
-
-	FORCEINLINE FType GetType() const
-	{
-		return Type;
-	}
-
-	FORCEINLINE bool IsDynamic() const
-	{
-		return (Type == Dynamic) || (Type == NestedDynamic);
-	}
-
-private:
-
-	FORCEINLINE_DEBUGGABLE uint8* GetData_Internal_NoOffset(int32 SliceIndex, bool bGetArrayData = false) const
-	{
-		ensure(SliceIndex >= 0);
-
-		if (Size == 0 || Ptr == nullptr)
-		{
-			return nullptr;
-		}
-
-		switch (Type)
-		{
-			case FType::Plain:
-			{
-				return Ptr + SliceIndex * Size;
-			}
-			case FType::ArraySize:
-			{
-				return Ptr;
-			}
-			case FType::Dynamic:
-			{
-				ensure(RegisterOffset == nullptr);
-				if (!bGetArrayData)
-				{
-					return Ptr;
-				}
-
-				FRigVMByteArray* ArrayStorage = (FRigVMByteArray*)Ptr;
-				if (ArrayStorage->Num() == 0)
-				{
-					return nullptr;
-				}
-
-				if (SliceIndex > 0)
-				{
-					SliceIndex = FMath::Max(SliceIndex, (ArrayStorage->Num() / Size) - 1);
-				}
-				return (uint8*)&(*ArrayStorage)[SliceIndex * Size];
-			}
-			case FType::NestedDynamic:
-			{
-				ensure(RegisterOffset == nullptr);
-
-				FRigVMNestedByteArray* ArrayStorage = (FRigVMNestedByteArray*)Ptr;
-				if (ArrayStorage->Num() == 0 && bGetArrayData)
-				{
-					return nullptr;
-				}
-
-				if (SliceIndex > 0)
-				{
-					SliceIndex = FMath::Max(SliceIndex, ArrayStorage->Num() - 1);
-				}
-
-				if (!bGetArrayData)
-				{
-					return Ptr;
-				}
-
-				return (*ArrayStorage)[SliceIndex].GetData();
-			}
-			default:
-			{
-				return nullptr;
-			}
-		}
-	}
-
-	FORCEINLINE_DEBUGGABLE uint8* GetData_Internal(int32 SliceIndex, bool bGetArrayData = false) const
-	{
-		uint8* Result = GetData_Internal_NoOffset(SliceIndex, bGetArrayData);
-		if (Result && RegisterOffset)
-		{
-			Result = RegisterOffset->GetData(Result);
-		}
-		return Result;
-	}
-
-	uint8* Ptr;
-	FType Type;
-	uint16 Size;
-	const FRigVMRegisterOffset* RegisterOffset;
-
-	friend class URigVM;
-};
-
-/**
- * The FRigVMMemoryContainer provides a heterogeneous memory container to store arbitrary
- * data. Each element stored can be referred to using a FRigVMRegister.
- * Elements can be accessed by index (index of the register), FRigVMOperand or by name.
- * Name access is optional and is specified upon construction of the container.
- * The memory container provides a series of templated functions to add and get data.
- *
- * For example:
- * 		int32 Index = Container.Add<float>(4.f);
- *      float& ValueRef = Container.GetRef<float>(Index);
- *
- * This can also be done with arrays:
- *      TArray<float> MyArray = {3.f, 4.f, 5.f};
- * 		int32 Index = Container.AddFixedArray<float>(MyArray);
- *      FRigVMFixedArray<float> ArrayView = Container.GetFixedArray<float>(Index);
- *
- *
- * Registers can store dynamically resizable memory as well by reyling on indirection.
- * 
- * Arrays with a single slice are going to be stored as a FRigVMByteArray, [unit8], so for example for a float array in C++ it would be (TArray<float> Param).
- * Single values with multiple slices are also going to be stored as a FRigVMByteArray, [unit8], so for example for a float array in C++ it would be (float Param). 
- * Arrays with multiple slices are going to be stored as a FRigVMNestedByteArray, [[unit8]], so for example for a float array in C++ it would be (TArray<float> Param).
- */
-USTRUCT()
-struct RIGVM_API FRigVMMemoryContainer
-{
-	GENERATED_BODY()
-	
-public:
-
-	FRigVMMemoryContainer(bool bInUseNames = true);
-	FRigVMMemoryContainer(const FRigVMMemoryContainer& Other);
-	~FRigVMMemoryContainer();
-
-	bool CopyRegisters(const FRigVMMemoryContainer &InOther);
-	FRigVMMemoryContainer& operator= (const FRigVMMemoryContainer &InOther);
-
-	// returns the memory type of this container
-	FORCEINLINE_DEBUGGABLE ERigVMMemoryType GetMemoryType() const { return MemoryType;  }
-
-	// sets the memory type. should only be used when the container is empty
-	FORCEINLINE_DEBUGGABLE void SetMemoryType(ERigVMMemoryType InMemoryType) { MemoryType = InMemoryType; }
-
-	// returns true if this container supports name based lookup
-	FORCEINLINE_DEBUGGABLE bool SupportsNames() const { return bUseNameMap;  }
-
-	// returns the number of registers in this container
-	FORCEINLINE_DEBUGGABLE int32 Num() const { return Registers.Num(); }
-
-	// resets the container but maintains storage.
-	void Reset();
-
-	// resets the container and removes all storage.
-	void Empty();
-
-	// const accessor for a register based on index
-	FORCEINLINE_DEBUGGABLE const FRigVMRegister& operator[](int32 InIndex) const { return GetRegister(InIndex); }
-
-	// accessor for a register based on index
-	FORCEINLINE_DEBUGGABLE FRigVMRegister& operator[](int32 InIndex) { return GetRegister(InIndex); }
-
-	// const accessor for a register based on an argument
-	FORCEINLINE_DEBUGGABLE const FRigVMRegister& operator[](const FRigVMOperand& InArg) const { return GetRegister(InArg); }
-
-	// accessor for a register based on an argument
-	FORCEINLINE_DEBUGGABLE FRigVMRegister& operator[](const FRigVMOperand& InArg) { return GetRegister(InArg); }
-
-	// const accessor for a register based on a a name. note: only works if SupportsNames() == true
-	FORCEINLINE_DEBUGGABLE const FRigVMRegister& operator[](const FName& InName) const { return GetRegister(InName); }
-
-	// accessor for a register based on a a name. note: only works if SupportsNames() == true
-	FORCEINLINE_DEBUGGABLE FRigVMRegister& operator[](const FName& InName) { return GetRegister(InName); }
-
-	FORCEINLINE_DEBUGGABLE TArray<FRigVMRegister>::RangedForIteratorType      begin() { return Registers.begin(); }
-	FORCEINLINE_DEBUGGABLE TArray<FRigVMRegister>::RangedForConstIteratorType begin() const { return Registers.begin(); }
-	FORCEINLINE_DEBUGGABLE TArray<FRigVMRegister>::RangedForIteratorType      end() { return Registers.end(); }
-	FORCEINLINE_DEBUGGABLE TArray<FRigVMRegister>::RangedForConstIteratorType end() const { return Registers.end(); }
-
-	// const accessor for a register based on index
-	FORCEINLINE_DEBUGGABLE const FRigVMRegister& GetRegister(int32 InIndex) const { return Registers[InIndex]; }
-
-	// accessor for a register based on index
-	FORCEINLINE_DEBUGGABLE FRigVMRegister& GetRegister(int32 InIndex) { return Registers[InIndex]; }
-
-	// const accessor for a register based on an argument
-	FORCEINLINE_DEBUGGABLE const FRigVMRegister& GetRegister(const FRigVMOperand& InArg) const { return Registers[InArg.GetRegisterIndex()]; }
-
-	// accessor for a register based on an argument
-	FORCEINLINE_DEBUGGABLE FRigVMRegister& GetRegister(const FRigVMOperand& InArg) { return Registers[InArg.GetRegisterIndex()]; }
-
-	// const accessor for a register based on a a name. note: only works if SupportsNames() == true
-	FORCEINLINE_DEBUGGABLE const FRigVMRegister& GetRegister(const FName& InName) const { return Registers[GetIndex(InName)]; }
-
-	// accessor for a register based on a a name. note: only works if SupportsNames() == true
-	FORCEINLINE_DEBUGGABLE FRigVMRegister& GetRegister(const FName& InName) { return Registers[GetIndex(InName)]; }
-
-	void Serialize(FArchive& Ar);
-	void Save(FArchive& Ar);
-	void Load(FArchive& Ar);
-	FORCEINLINE_DEBUGGABLE friend FArchive& operator<<(FArchive& Ar, FRigVMMemoryContainer& P)
-	{
-		P.Serialize(Ar);
-		return Ar;
-	}
-
-	// Returns an argument for a given register.
-	// This is typically used to store a light weight address for use within a VM.
-	FORCEINLINE_DEBUGGABLE FRigVMOperand GetOperand(int32 InRegisterIndex, int32 InRegisterOffset)
-	{
-		ensure(Registers.IsValidIndex(InRegisterIndex));
-		return FRigVMOperand(MemoryType, InRegisterIndex, InRegisterOffset);
-	}
-
-	// Returns an argument for a given register.
-	// This is typically used to store a light weight address for use within a VM.
-	FORCEINLINE_DEBUGGABLE FRigVMOperand GetOperand(int32 InRegisterIndex, const FString& InSegmentPath = FString(), int32 InArrayElement = INDEX_NONE)
-	{
-		ensure(Registers.IsValidIndex(InRegisterIndex));
-		// Register offset must hold on to the ScriptStruct such that it can recalculate the struct size after cook
-		UScriptStruct* ScriptStruct = GetScriptStruct(InRegisterIndex);
-
-		int32 InitialOffset = 0;
-		int32 ElementSize = 0;
-		if (InArrayElement != INDEX_NONE)
-		{
-			InitialOffset = InArrayElement * Registers[InRegisterIndex].ElementSize;
-			ElementSize = Registers[InRegisterIndex].ElementSize;
-		}
-
-		return GetOperand(InRegisterIndex, GetOrAddRegisterOffset(InRegisterIndex, ScriptStruct, InSegmentPath, InitialOffset, ElementSize));
-	}
-
-private:
-
-	FORCEINLINE_DEBUGGABLE uint8* GetDataPtr(const FRigVMRegister& Register, int32 InRegisterOffset = INDEX_NONE, int32 InSliceIndex = 0, bool bArrayContent = false) const
-	{
-		if (Register.ElementCount == 0 && !Register.IsNestedDynamic())
-		{
-			return nullptr;
-		}
-
-		uint8* Ptr = nullptr;
-		if (Register.IsDynamic())
-		{
-			Ptr = (uint8*)&Data[Register.GetWorkByteIndex()];
-
-			if (Register.IsNestedDynamic())
-			{
-				FRigVMNestedByteArray* ArrayStorage = (FRigVMNestedByteArray*)Ptr;
-				Ptr = (uint8*)ArrayStorage->GetData();
-
-				if (Ptr)
-				{
-					Ptr = Ptr + InSliceIndex * sizeof(FRigVMByteArray);
-					if (Ptr && bArrayContent)
-					{
-						Ptr = ((FRigVMByteArray*)Ptr)->GetData();
-					}
-				}
-			}
-			else if(bArrayContent)
-			{
-				FRigVMByteArray* ArrayStorage = (FRigVMByteArray*)Ptr;
-				Ptr = (uint8*)ArrayStorage->GetData();
-
-				if (Ptr)
-				{
-					Ptr = Ptr + InSliceIndex * Register.GetNumBytesPerSlice();
-				}
-			}
-		}
-		else
-		{
-			Ptr = (uint8*)&Data[Register.GetWorkByteIndex(InSliceIndex)];
-		}
-
-		if (InRegisterOffset != INDEX_NONE && Ptr != nullptr)
-		{
-			Ptr = RegisterOffsets[InRegisterOffset].GetData(Ptr);
-		}
-		return Ptr;
-	}
-
-public:
-
-	// Returns a memory handle for a given register
-	FORCEINLINE_DEBUGGABLE FRigVMMemoryHandle GetHandle(const FRigVMRegister& Register, int32 InRegisterOffset = INDEX_NONE) const
-	{
-		const FRigVMRegisterOffset& RegisterOffset = GetRegisterOffset(InRegisterOffset);
-
-		if (Register.IsDynamic())
-		{
-			uint8* Ptr = (uint8*)&Data[Register.GetWorkByteIndex()];
-			return FRigVMMemoryHandle(Ptr, Register, RegisterOffset.IsValid() ? &RegisterOffset : nullptr);
-		}
-
-		uint8* Ptr = GetDataPtr(Register);
-		uint16 NumBytes = Register.GetNumBytesPerSlice();
-
-		if (RegisterOffset.IsValid())
-		{
-			NumBytes = RegisterOffset.GetElementSize();
-
-			if (!RegisterOffset.ContainsArraySegment())
-			{
-				// flatten the offset and cache the offset ptr directly
-				Ptr = RegisterOffset.GetData(Ptr);
-				return FRigVMMemoryHandle(Ptr, NumBytes);
-			}
-		}
-
-		return FRigVMMemoryHandle(Ptr, NumBytes, FRigVMMemoryHandle::FType::Plain, RegisterOffset.IsValid() ? &RegisterOffset : nullptr);
-	}
-
-	// Returns a memory handle for a given register
-	FORCEINLINE_DEBUGGABLE FRigVMMemoryHandle GetHandle(int32 InRegisterIndex, int32 InRegisterOffset = INDEX_NONE) const
-	{
-		ensure(Registers.IsValidIndex(InRegisterIndex));
-		const FRigVMRegister& Register = Registers[InRegisterIndex];
-		return GetHandle(Register, InRegisterOffset);
-	}
-
-	// Returns the number of elements for a given slice
-	FORCEINLINE_DEBUGGABLE int32 GetNumElements(const FRigVMRegister& Register, int32 InSliceIndex = 0) const
-	{
-		if (!Register.IsDynamic())
-		{
-			return Register.ElementCount;
-		}
-		
-		if (Register.IsNestedDynamic())
-		{
-			FRigVMNestedByteArray& ArrayStorage = *(FRigVMNestedByteArray*)&Data[Register.GetWorkByteIndex()];
-			return ArrayStorage[InSliceIndex].Num() / Register.ElementSize;
-		}
-
-		return 1;
-	}
-
-	// Returns the number of elements for a given slice
-	FORCEINLINE_DEBUGGABLE int32 GetNumElements(int32 InRegisterIndex, int32 InSliceIndex = 0) const
-	{
-		ensure(Registers.IsValidIndex(InRegisterIndex));
-		const FRigVMRegister& Register = Registers[InRegisterIndex];
-		return GetNumElements(Register, InSliceIndex);
-	}
-
-	// Returns the current const data pointer for a given register.
-	FORCEINLINE_DEBUGGABLE const uint8* GetData(const FRigVMRegister& Register, int32 InRegisterOffset = INDEX_NONE, int32 InSliceIndex = 0) const
-	{
-		return GetHandle(Register, InRegisterOffset).GetData(InSliceIndex, true);
-	}
-
-	// Returns the current const data pointer for a given register index.
-	FORCEINLINE_DEBUGGABLE const uint8* GetData(int32 InRegisterIndex, int32 InRegisterOffset = INDEX_NONE, int32 InSliceIndex = 0) const
-	{
-		ensure(Registers.IsValidIndex(InRegisterIndex));
-		const FRigVMRegister& Register = Registers[InRegisterIndex];
-		return GetData(Register, InRegisterOffset, InSliceIndex);
-	}
-
-	// Returns the current data pointer for a given register.
-	FORCEINLINE_DEBUGGABLE uint8* GetData(FRigVMRegister& Register, int32 InRegisterOffset = INDEX_NONE, int32 InSliceIndex = 0)
-	{
-		return GetHandle(Register, InRegisterOffset).GetData(InSliceIndex, true);
-	}
-
-	// Returns the current data pointer for a given register index.
-	// Note: This refers to the active slice - and can change over time.
-	FORCEINLINE_DEBUGGABLE uint8* GetData(int32 InRegisterIndex, int32 InRegisterOffset = INDEX_NONE, int32 InSliceIndex = 0)
-	{
-		ensure(Registers.IsValidIndex(InRegisterIndex));
-		FRigVMRegister& Register = Registers[InRegisterIndex];
-		return GetData(Register, InRegisterOffset, InSliceIndex);
-	}
-
-	// Returns the current const typed pointer for a given register.
-	// Note: This refers to the active slice - and can change over time.
-	template<typename T>
-	FORCEINLINE_DEBUGGABLE const T* Get(const FRigVMRegister& InRegister, int32 InRegisterOffset = INDEX_NONE, int32 InSliceIndex = 0) const
-	{
-		return (const T*)GetData(InRegister, InRegisterOffset, InSliceIndex);
-	}
-
-	// Returns the current const typed pointer for a given register index.
-	// Note: This refers to the active slice - and can change over time.
-	template<typename T>
-	FORCEINLINE_DEBUGGABLE const T* Get(int32 InRegisterIndex, int32 InRegisterOffset = INDEX_NONE, int32 InSliceIndex = 0) const
-	{
-		ensure(Registers.IsValidIndex(InRegisterIndex));
-		const FRigVMRegister& Register = Registers[InRegisterIndex];
-		return Get<T>(Register, InRegisterOffset, InSliceIndex);
-	}
-
-	// Returns the current const typed pointer for a given argument
-	// Note: This refers to the active slice - and can change over time.
-	template<typename T>
-	FORCEINLINE_DEBUGGABLE const T* Get(const FRigVMOperand& InOperand, int32 InSliceIndex = 0) const
-	{
-		return Get<T>(InOperand.GetRegisterIndex(), InOperand.GetRegisterOffset(), InSliceIndex);
-	}
-
-	// Returns the current const typed reference for a given register
-	// Note: This refers to the active slice - and can change over time.
-	template<typename T>
-	FORCEINLINE_DEBUGGABLE const T& GetRef(const FRigVMRegister& Register, int32 InRegisterOffset = INDEX_NONE, int32 InSliceIndex = 0) const
-	{
-		return *Get<T>(Register, InRegisterOffset, InSliceIndex);
-	}
-
-	// Returns the current const typed reference for a given register index
-	// Note: This refers to the active slice - and can change over time.
-	template<typename T>
-	FORCEINLINE_DEBUGGABLE const T& GetRef(int32 InRegisterIndex, int32 InRegisterOffset = INDEX_NONE, int32 InSliceIndex = 0) const
-	{
-		return *Get<T>(InRegisterIndex, InRegisterOffset, InSliceIndex);
-	}
-
-	// Returns the current const typed reference for a given argument
-	// Note: This refers to the active slice - and can change over time.
-	template<typename T>
-	FORCEINLINE_DEBUGGABLE const T& GetRef(const FRigVMOperand& InOperand, int32 InSliceIndex = 0) const
-	{
-		return GetRef<T>(InOperand.GetRegisterIndex(), InOperand.GetRegisterOffset(), InSliceIndex);
-	}
-
-	// Returns the current typed pointer for a given register
-	// Note: This refers to the active slice - and can change over time.
-	template<typename T>
-	FORCEINLINE_DEBUGGABLE T* Get(FRigVMRegister& InRegister, int32 InRegisterOffset = INDEX_NONE, int32 InSliceIndex = 0)
-	{
-		return (T*)GetData(InRegister, InRegisterOffset, InSliceIndex);
-	}
-
-	// Returns the current typed pointer for a given register index
-	// Note: This refers to the active slice - and can change over time.
-	template<typename T>
-	FORCEINLINE_DEBUGGABLE T* Get(int32 InRegisterIndex, int32 InRegisterOffset = INDEX_NONE, int32 InSliceIndex = 0)
-	{
-		ensure(Registers.IsValidIndex(InRegisterIndex));
-		FRigVMRegister& Register = Registers[InRegisterIndex];
-		return Get<T>(Register, InRegisterOffset, InSliceIndex);
-	}
-
-	// Returns the current typed pointer for a given argument
-	// Note: This refers to the active slice - and can change over time.
-	template<typename T>
-	FORCEINLINE_DEBUGGABLE T* Get(const FRigVMOperand& InOperand, int32 InSliceIndex = 0)
-	{
-		return Get<T>(InOperand.GetRegisterIndex(), InOperand.GetRegisterOffset(), InSliceIndex);
-	}
-
-	// Returns the current typed reference for a given register
-	// Note: This refers to the active slice - and can change over time.
-	template<typename T>
-	FORCEINLINE_DEBUGGABLE T& GetRef(FRigVMRegister& Register, int32 InRegisterOffset = INDEX_NONE, int32 InSliceIndex = 0)
-	{
-		return *Get<T>(Register, InRegisterOffset, InSliceIndex);
-	}
-
-	// Returns the current typed reference for a given register index
-	// Note: This refers to the active slice - and can change over time.
-	template<typename T>
-	FORCEINLINE_DEBUGGABLE T& GetRef(int32 InRegisterIndex, int32 InRegisterOffset = INDEX_NONE, int32 InSliceIndex = 0)
-	{
-		return *Get<T>(InRegisterIndex, InRegisterOffset, InSliceIndex);
-	}
-
-	// Returns the current typed reference for a given argument
-	// Note: This refers to the active slice - and can change over time.
-	template<typename T>
-	FORCEINLINE_DEBUGGABLE T& GetRef(const FRigVMOperand& InOperand, int32 InSliceIndex = 0)
-	{
-		return GetRef<T>(InOperand.GetRegisterIndex(), InOperand.GetRegisterOffset(), InSliceIndex);
-	}
-
-	// Returns an array view for all elements of the current slice for a given register.
-	template<typename T>
-	FORCEINLINE_DEBUGGABLE FRigVMFixedArray<T> GetFixedArray(FRigVMRegister& InRegister, int32 InRegisterOffset = INDEX_NONE, int32 InSliceIndex = 0)
-	{
-		if (InRegisterOffset == INDEX_NONE)
-		{
-			uint8* Ptr = (uint8*)GetDataPtr(InRegister, InRegisterOffset, InSliceIndex, false);
-			if (InRegister.IsNestedDynamic())
-			{
-				FRigVMByteArray* Storage = (FRigVMByteArray*)Ptr;
-				return FRigVMFixedArray<T>(*Storage);
-			}
-			else if (InRegister.IsDynamic())
-			{
-				FRigVMByteArray* Storage = (FRigVMByteArray*)Ptr;
-				Ptr = Storage->GetData() + InSliceIndex * InRegister.GetNumBytesPerSlice();
-			}
-
-			return FRigVMFixedArray<T>((T*)Ptr, InRegister.ElementCount);
-		}
-
-		TArray<T>* StoredArray = (TArray<T>*)GetData(InRegister, InRegisterOffset, InSliceIndex);
-		return FRigVMFixedArray<T>(StoredArray->GetData(), StoredArray->Num());
-	}
-
-	// Returns an array view for all elements of the current slice for a given register index.
-	template<typename T>
-	FORCEINLINE_DEBUGGABLE FRigVMFixedArray<T> GetFixedArray(int32 InRegisterIndex, int32 InRegisterOffset = INDEX_NONE, int32 InSliceIndex = 0)
-	{
-		ensure(Registers.IsValidIndex(InRegisterIndex));
-		FRigVMRegister& Register = Registers[InRegisterIndex];
-		return GetFixedArray<T>(Register, InRegisterOffset, InSliceIndex);
-	}
-	
-	// Returns an array view for all elements of the current slice for a given argument.
-	template<typename T>
-	FORCEINLINE_DEBUGGABLE FRigVMFixedArray<T> GetFixedArray(const FRigVMOperand& InOperand, int32 InSliceIndex = 0)
-	{
-		return GetFixedArray<T>(InOperand.GetRegisterIndex(), InOperand.GetRegisterOffset(), InSliceIndex);
-	}
-
-	// Returns an array view for all elements of the current slice for a given register.
-	template<typename T>
-	FORCEINLINE_DEBUGGABLE FRigVMDynamicArray<T> GetDynamicArray(FRigVMRegister& InRegister, int32 InSliceIndex = 0)
-	{
-		if (!InRegister.IsDynamic())
-		{
-			return FRigVMDynamicArray<T>(DefaultByteArray);
-		}
-		
-		FRigVMByteArray* ArrayStorage = (FRigVMByteArray*)GetDataPtr(InRegister, INDEX_NONE, 0, false);
-		return FRigVMDynamicArray<T>(*ArrayStorage);
-	}
-
-	// Returns an array view for all elements of the current slice for a given register index.
-	template<typename T>
-	FORCEINLINE_DEBUGGABLE FRigVMDynamicArray<T> GetDynamicArray(int32 InRegisterIndex, int32 InSliceIndex = 0)
-	{
-		ensure(Registers.IsValidIndex(InRegisterIndex));
-		FRigVMRegister& Register = Registers[InRegisterIndex];
-		return GetDynamicArray<T>(Register, InSliceIndex);
-	}
-
-	// Returns an array view for all elements of the current slice for a given argument.
-	template<typename T>
-	FORCEINLINE_DEBUGGABLE FRigVMDynamicArray<T> GetDynamicArray(const FRigVMOperand& InOperand, int32 InSliceIndex = 0)
-	{
-		return GetDynamicArray<T>(InOperand.GetRegisterIndex(), InOperand.GetRegisterOffset(), InSliceIndex);
-	}
-
-	// Returns the script struct used for a given register (can be nullptr for non-struct-registers).
-	FORCEINLINE_DEBUGGABLE UScriptStruct* GetScriptStruct(const FRigVMRegister& Register) const
-	{
-		if (Register.ScriptStructIndex != INDEX_NONE)
-		{
-			ensure(ScriptStructs.IsValidIndex(Register.ScriptStructIndex));
-			return ScriptStructs[Register.ScriptStructIndex];
-		}
-		return nullptr;
-	}
-
-	// Returns the script struct used for a given register index (can be nullptr for non-struct-registers).
-	FORCEINLINE_DEBUGGABLE UScriptStruct* GetScriptStruct(int32 InRegisterIndex, int32 InRegisterOffset = INDEX_NONE) const
-	{
-		if (InRegisterOffset == INDEX_NONE)
-		{
-			ensure(Registers.IsValidIndex(InRegisterIndex));
-			const FRigVMRegister& Register = Registers[InRegisterIndex];
-			return GetScriptStruct(Register);
-		}
-		ensure(RegisterOffsets.IsValidIndex(InRegisterOffset));
-		const FRigVMRegisterOffset& Path = RegisterOffsets[InRegisterOffset];
-		return Path.GetScriptStruct();
-	}
-
-	// Copies the content of a source register to a target register.
-	// The source register can optionally be referencing a specific source memory container.
-	bool Copy(
-		int32 InSourceRegisterIndex,
-		int32 InTargetRegisterIndex,
-		const FRigVMMemoryContainer* InSourceMemory = nullptr,
-		int32 InSourceRegisterOffset = INDEX_NONE,
-		int32 InTargetRegisterOffset = INDEX_NONE,
-		int32 InSourceSliceIndex = 0,
-		int32 InTargetSliceIndex = 0);
-
-	// Copies the content of a source register to a target register.
-	// The source register can optionally be referencing a specific source memory container.
-	// Note: This only works if SupportsNames() == true
-	bool Copy(
-		const FName& InSourceName,
-		const FName& InTargetName,
-		const FRigVMMemoryContainer* InSourceMemory = nullptr,
-		int32 InSourceRegisterOffset = INDEX_NONE,
-		int32 InTargetRegisterOffset = INDEX_NONE,
-		int32 InSourceSliceIndex = 0,
-		int32 InTargetSliceIndex = 0);
-
-	// Copies the content of a source register to a target register.
-	// The source register can optionally be referencing a specific source memory container.
-	bool Copy(
-		const FRigVMOperand& InSourceOperand,
-		const FRigVMOperand& InTargetOperand,
-		const FRigVMMemoryContainer* InSourceMemory = nullptr,
-		int32 InSourceSliceIndex = 0,
-		int32 InTargetSliceIndex = 0);
-
-	// Returns the index of a register based on the register name.
-	// Note: This only works if SupportsNames() == true
-	FORCEINLINE_DEBUGGABLE int32 GetIndex(const FName& InName) const
-	{
-		if (!bUseNameMap)
-		{
-			return INDEX_NONE;
-		}
-
-		if (NameMap.Num() != Registers.Num())
-		{
-			for (int32 Index = 0; Index < Registers.Num(); Index++)
-			{
-				if (Registers[Index].Name == InName)
-				{
-					return Index;
-				}
-			}
-		}
-		else
-		{
-			const int32* Index = NameMap.Find(InName);
-			if (Index != nullptr)
-			{
-				return *Index;
-			}
-		}
-
-		return INDEX_NONE;
-	}
-
-	// Returns true if a given name is available for a new register.
-	FORCEINLINE_DEBUGGABLE bool IsNameAvailable(const FName& InPotentialNewName) const
-	{
-		if (!bUseNameMap)
-		{
-			return false;
-		}
-		return GetIndex(InPotentialNewName) == INDEX_NONE;
-	}
-
-	// Adds a new named register for a typed array from an array view (used by compiler)
-	template<typename T>
-	FORCEINLINE_DEBUGGABLE int32 AddFixedArray(const FName& InNewName, const FRigVMFixedArray<T>& InArrayView, int32 InSliceCount = 1)
-	{
-		return AddRegisterArray<T>(true, InNewName, InArrayView.Num(), true, (const uint8*)InArrayView.GetData(), InSliceCount);
-	}
-
-	// Adds a new unnamed register for a typed array from an array view.
-	template<typename T>
-	FORCEINLINE_DEBUGGABLE int32 AddFixedArray(const FRigVMFixedArray<T>& InArrayView, int32 InSliceCount = 1)
-	{
-		return AddFixedArray<T>(NAME_None, InArrayView, InSliceCount);
-	}
-
-	// Adds a new named register for a typed array from an array view (used by compiler)
-	template<typename T>
-	FORCEINLINE_DEBUGGABLE int32 AddDynamicArray(const FName& InNewName, const FRigVMFixedArray<T>& InArrayView, int32 InSliceCount = 1)
-	{
-		return AddRegisterArray<T>(false, InNewName, InArrayView.Num(), true, (const uint8*)InArrayView.GetData(), InSliceCount);
-	}
-
-	// Adds a new unnamed register for a typed array from an array view.
-	template<typename T>
-	FORCEINLINE_DEBUGGABLE int32 AddDynamicArray(const FRigVMFixedArray<T>& InArrayView, int32 InSliceCount = 1)
-	{
-		return AddDynamicArray<T>(NAME_None, InArrayView, InSliceCount);
-	}
-
-	// Adds a new named register for a typed array from an array view (used by compiler)
-	template<typename T>
-	FORCEINLINE_DEBUGGABLE int32 AddDynamicArray(const FName& InNewName, const TArray<T>& InArray, int32 InSliceCount = 1)
-	{
-		return AddDynamicArray<T>(InNewName, FRigVMFixedArray<T>(InArray), InSliceCount);
-	}
-
-	// Adds a new unnamed register for a typed array from an array view.
-	template<typename T>
-	FORCEINLINE_DEBUGGABLE int32 AddDynamicArray(const TArray<T>& InArray, int32 InSliceCount = 1)
-	{
-		return AddDynamicArray<T>(NAME_None, InArray, InSliceCount);
-	}
-
-	// Adds a new named register for a typed value from a value reference.
-	template<typename T>
-	FORCEINLINE_DEBUGGABLE int32 Add(const FName& InNewName, const T& InValue, int32 InSliceCount = 1)
-	{
-		return AddRegisterArray<T>(true, InNewName, 1, false, (const uint8*)&InValue, InSliceCount);
-	}
-
-	// Adds a new unnamed register for a typed value from a value reference.
-	template<typename T>
-	FORCEINLINE_DEBUGGABLE int32 Add(const T& InValue, int32 InSliceCount = 1)
-	{
-		return Add<T>(NAME_None, InValue, InSliceCount);
-	}
-
-	// Adds a new named register for a typed value from a value reference.
-	template<typename T>
-	FORCEINLINE_DEBUGGABLE int32 AddDynamicValue(const FName& InNewName, const T& InValue, int32 InSliceCount = 1)
-	{
-		return AddRegisterArray<T>(false, InNewName, 1, false, (const uint8*)&InValue, InSliceCount);
-	}
-
-	// Adds a new unnamed register for a typed value from a value reference.
-	template<typename T>
-	FORCEINLINE_DEBUGGABLE int32 AddDynamicValue(const T& InValue, int32 InSliceCount = 1)
-	{
-		return AddDynamicValue<T>(NAME_None, InValue, InSliceCount);
-	}
-
-	// Adds a register path and returns its index
-	int32 GetOrAddRegisterOffset(int32 InRegisterIndex, int32 InArrayElement = 0);
-
-	// Adds a register path and returns its index
-	int32 GetOrAddRegisterOffset(int32 InRegisterIndex, const FString& InSegmentPath, int32 InArrayElement = 0);
-
-	// Adds a register path and returns its index
-	int32 GetOrAddRegisterOffset(int32 InRegisterIndex, UScriptStruct* InScriptStruct, const FString& InSegmentPath, int32 InInitialOffset = 0, int32 InElementSize = 0);
-
-	void SetRegisterValueFromString(const FRigVMOperand& InOperand, const FString& InCPPType, const UObject* InCPPTypeObject, const TArray<FString>& InDefaultValues);
-	TArray<FString> GetRegisterValueAsString(const FRigVMOperand& InOperand, const FString& InCPPType, const UObject* InCPPTypeObject);
-
-	// Returns the register offset given an index
-	const FRigVMRegisterOffset& GetRegisterOffset(int32 InRegisterOffsetIndex) const;
-
-	// Returns the register offset for a given operand
-	const FRigVMRegisterOffset& GetRegisterOffsetForOperand(const FRigVMOperand& InOperand) const;
-
-	// returns the statistics information
-	FRigVMMemoryStatistics GetStatistics() const
-	{
-		FRigVMMemoryStatistics Statistics;
-		Statistics.RegisterCount = Registers.Num();
-		Statistics.DataBytes = Data.GetAllocatedSize();
-		Statistics.TotalBytes = Data.GetAllocatedSize() + Registers.GetAllocatedSize() + RegisterOffsets.GetAllocatedSize();
-		return Statistics;
-	}
-
-private:
-
-	// Copies the source memory into a known register
-	bool Copy(
-		int32 InTargetRegisterIndex,
-		int32 InTargetRegisterOffset,
-		ERigVMRegisterType InTargetType,
-		const uint8* InSourcePtr,
-		uint8* InTargetPtr,
-		uint16 InNumBytes);
-
-	template<
-		typename T,
-		typename TEnableIf<TIsArithmetic<T>::Value>::Type* = nullptr
-	>
-	FORCEINLINE_DEBUGGABLE int32 AddRegisterArray(bool bFixed, const FName& InNewName, int32 InCount, bool bIsArrayPerSlice, const uint8* InDataPtr, int32 InSliceCount)
-	{
-		return AddRegisterArray<T>(bFixed, InNewName, InCount, bIsArrayPerSlice, InDataPtr, InSliceCount, ERigVMRegisterType::Plain, nullptr);
-	}
-
-	template<
-		typename T,
-		typename TEnableIf<TRigVMIsName<T>::Value>::Type* = nullptr
-	>
-	FORCEINLINE_DEBUGGABLE int32 AddRegisterArray(bool bFixed, const FName& InNewName, int32 InCount, bool bIsArrayPerSlice, const uint8* InDataPtr, int32 InSliceCount)
-	{
-		return AddRegisterArray<T>(bFixed, InNewName, InCount, bIsArrayPerSlice, InDataPtr, InSliceCount, ERigVMRegisterType::Name, nullptr);
-	}
-
-	template<
-		typename T,
-		typename TEnableIf<TRigVMIsString<T>::Value>::Type* = nullptr
-	>
-	FORCEINLINE_DEBUGGABLE int32 AddRegisterArray(bool bFixed, const FName& InNewName, int32 InCount, bool bIsArrayPerSlice, const uint8* InDataPtr, int32 InSliceCount)
-	{
-		return AddRegisterArray<T>(bFixed, InNewName, InCount, bIsArrayPerSlice, InDataPtr, InSliceCount, ERigVMRegisterType::String, nullptr);
-	}
-
-	template<
-		typename T,
-		typename TEnableIf<TModels<CRigVMUStruct, T>::Value>::Type* = nullptr
-	>
-	FORCEINLINE_DEBUGGABLE int32 AddRegisterArray(bool bFixed, const FName& InNewName, int32 InCount, bool bIsArrayPerSlice, const uint8* InDataPtr, int32 InSliceCount)
-	{
-		return AddRegisterArray<T>(bFixed, InNewName, InCount, bIsArrayPerSlice, InDataPtr, InSliceCount, ERigVMRegisterType::Struct, T::StaticStruct());
-	}
-
-	template<
-		typename T,
-		typename TEnableIf<TIsEnum<T>::Value>::Type* = nullptr
-	>
-	FORCEINLINE_DEBUGGABLE int32 AddRegisterArray(bool bFixed, const FName& InNewName, int32 InCount, bool bIsArrayPerSlice, const uint8* InDataPtr, int32 InSliceCount)
-	{
-		return AddRegisterArray<T>(bFixed, InNewName, InCount, bIsArrayPerSlice, InDataPtr, InSliceCount, ERigVMRegisterType::Plain, nullptr);
-	}
-
-	template <
-		typename T,
-		typename TEnableIf<TRigVMIsBaseStructure<T>::Value, T>::Type* = nullptr
-	>
-	FORCEINLINE_DEBUGGABLE int32 AddRegisterArray(bool bFixed, const FName& InNewName, int32 InCount, bool bIsArrayPerSlice, const uint8* InDataPtr, int32 InSliceCount)
-	{
-		return AddRegisterArray<T>(bFixed, InNewName, InCount, bIsArrayPerSlice, InDataPtr, InSliceCount, ERigVMRegisterType::Plain, TBaseStructure<T>::Get());
-	}
-
-	template<typename T>
-	FORCEINLINE_DEBUGGABLE int32 AddRegisterArray(bool bFixed, const FName& InNewName, int32 InCount, bool bIsArrayPerSlice, const uint8* InDataPtr, int32 InSliceCount, ERigVMRegisterType InType, UScriptStruct* InScriptStruct)
-	{
-		if(bFixed)
-		{
-			return AddFixedArray(InNewName, sizeof(T), InCount, bIsArrayPerSlice, InDataPtr, InSliceCount, InType, InScriptStruct);
-		}
-		else
-		{
-			return AddDynamicArray(InNewName, sizeof(T), InCount, bIsArrayPerSlice, InDataPtr, InSliceCount, InType, InScriptStruct);
-		}
-	}
-
-	FORCEINLINE_DEBUGGABLE int32 AddRegisterArray(bool bFixed, const FName& InNewName, int32 InElementSize, int32 InCount, bool bIsArrayPerSlice, const uint8* InDataPtr, int32 InSliceCount, ERigVMRegisterType InType, UScriptStruct* InScriptStruct)
-	{
-		if (bFixed)
-		{
-			return AddFixedArray(InNewName, InElementSize, InCount, bIsArrayPerSlice, InDataPtr, InSliceCount, InType, InScriptStruct);
-		}
-		else
-		{
-			return AddDynamicArray(InNewName, InElementSize, InCount, bIsArrayPerSlice, InDataPtr, InSliceCount, InType, InScriptStruct);
-		}
-	}
-
-	// Adds a new named register for a fixed array from a data pointer (used by compiler)
-	FORCEINLINE_DEBUGGABLE int32 AddFixedArray(const FName& InNewName, int32 InElementSize, int32 InCount, bool bIsArrayPerSlice, const uint8* InDataPtr, int32 InSliceCount, ERigVMRegisterType InType, UScriptStruct* InScriptStruct)	{
-		int32 Register = Allocate(InNewName, InElementSize, InCount, InSliceCount, nullptr, false);
-		if (!Registers.IsValidIndex(Register))
-		{
-			return Register;
-		}
-
-		Registers[Register].Type = InType;
-		Registers[Register].ScriptStructIndex = FindOrAddScriptStruct(InScriptStruct);
-		Registers[Register].bIsArray = bIsArrayPerSlice;
-
-		UpdateRegisters();
-		Construct(Register);
-
-		if (InDataPtr)
-		{
-			for (int32 SliceIndex = 0; SliceIndex < InSliceCount; SliceIndex++)
-			{
-				Copy(Register, INDEX_NONE, InType, InDataPtr, GetData(Registers[Register], INDEX_NONE, SliceIndex), InElementSize * InCount);
-			}
-		}
-		return Register;
-	}
-
-	// Adds a new named register for a dynamic array from a data pointer (used by compiler)
-	FORCEINLINE_DEBUGGABLE int32 AddDynamicArray(const FName& InNewName, int32 InElementSize, int32 InCount, bool bIsArrayPerSlice, const uint8* InDataPtr, int32 InSliceCount, ERigVMRegisterType InType, UScriptStruct* InScriptStruct)
-	{
-		int32 Register = INDEX_NONE;
-		
-		if (bIsArrayPerSlice)
-		{
-			Register = Allocate(InNewName, sizeof(FRigVMNestedByteArray), 1, 1, nullptr);
-			if (!Registers.IsValidIndex(Register))
-			{
-				return Register;
-			}
-
-			Registers[Register].Type = InType;
-			Registers[Register].ScriptStructIndex = FindOrAddScriptStruct(InScriptStruct);
-			Registers[Register].bIsDynamic = true;
-			Registers[Register].bIsArray = bIsArrayPerSlice;
-			Registers[Register].ElementSize = InElementSize;
-			Registers[Register].ElementCount = InCount;
-			Registers[Register].SliceCount = InSliceCount;
-
-			uint8* Ptr = (uint8*)&Data[Registers[Register].GetWorkByteIndex()];
-			FRigVMNestedByteArray& Storage = *(FRigVMNestedByteArray*)Ptr;
-			Storage.SetNum(InSliceCount);
-
-			for (int32 SliceIndex = 0; SliceIndex < InSliceCount; SliceIndex++)
-			{
-				Storage[SliceIndex].SetNumZeroed(InCount * InElementSize);
-				if (InDataPtr)
-				{
-					Copy(Register, INDEX_NONE, InType, InDataPtr, Storage[SliceIndex].GetData(), InElementSize * InCount);
-				}
-			}
-
-			if (InDataPtr == nullptr)
-			{
-				Construct(Register);
-			}
-		}
-		else
-		{
-			Register = Allocate(InNewName, sizeof(FRigVMByteArray), 1, 1, nullptr);
-			if (!Registers.IsValidIndex(Register))
-			{
-				return Register;
-			}
-
-			Registers[Register].Type = InType;
-			Registers[Register].ScriptStructIndex = FindOrAddScriptStruct(InScriptStruct);
-			Registers[Register].bIsDynamic = true;
-			Registers[Register].bIsArray = bIsArrayPerSlice;
-			Registers[Register].ElementSize = InElementSize;
-			Registers[Register].ElementCount = InCount;
-			Registers[Register].SliceCount = InSliceCount;
-
-			uint8* Ptr = (uint8*)&Data[Registers[Register].GetWorkByteIndex()];
-			FRigVMByteArray& Storage = *(FRigVMByteArray*)Ptr;
-
-			Storage.SetNumZeroed(InSliceCount * InElementSize);
-			if (InDataPtr)
-			{
-				for (int32 SliceIndex = 0; SliceIndex < InSliceCount; SliceIndex++)
-				{
-					Copy(Register, INDEX_NONE, InType, InDataPtr, &Storage[SliceIndex * InElementSize], InElementSize);
-				}
-			}
-			else
-			{
-				Construct(Register);
-			}
-		}
-
-		return Register;
-	}
-
-	// Updates internal data for topological changes
-	void UpdateRegisters();
-
-	// Allocates a new named register
-	int32 Allocate(const FName& InNewName, int32 InElementSize, int32 InElementCount, int32 InSliceCount, const uint8* InDataPtr = nullptr, bool bUpdateRegisters = true);
-
-	// Allocates a new unnamed register
-	int32 Allocate(int32 InElementSize, int32 InElementCount, int32 InSliceCount, const uint8* InDataPtr = nullptr, bool bUpdateRegisters = true);
-
-	// Performs optional construction of data within a struct register
-	bool Construct(int32 InRegisterIndex, int32 InElementIndex = INDEX_NONE, int32 InSliceIndex = 0);
-
-	// Performs optional destruction of data within a struct register
-	bool Destroy(int32 InRegisterIndex, int32 InElementIndex = INDEX_NONE, int32 InSliceIndex = 0);
-
-	// Fills a register with zero memory
-	void FillWithZeroes(int32 InRegisterIndex);
-
-	// Ensures to add a script struct to the internal map if needed
-	int32 FindOrAddScriptStruct(UScriptStruct* InScriptStruct);
-
-	UPROPERTY()
-	bool bUseNameMap;
-
-	UPROPERTY()
-	ERigVMMemoryType MemoryType;
-
-	UPROPERTY()
-	TArray<FRigVMRegister> Registers;
-
-	UPROPERTY()
-	TArray<FRigVMRegisterOffset> RegisterOffsets;
-
-	UPROPERTY(transient)
-	TArray<uint8> Data;
-
-	UPROPERTY(transient)
-	TArray<TObjectPtr<UScriptStruct>> ScriptStructs;
-
-	UPROPERTY(transient)
-	TMap<FName, int32> NameMap;
-
-	UPROPERTY(transient)
-	bool bEncounteredErrorDuringLoad;
-
-	static FRigVMByteArray DefaultByteArray;
-	static FRigVMRegisterOffset InvalidRegisterOffset;
-
-	friend class URigVM;
-	friend class URigVMCompiler;
-	friend struct FRigVMCompilerWorkData;
-};
-
-typedef FRigVMMemoryContainer* FRigVMMemoryContainerPtr;
-typedef FRigVMFixedArray<FRigVMMemoryContainer*> FRigVMMemoryContainerPtrArray;
\ No newline at end of file
diff --git a/Engine/Source/Runtime/RigVM/Public/RigVMCore/RigVMMemoryCommon.h b/Engine/Source/Runtime/RigVM/Public/RigVMCore/RigVMMemoryCommon.h
new file mode 100644
--- /dev/null
+++ b/Engine/Source/Runtime/RigVM/Public/RigVMCore/RigVMMemoryCommon.h
@@ -0,0 +1,153 @@
+// Copyright Epic Games, Inc. All Rights Reserved.
+
+#pragma once
+
+#include "CoreMinimal.h"
+#include "UObject/ObjectMacros.h"
+#include "RigVMTraits.h"
+#include "RigVMArray.h"
+#include "RigVMMemoryCommon.generated.h"
+
+#ifdef UE_BUILD_DEBUG
+	#define DEBUG_RIGVMMEMORY 0
+	//#define DEBUG_RIGVMMEMORY WITH_EDITOR
+#else
+	#define DEBUG_RIGVMMEMORY 0
+#endif
+
+#if DEBUG_RIGVMMEMORY
+	RIGVM_API DECLARE_LOG_CATEGORY_EXTERN(LogRigVMMemory, Log, All);
+	#define UE_LOG_RIGVMMEMORY(Format, ...) UE_LOG(LogRigVMMemory, Display, (Format), ##__VA_ARGS__)
+#else
+	#define UE_LOG_RIGVMMEMORY(Format, ...)
+#endif
+
+// The ERigVMMemoryType maps to memory container index in RigVM through
+// FRigVMOperand::GetContainerIndex() or URigVM::GetContainerIndex(...) 
+/**
+ * The type of memory used. Typically we differentiate between
+ * Work (Mutable) and Literal (Constant) memory.
+ */
+UENUM()
+enum class ERigVMMemoryType: uint8
+{
+	Work = 0, // Mutable state
+	Literal = 1, // Const / fixed state
+	External = 2, // Unowned external memory
+	Debug = 3, // Owned memory used for debug watches
+	Invalid
+};
+
+/**
+ * The FRigVMOperand represents an argument used for an operator
+ * within the virtual machine. Operands provide information about
+ * which memory needs to be referred to, which register within the
+ * memory all the way to the actual byte address in memory.
+ * The FRigVMOperand is a light weight address for a register in
+ * a FRigVMMemoryContainer.
+ * For external variables the register index represents the
+ * index of the external variable within the running VM.
+ */
+USTRUCT()
+struct RIGVM_API FRigVMOperand
+{
+	GENERATED_BODY()
+
+public:
+
+	FRigVMOperand()
+		: MemoryType(ERigVMMemoryType::Work)
+		, RegisterIndex(UINT16_MAX)
+		, RegisterOffset(UINT16_MAX)
+	{
+	}
+
+	FRigVMOperand(ERigVMMemoryType InMemoryType, int32 InRegisterIndex, int32 InRegisterOffset = INDEX_NONE)
+		: MemoryType(InMemoryType)
+		, RegisterIndex(InRegisterIndex < 0 ? UINT16_MAX : (uint16)InRegisterIndex)
+		, RegisterOffset(InRegisterOffset < 0 ? UINT16_MAX : (uint16)InRegisterOffset)
+	{
+	}
+
+	FORCEINLINE_DEBUGGABLE bool operator == (const FRigVMOperand& InOther) const
+	{
+		return MemoryType == InOther.MemoryType &&
+			RegisterIndex == InOther.RegisterIndex &&
+			RegisterOffset == InOther.RegisterOffset;
+	}
+
+	FORCEINLINE_DEBUGGABLE bool operator != (const FRigVMOperand& InOther) const
+	{
+		return !(*this == InOther);
+	}
+
+	// returns the memory type of this argument
+	FORCEINLINE_DEBUGGABLE bool IsValid() const { return RegisterIndex != UINT16_MAX; }
+
+	// returns the memory type of this argument
+	FORCEINLINE_DEBUGGABLE ERigVMMemoryType GetMemoryType() const { return MemoryType; }
+
+	// returns the index of the container of this argument
+	// this function should be kept in sync with URigVM::GetContainerIndex()
+	FORCEINLINE_DEBUGGABLE int32 GetContainerIndex() const
+	{
+		if(MemoryType == ERigVMMemoryType::External)
+		{
+			return (int32)ERigVMMemoryType::Work;
+		}
+		
+		if(MemoryType == ERigVMMemoryType::Debug)
+		{
+			return 2;
+		}
+		return (int32)MemoryType;
+	}
+
+	// returns the index of the register of this argument
+	FORCEINLINE_DEBUGGABLE int32 GetRegisterIndex() const { return RegisterIndex == UINT16_MAX ? INDEX_NONE : (int32)RegisterIndex; }
+
+	// returns the register offset of this argument
+	FORCEINLINE_DEBUGGABLE int32 GetRegisterOffset() const { return RegisterOffset == UINT16_MAX ? INDEX_NONE : (int32)RegisterOffset; }
+
+	friend FORCEINLINE uint32 GetTypeHash(const FRigVMOperand& Operand)
+	{
+  		return HashCombine(HashCombine((uint32)Operand.MemoryType, (uint32)Operand.RegisterIndex), (uint32)Operand.RegisterOffset);
+  	}
+	
+	void Serialize(FArchive& Ar);
+	void Save(FArchive& Ar) const;
+	void Load(FArchive& Ar);
+	FORCEINLINE friend FArchive& operator<<(FArchive& Ar, FRigVMOperand& P)
+	{
+		P.Serialize(Ar);
+		return Ar;
+	}
+
+private:
+
+	UPROPERTY()
+	ERigVMMemoryType MemoryType;
+
+	/**
+	 * The index of the register inside of the specific type of memory (work, literal etc).
+	 * For external variables the register index represents the index of the external variable within the running VM.
+	 */
+	UPROPERTY()
+	uint16 RegisterIndex;
+	
+	UPROPERTY()
+	uint16 RegisterOffset;
+
+	friend class URigVMCompiler;
+};
+
+#if UE_RIGVM_UCLASS_BASED_STORAGE_DISABLED
+
+typedef FRigVMFixedArray<FRigVMOperand> FRigVMOperandArray;
+
+#else
+
+typedef TArrayView<const FRigVMOperand> FRigVMOperandArray;
+
+#endif
+
diff --git a/Engine/Source/Runtime/RigVM/Public/RigVMCore/RigVMMemoryDeprecated.h b/Engine/Source/Runtime/RigVM/Public/RigVMCore/RigVMMemoryDeprecated.h
new file mode 100644
--- /dev/null
+++ b/Engine/Source/Runtime/RigVM/Public/RigVMCore/RigVMMemoryDeprecated.h
@@ -0,0 +1,1382 @@
+// Copyright Epic Games, Inc. All Rights Reserved.
+
+#pragma once
+
+#include "CoreMinimal.h"
+#include "UObject/ObjectMacros.h"
+#include "RigVMTraits.h"
+#include "RigVMStatistics.h"
+#include "RigVMArray.h"
+#include "RigVMMemoryCommon.h"
+#include "RigVMMemoryDeprecated.generated.h"
+
+/**
+ * The type of register within the memory.
+ */
+UENUM()
+enum class ERigVMRegisterType : uint8
+{
+	Plain, // bool, int32, float, FVector etc. (also structs that do NOT require a constructor / destructor to be valid)
+	String, // FString
+	Name, // FName
+	Struct, // Any USTRUCT (structs which require a constructor / destructor to be valid. structs which have indirection (arrays or strings in them, for example).
+	Invalid
+};
+
+// The register represents an address with the VM's memory. Within a register
+// we can store arbitrary data, so it provides a series of properties to
+// describe the memory location.
+// Registers also support the notion of slices. A slice is a complete copy of
+// the memory - so for example if your register stores 4 Vectors, then a slice 
+// would contain 48 bytes (4 * 3 * 4). The register can however store multiple
+// slices / copies of that if needed. Slices can be used to provide 
+// per-invocation memory to functions within the same register.
+// An integrator for example that needs to store a simulated position
+// might want access to a separate memory per loop iteration.
+USTRUCT()
+struct RIGVM_API FRigVMRegister
+{
+	GENERATED_BODY()
+
+	FRigVMRegister()
+		: Type(ERigVMRegisterType::Invalid)
+		, ByteIndex(INDEX_NONE)
+		, ElementSize(0)
+		, ElementCount(0)
+		, SliceCount(1)
+		, AlignmentBytes(0)
+		, TrailingBytes(0)
+		, Name(NAME_None)
+		, ScriptStructIndex(INDEX_NONE)
+		, bIsArray(false)
+		, bIsDynamic(false)
+#if WITH_EDITORONLY_DATA
+		, BaseCPPType(NAME_None)
+		, BaseCPPTypeObject(nullptr)
+#endif
+	{
+	}
+
+	// The type of register (plain, name, string, etc.)
+	UPROPERTY()
+	ERigVMRegisterType Type;
+
+	// The index of the first work byte
+	UPROPERTY()
+	uint32 ByteIndex;
+
+	// The size of each store element (for example 4 for a float)
+	UPROPERTY()
+	uint16 ElementSize;
+
+	// The number of elements in this register (for example the number of elements in an array)
+	UPROPERTY()
+	uint16 ElementCount;
+
+	// The number of slices (complete copies) (for example the number of iterations on a fixed loop)
+	// Potentially redundant state - can be removed.
+	UPROPERTY()
+	uint16 SliceCount;
+
+	// The number of leading bytes for alignment.
+	// Bytes that had to be introduced before the register
+	// to align the registers memory as per platform specification.
+	UPROPERTY()
+	uint8 AlignmentBytes;
+
+	// The number of trailing bytes.
+	// These originate after shrinking a register.
+	// Potentially unused state - can be removed.
+	UPROPERTY()
+	uint16 TrailingBytes;
+
+	// The name of the register (can be None)
+	UPROPERTY()
+	FName Name;
+
+	// For struct registers this is the index of the
+	// struct used - otherwise INDEX_NONE
+	UPROPERTY()
+	int32 ScriptStructIndex;
+
+	// If true defines this register as an array
+	UPROPERTY()
+	bool bIsArray;
+
+	// If true defines this register to use dynamic storage.
+	// Is this an array or singleton value with multiple slices
+	// which potentially requires changing count at runtime. 
+	UPROPERTY()
+	bool bIsDynamic;
+
+#if WITH_EDITORONLY_DATA
+
+	// Defines the CPP type used for the register
+	// This is only used for debugging purposes in editor.
+	UPROPERTY(transient)
+	FName BaseCPPType;
+
+	// The resolved UScriptStruct / UObject (in the future)
+	// used for debugging.
+	UPROPERTY(transient)
+	TObjectPtr<UObject> BaseCPPTypeObject;
+
+#endif
+
+	void Serialize(FArchive& Ar);
+	void Save(FArchive& Ar);
+	void Load(FArchive& Ar);
+	FORCEINLINE_DEBUGGABLE friend FArchive& operator<<(FArchive& Ar, FRigVMRegister& P)
+	{
+		P.Serialize(Ar);
+		return Ar;
+	}
+
+	// returns true if this register is using a dynamic array for storage
+	FORCEINLINE_DEBUGGABLE bool IsDynamic() const { return bIsDynamic; }
+
+	// returns true if this register is using a dynamic array for storage
+	FORCEINLINE_DEBUGGABLE bool IsNestedDynamic() const { return bIsDynamic && bIsArray; }
+
+	// returns the current address of the register within the data byte array.
+	// this can change over time - as the register is moving through slices.
+	// use GetFirstAllocatedByte to get the fixed first byte.
+	FORCEINLINE_DEBUGGABLE uint64 GetWorkByteIndex(int32 InSliceIndex = 0) const
+	{
+		ensure(InSliceIndex >= 0);
+		return ByteIndex + ((uint64)InSliceIndex * GetNumBytesPerSlice());
+	}
+
+	// returns the first allocated byte in the data byte array
+	FORCEINLINE_DEBUGGABLE uint64 GetFirstAllocatedByte() const
+	{ 
+		return ByteIndex - (uint64)AlignmentBytes;
+	}
+
+	// Returns the leading alignment bytes
+	FORCEINLINE_DEBUGGABLE uint8 GetAlignmentBytes() const { return AlignmentBytes; }
+
+	// Returns true if the register stores more than one element
+	FORCEINLINE_DEBUGGABLE bool IsArray() const { return bIsArray || (ElementCount > 1); }
+
+	// Returns the number of allocated bytes (including alignment + trailing bytes)
+	FORCEINLINE_DEBUGGABLE uint16 GetAllocatedBytes() const { return ElementCount * ElementSize * SliceCount + (uint16)AlignmentBytes + TrailingBytes; }
+
+	// Returns the number of bytes for a complete slice
+	FORCEINLINE_DEBUGGABLE uint16 GetNumBytesPerSlice() const { return ElementCount * ElementSize; }
+
+	// Returns the number of bytes for all slices
+	FORCEINLINE_DEBUGGABLE uint16 GetNumBytesAllSlices() const { return ElementCount * ElementSize * SliceCount; }
+
+	// Returns the total number of elements (elementcount * slicecount) in the register
+	FORCEINLINE_DEBUGGABLE uint32 GetTotalElementCount() const { return (uint32)ElementCount * (uint32)SliceCount; }
+
+};
+
+#if UE_RIGVM_UCLASS_BASED_STORAGE_DISABLED
+
+typedef FRigVMFixedArray<FRigVMRegister> FRigVMRegisterArray;
+
+#endif
+
+// The register offset represents a memory offset within a register's memory.
+// This can be used to represent memory addresses of array elements within
+// a struct, for example.
+// A register offset's path can look like MyTransformStruct.Transforms[3].Translation.X
+USTRUCT()
+struct RIGVM_API FRigVMRegisterOffset
+{
+	GENERATED_BODY()
+
+public:
+
+	// default constructor
+	FRigVMRegisterOffset()
+		: Segments()
+		, Type(ERigVMRegisterType::Invalid)
+		, CPPType(NAME_None)
+		, ScriptStruct(nullptr)
+		, ParentScriptStruct(nullptr)
+		, ArrayIndex(0)
+		, ElementSize(0)
+		, CachedSegmentPath() 
+	{
+	}
+
+	void Serialize(FArchive& Ar);
+	void Save(FArchive& Ar);
+	void Load(FArchive& Ar);
+	FORCEINLINE_DEBUGGABLE friend FArchive& operator<<(FArchive& Ar, FRigVMRegisterOffset& P)
+	{
+		P.Serialize(Ar);
+		return Ar;
+	}
+
+	// constructs a path given a struct and a segment path
+	FRigVMRegisterOffset(UScriptStruct* InScriptStruct, const FString& InSegmentPath, int32 InInitialOffset = 0, uint16 InElementSize = 0, const FName& InCPPType = NAME_None);
+
+#if UE_RIGVM_UCLASS_BASED_STORAGE_DISABLED
+
+	// returns the data pointer within a container
+	uint8* GetData(uint8* InContainer) const;
+
+#endif
+	
+	// returns the segments of this path
+	const TArray<int32>& GetSegments() const { return Segments; }
+
+	// returns true if this register offset contains an array segment.
+	// this means the register offset has to traverse a pointer and
+	// cannot be cached / collapsed easily
+	bool ContainsArraySegment() const;
+
+	bool operator == (const FRigVMRegisterOffset& InOther) const;
+
+	FORCEINLINE_DEBUGGABLE bool IsValid() const { return Type != ERigVMRegisterType::Invalid; }
+	FORCEINLINE_DEBUGGABLE ERigVMRegisterType GetType() const { return Type; }
+	FORCEINLINE_DEBUGGABLE FName GetCPPType() const { return CPPType; }
+	FORCEINLINE_DEBUGGABLE FString GetCachedSegmentPath() const { return CachedSegmentPath; }
+	FORCEINLINE_DEBUGGABLE int32 GetArrayIndex() const { return ArrayIndex; }
+	uint16 GetElementSize() const;
+	void SetElementSize(uint16 InElementSize) { ElementSize = InElementSize; };
+	UScriptStruct* GetScriptStruct() const;
+
+private:
+
+	// Segments represent the memory offset(s) to use when accessing the target memory.
+	// In case of indirection in the source memory each memory offset / jump is stored.
+	// So for example: If you are accessing the third Rotation.X of an array of transforms,
+	// the segments would read as: [-2, 12] (second array element, and the fourth float)
+	// Segment indices less than zero represent array element offsets, while positive numbers
+	// represents jumps within a struct.
+	UPROPERTY()
+	TArray<int32> Segments;
+
+	// Type of resulting register (for example Plain for Transform.Translation.X)
+	UPROPERTY()
+	ERigVMRegisterType Type;
+
+	// The c++ type of the resulting memory (for example float for Transform.Translation.X)
+	UPROPERTY()
+	FName CPPType;
+
+	// The c++ script struct of the resulting memory (for example FVector for Transform.Translation)
+	UPROPERTY()
+	TObjectPtr<UScriptStruct> ScriptStruct;
+
+	// The c++ script struct of the source memory (for example FTransform for Transform.Translation)
+	UPROPERTY()
+	TObjectPtr<UScriptStruct> ParentScriptStruct;
+
+	// The index of the element within an array (for example 3 for Transform[3])
+	UPROPERTY()
+	int32 ArrayIndex;
+
+	// The number of bytes of the resulting memory (for example 4 (float) for Transform.Translation.X)
+	UPROPERTY()
+	uint16 ElementSize;
+
+	// The cached path of the segments within this register, for example FTransform.Translation.X
+	UPROPERTY()
+	FString CachedSegmentPath;
+
+	friend struct FRigVMRegisterOffsetBuilder;
+	friend class URigVM;
+};
+
+#if UE_RIGVM_UCLASS_BASED_STORAGE_DISABLED
+
+/**
+ * The FRigVMMemoryHandle is used to access the memory used within a FRigMemoryContainer.
+ */
+struct FRigVMMemoryHandle
+{
+public:
+
+	// The type of memory being accessed.
+	enum FType
+	{
+		Plain, // Shallow memory with no indirection (for example FTransform.Translation.X)
+		Dynamic, // Memory representing an array
+		NestedDynamic, // Memory representing an array or single element with slices (for example within loop)
+		ArraySize // a count represented as a pointer address
+	};
+
+	FORCEINLINE_DEBUGGABLE FRigVMMemoryHandle()
+		: Ptr(nullptr)
+		, Type(FType::Plain)
+		, Size(1)
+		, RegisterOffset(nullptr)
+	{}
+
+	FORCEINLINE_DEBUGGABLE FRigVMMemoryHandle(uint8* InPtr, uint16 InSize = 1, FType InType = FType::Plain, const FRigVMRegisterOffset* InRegisterOffset = nullptr)
+		: Ptr(InPtr)
+		, Type(InType)
+		, Size(InSize)
+		, RegisterOffset(InRegisterOffset)
+	{}
+
+	FORCEINLINE_DEBUGGABLE FRigVMMemoryHandle(uint8* InPtr, const FRigVMRegister& InRegister, const FRigVMRegisterOffset* InRegisterOffset = nullptr)
+		: Ptr(InPtr)
+		, Type(FType::Plain)
+		, Size(InRegister.ElementSize)
+		, RegisterOffset(InRegisterOffset)
+	{
+		if (InRegister.IsNestedDynamic())
+		{
+			Type = FType::NestedDynamic;
+		}
+		else if (InRegister.IsDynamic())
+		{
+			Type = FType::Dynamic;
+		}
+		else
+		{
+			Size = (uint16)InRegister.GetNumBytesPerSlice();
+		}
+	}
+
+	FORCEINLINE_DEBUGGABLE FRigVMMemoryHandle(FRigVMByteArray* InPtr, uint16 InSize = 1, const FRigVMRegisterOffset* InRegisterOffset = nullptr)
+		: Ptr((uint8*)InPtr)
+		, Type(FType::Dynamic)
+		, Size(InSize)
+		, RegisterOffset(InRegisterOffset)
+	{}
+
+	FORCEINLINE_DEBUGGABLE FRigVMMemoryHandle(FRigVMNestedByteArray* InPtr, uint16 InSize = 1, const FRigVMRegisterOffset* InRegisterOffset = nullptr)
+		: Ptr((uint8*)InPtr)
+		, Type(FType::NestedDynamic)
+		, Size(InSize)
+		, RegisterOffset(InRegisterOffset)
+	{}
+
+	// getter for the contained const memory
+	FORCEINLINE_DEBUGGABLE operator const uint8*() const
+	{
+		return GetData();
+	}
+
+	// getter for the contained mutable memory
+	FORCEINLINE_DEBUGGABLE operator uint8*()
+	{
+		return GetData();
+	}
+
+	// returns the contained const memory for a given slice index
+	FORCEINLINE_DEBUGGABLE const uint8* GetData(int32 SliceIndex = 0, bool bGetArrayData = false) const
+	{
+		return GetData_Internal(SliceIndex, bGetArrayData);
+	}
+
+	// returns the contained mutable memory for a given slice index
+	FORCEINLINE_DEBUGGABLE uint8* GetData(int32 SliceIndex = 0, bool bGetArrayData = false)
+	{
+		return GetData_Internal(SliceIndex, bGetArrayData);
+	}
+
+	FORCEINLINE FType GetType() const
+	{
+		return Type;
+	}
+
+	FORCEINLINE bool IsDynamic() const
+	{
+		return (Type == Dynamic) || (Type == NestedDynamic);
+	}
+
+private:
+
+	FORCEINLINE_DEBUGGABLE uint8* GetData_Internal_NoOffset(int32 SliceIndex, bool bGetArrayData = false) const
+	{
+		ensure(SliceIndex >= 0);
+
+		if (Size == 0 || Ptr == nullptr)
+		{
+			return nullptr;
+		}
+
+		switch (Type)
+		{
+			case FType::Plain:
+			{
+				return Ptr + SliceIndex * Size;
+			}
+			case FType::ArraySize:
+			{
+				return Ptr;
+			}
+			case FType::Dynamic:
+			{
+				ensure(RegisterOffset == nullptr);
+				if (!bGetArrayData)
+				{
+					return Ptr;
+				}
+
+				FRigVMByteArray* ArrayStorage = (FRigVMByteArray*)Ptr;
+				if (ArrayStorage->Num() == 0)
+				{
+					return nullptr;
+				}
+
+				if (SliceIndex > 0)
+				{
+					SliceIndex = FMath::Max(SliceIndex, (ArrayStorage->Num() / Size) - 1);
+				}
+				return (uint8*)&(*ArrayStorage)[SliceIndex * Size];
+			}
+			case FType::NestedDynamic:
+			{
+				ensure(RegisterOffset == nullptr);
+
+				FRigVMNestedByteArray* ArrayStorage = (FRigVMNestedByteArray*)Ptr;
+				if (ArrayStorage->Num() == 0 && bGetArrayData)
+				{
+					return nullptr;
+				}
+
+				if (SliceIndex > 0)
+				{
+					SliceIndex = FMath::Max(SliceIndex, ArrayStorage->Num() - 1);
+				}
+
+				if (!bGetArrayData)
+				{
+					return Ptr;
+				}
+
+				return (*ArrayStorage)[SliceIndex].GetData();
+			}
+			default:
+			{
+				return nullptr;
+			}
+		}
+	}
+
+	FORCEINLINE_DEBUGGABLE uint8* GetData_Internal(int32 SliceIndex, bool bGetArrayData = false) const
+	{
+		uint8* Result = GetData_Internal_NoOffset(SliceIndex, bGetArrayData);
+		if (Result && RegisterOffset)
+		{
+			Result = RegisterOffset->GetData(Result);
+		}
+		return Result;
+	}
+
+	uint8* Ptr;
+	FType Type;
+	uint16 Size;
+	const FRigVMRegisterOffset* RegisterOffset;
+
+	friend class URigVM;
+};
+
+#endif
+
+/**
+ * The FRigVMMemoryContainer provides a heterogeneous memory container to store arbitrary
+ * data. Each element stored can be referred to using a FRigVMRegister.
+ * Elements can be accessed by index (index of the register), FRigVMOperand or by name.
+ * Name access is optional and is specified upon construction of the container.
+ * The memory container provides a series of templated functions to add and get data.
+ *
+ * For example:
+ * 		int32 Index = Container.Add<float>(4.f);
+ *      float& ValueRef = Container.GetRef<float>(Index);
+ *
+ * This can also be done with arrays:
+ *      TArray<float> MyArray = {3.f, 4.f, 5.f};
+ * 		int32 Index = Container.AddFixedArray<float>(MyArray);
+ *      FRigVMFixedArray<float> ArrayView = Container.GetFixedArray<float>(Index);
+ *
+ *
+ * Registers can store dynamically resizable memory as well by reyling on indirection.
+ * 
+ * Arrays with a single slice are going to be stored as a FRigVMByteArray, [unit8], so for example for a float array in C++ it would be (TArray<float> Param).
+ * Single values with multiple slices are also going to be stored as a FRigVMByteArray, [unit8], so for example for a float array in C++ it would be (float Param). 
+ * Arrays with multiple slices are going to be stored as a FRigVMNestedByteArray, [[unit8]], so for example for a float array in C++ it would be (TArray<float> Param).
+ */
+USTRUCT()
+struct RIGVM_API FRigVMMemoryContainer
+{
+	GENERATED_BODY()
+	
+public:
+
+	FRigVMMemoryContainer(bool bInUseNames = true);
+	
+#if UE_RIGVM_UCLASS_BASED_STORAGE_DISABLED
+	
+	FRigVMMemoryContainer(const FRigVMMemoryContainer& Other);
+	~FRigVMMemoryContainer();
+
+	bool CopyRegisters(const FRigVMMemoryContainer &InOther);
+	FRigVMMemoryContainer& operator= (const FRigVMMemoryContainer &InOther);
+
+	// returns the memory type of this container
+	FORCEINLINE_DEBUGGABLE ERigVMMemoryType GetMemoryType() const { return MemoryType;  }
+
+	// sets the memory type. should only be used when the container is empty
+	FORCEINLINE_DEBUGGABLE void SetMemoryType(ERigVMMemoryType InMemoryType) { MemoryType = InMemoryType; }
+
+	// returns true if this container supports name based lookup
+	FORCEINLINE_DEBUGGABLE bool SupportsNames() const { return bUseNameMap;  }
+
+	// returns the number of registers in this container
+	FORCEINLINE_DEBUGGABLE int32 Num() const { return Registers.Num(); }
+
+	// resets the container but maintains storage.
+	void Reset();
+
+	// resets the container and removes all storage.
+	void Empty();
+
+	// const accessor for a register based on index
+	FORCEINLINE_DEBUGGABLE const FRigVMRegister& operator[](int32 InIndex) const { return GetRegister(InIndex); }
+
+	// accessor for a register based on index
+	FORCEINLINE_DEBUGGABLE FRigVMRegister& operator[](int32 InIndex) { return GetRegister(InIndex); }
+
+	// const accessor for a register based on an argument
+	FORCEINLINE_DEBUGGABLE const FRigVMRegister& operator[](const FRigVMOperand& InArg) const { return GetRegister(InArg); }
+
+	// accessor for a register based on an argument
+	FORCEINLINE_DEBUGGABLE FRigVMRegister& operator[](const FRigVMOperand& InArg) { return GetRegister(InArg); }
+
+	// const accessor for a register based on a a name. note: only works if SupportsNames() == true
+	FORCEINLINE_DEBUGGABLE const FRigVMRegister& operator[](const FName& InName) const { return GetRegister(InName); }
+
+	// accessor for a register based on a a name. note: only works if SupportsNames() == true
+	FORCEINLINE_DEBUGGABLE FRigVMRegister& operator[](const FName& InName) { return GetRegister(InName); }
+
+	FORCEINLINE_DEBUGGABLE TArray<FRigVMRegister>::RangedForIteratorType      begin() { return Registers.begin(); }
+	FORCEINLINE_DEBUGGABLE TArray<FRigVMRegister>::RangedForConstIteratorType begin() const { return Registers.begin(); }
+	FORCEINLINE_DEBUGGABLE TArray<FRigVMRegister>::RangedForIteratorType      end() { return Registers.end(); }
+	FORCEINLINE_DEBUGGABLE TArray<FRigVMRegister>::RangedForConstIteratorType end() const { return Registers.end(); }
+
+	// const accessor for a register based on index
+	FORCEINLINE_DEBUGGABLE const FRigVMRegister& GetRegister(int32 InIndex) const { return Registers[InIndex]; }
+
+	// accessor for a register based on index
+	FORCEINLINE_DEBUGGABLE FRigVMRegister& GetRegister(int32 InIndex) { return Registers[InIndex]; }
+
+	// const accessor for a register based on an argument
+	FORCEINLINE_DEBUGGABLE const FRigVMRegister& GetRegister(const FRigVMOperand& InArg) const { return Registers[InArg.GetRegisterIndex()]; }
+
+	// accessor for a register based on an argument
+	FORCEINLINE_DEBUGGABLE FRigVMRegister& GetRegister(const FRigVMOperand& InArg) { return Registers[InArg.GetRegisterIndex()]; }
+
+	// const accessor for a register based on a a name. note: only works if SupportsNames() == true
+	FORCEINLINE_DEBUGGABLE const FRigVMRegister& GetRegister(const FName& InName) const { return Registers[GetIndex(InName)]; }
+
+	// accessor for a register based on a a name. note: only works if SupportsNames() == true
+	FORCEINLINE_DEBUGGABLE FRigVMRegister& GetRegister(const FName& InName) { return Registers[GetIndex(InName)]; }
+
+#endif
+	
+	void Serialize(FArchive& Ar);
+
+#if UE_RIGVM_UCLASS_BASED_STORAGE_DISABLED
+
+	void Save(FArchive& Ar);
+	void Load(FArchive& Ar);
+	FORCEINLINE_DEBUGGABLE friend FArchive& operator<<(FArchive& Ar, FRigVMMemoryContainer& P)
+	{
+		P.Serialize(Ar);
+		return Ar;
+	}
+
+	// Returns an argument for a given register.
+	// This is typically used to store a light weight address for use within a VM.
+	FORCEINLINE_DEBUGGABLE FRigVMOperand GetOperand(int32 InRegisterIndex, int32 InRegisterOffset)
+	{
+		ensure(Registers.IsValidIndex(InRegisterIndex));
+		return FRigVMOperand(MemoryType, InRegisterIndex, InRegisterOffset);
+	}
+
+	// Returns an argument for a given register.
+	// This is typically used to store a light weight address for use within a VM.
+	FORCEINLINE_DEBUGGABLE FRigVMOperand GetOperand(int32 InRegisterIndex, const FString& InSegmentPath = FString(), int32 InArrayElement = INDEX_NONE)
+	{
+		ensure(Registers.IsValidIndex(InRegisterIndex));
+		// Register offset must hold on to the ScriptStruct such that it can recalculate the struct size after cook
+		UScriptStruct* ScriptStruct = GetScriptStruct(InRegisterIndex);
+
+		int32 InitialOffset = 0;
+		int32 ElementSize = 0;
+		if (InArrayElement != INDEX_NONE)
+		{
+			InitialOffset = InArrayElement * Registers[InRegisterIndex].ElementSize;
+			ElementSize = Registers[InRegisterIndex].ElementSize;
+		}
+
+		return GetOperand(InRegisterIndex, GetOrAddRegisterOffset(InRegisterIndex, ScriptStruct, InSegmentPath, InitialOffset, ElementSize));
+	}
+
+private:
+
+	FORCEINLINE_DEBUGGABLE uint8* GetDataPtr(const FRigVMRegister& Register, int32 InRegisterOffset = INDEX_NONE, int32 InSliceIndex = 0, bool bArrayContent = false) const
+	{
+		if (Register.ElementCount == 0 && !Register.IsNestedDynamic())
+		{
+			return nullptr;
+		}
+
+		uint8* Ptr = nullptr;
+		if (Register.IsDynamic())
+		{
+			Ptr = (uint8*)&Data[Register.GetWorkByteIndex()];
+
+			if (Register.IsNestedDynamic())
+			{
+				FRigVMNestedByteArray* ArrayStorage = (FRigVMNestedByteArray*)Ptr;
+				Ptr = (uint8*)ArrayStorage->GetData();
+
+				if (Ptr)
+				{
+					Ptr = Ptr + InSliceIndex * sizeof(FRigVMByteArray);
+					if (Ptr && bArrayContent)
+					{
+						Ptr = ((FRigVMByteArray*)Ptr)->GetData();
+					}
+				}
+			}
+			else if(bArrayContent)
+			{
+				FRigVMByteArray* ArrayStorage = (FRigVMByteArray*)Ptr;
+				Ptr = (uint8*)ArrayStorage->GetData();
+
+				if (Ptr)
+				{
+					Ptr = Ptr + InSliceIndex * Register.GetNumBytesPerSlice();
+				}
+			}
+		}
+		else
+		{
+			Ptr = (uint8*)&Data[Register.GetWorkByteIndex(InSliceIndex)];
+		}
+
+		if (InRegisterOffset != INDEX_NONE && Ptr != nullptr)
+		{
+			Ptr = RegisterOffsets[InRegisterOffset].GetData(Ptr);
+		}
+		return Ptr;
+	}
+
+public:
+
+	// Returns a memory handle for a given register
+	FORCEINLINE_DEBUGGABLE FRigVMMemoryHandle GetHandle(const FRigVMRegister& Register, int32 InRegisterOffset = INDEX_NONE) const
+	{
+		const FRigVMRegisterOffset& RegisterOffset = GetRegisterOffset(InRegisterOffset);
+
+		if (Register.IsDynamic())
+		{
+			uint8* Ptr = (uint8*)&Data[Register.GetWorkByteIndex()];
+			return FRigVMMemoryHandle(Ptr, Register, RegisterOffset.IsValid() ? &RegisterOffset : nullptr);
+		}
+
+		uint8* Ptr = GetDataPtr(Register);
+		uint16 NumBytes = Register.GetNumBytesPerSlice();
+
+		if (RegisterOffset.IsValid())
+		{
+			NumBytes = RegisterOffset.GetElementSize();
+
+			if (!RegisterOffset.ContainsArraySegment())
+			{
+				// flatten the offset and cache the offset ptr directly
+				Ptr = RegisterOffset.GetData(Ptr);
+				return FRigVMMemoryHandle(Ptr, NumBytes);
+			}
+		}
+
+		return FRigVMMemoryHandle(Ptr, NumBytes, FRigVMMemoryHandle::FType::Plain, RegisterOffset.IsValid() ? &RegisterOffset : nullptr);
+	}
+
+	// Returns a memory handle for a given register
+	FORCEINLINE_DEBUGGABLE FRigVMMemoryHandle GetHandle(int32 InRegisterIndex, int32 InRegisterOffset = INDEX_NONE) const
+	{
+		ensure(Registers.IsValidIndex(InRegisterIndex));
+		const FRigVMRegister& Register = Registers[InRegisterIndex];
+		return GetHandle(Register, InRegisterOffset);
+	}
+
+	// Returns the number of elements for a given slice
+	FORCEINLINE_DEBUGGABLE int32 GetNumElements(const FRigVMRegister& Register, int32 InSliceIndex = 0) const
+	{
+		if (!Register.IsDynamic())
+		{
+			return Register.ElementCount;
+		}
+		
+		if (Register.IsNestedDynamic())
+		{
+			FRigVMNestedByteArray& ArrayStorage = *(FRigVMNestedByteArray*)&Data[Register.GetWorkByteIndex()];
+			return ArrayStorage[InSliceIndex].Num() / Register.ElementSize;
+		}
+
+		return 1;
+	}
+
+	// Returns the number of elements for a given slice
+	FORCEINLINE_DEBUGGABLE int32 GetNumElements(int32 InRegisterIndex, int32 InSliceIndex = 0) const
+	{
+		ensure(Registers.IsValidIndex(InRegisterIndex));
+		const FRigVMRegister& Register = Registers[InRegisterIndex];
+		return GetNumElements(Register, InSliceIndex);
+	}
+
+	// Returns the current const data pointer for a given register.
+	FORCEINLINE_DEBUGGABLE const uint8* GetData(const FRigVMRegister& Register, int32 InRegisterOffset = INDEX_NONE, int32 InSliceIndex = 0) const
+	{
+		return GetHandle(Register, InRegisterOffset).GetData(InSliceIndex, true);
+	}
+
+	// Returns the current const data pointer for a given register index.
+	FORCEINLINE_DEBUGGABLE const uint8* GetData(int32 InRegisterIndex, int32 InRegisterOffset = INDEX_NONE, int32 InSliceIndex = 0) const
+	{
+		ensure(Registers.IsValidIndex(InRegisterIndex));
+		const FRigVMRegister& Register = Registers[InRegisterIndex];
+		return GetData(Register, InRegisterOffset, InSliceIndex);
+	}
+
+	// Returns the current data pointer for a given register.
+	FORCEINLINE_DEBUGGABLE uint8* GetData(FRigVMRegister& Register, int32 InRegisterOffset = INDEX_NONE, int32 InSliceIndex = 0)
+	{
+		return GetHandle(Register, InRegisterOffset).GetData(InSliceIndex, true);
+	}
+
+	// Returns the current data pointer for a given register index.
+	// Note: This refers to the active slice - and can change over time.
+	FORCEINLINE_DEBUGGABLE uint8* GetData(int32 InRegisterIndex, int32 InRegisterOffset = INDEX_NONE, int32 InSliceIndex = 0)
+	{
+		ensure(Registers.IsValidIndex(InRegisterIndex));
+		FRigVMRegister& Register = Registers[InRegisterIndex];
+		return GetData(Register, InRegisterOffset, InSliceIndex);
+	}
+
+	// Returns the current const typed pointer for a given register.
+	// Note: This refers to the active slice - and can change over time.
+	template<typename T>
+	FORCEINLINE_DEBUGGABLE const T* Get(const FRigVMRegister& InRegister, int32 InRegisterOffset = INDEX_NONE, int32 InSliceIndex = 0) const
+	{
+		return (const T*)GetData(InRegister, InRegisterOffset, InSliceIndex);
+	}
+
+	// Returns the current const typed pointer for a given register index.
+	// Note: This refers to the active slice - and can change over time.
+	template<typename T>
+	FORCEINLINE_DEBUGGABLE const T* Get(int32 InRegisterIndex, int32 InRegisterOffset = INDEX_NONE, int32 InSliceIndex = 0) const
+	{
+		ensure(Registers.IsValidIndex(InRegisterIndex));
+		const FRigVMRegister& Register = Registers[InRegisterIndex];
+		return Get<T>(Register, InRegisterOffset, InSliceIndex);
+	}
+
+	// Returns the current const typed pointer for a given argument
+	// Note: This refers to the active slice - and can change over time.
+	template<typename T>
+	FORCEINLINE_DEBUGGABLE const T* Get(const FRigVMOperand& InOperand, int32 InSliceIndex = 0) const
+	{
+		return Get<T>(InOperand.GetRegisterIndex(), InOperand.GetRegisterOffset(), InSliceIndex);
+	}
+
+	// Returns the current const typed reference for a given register
+	// Note: This refers to the active slice - and can change over time.
+	template<typename T>
+	FORCEINLINE_DEBUGGABLE const T& GetRef(const FRigVMRegister& Register, int32 InRegisterOffset = INDEX_NONE, int32 InSliceIndex = 0) const
+	{
+		return *Get<T>(Register, InRegisterOffset, InSliceIndex);
+	}
+
+	// Returns the current const typed reference for a given register index
+	// Note: This refers to the active slice - and can change over time.
+	template<typename T>
+	FORCEINLINE_DEBUGGABLE const T& GetRef(int32 InRegisterIndex, int32 InRegisterOffset = INDEX_NONE, int32 InSliceIndex = 0) const
+	{
+		return *Get<T>(InRegisterIndex, InRegisterOffset, InSliceIndex);
+	}
+
+	// Returns the current const typed reference for a given argument
+	// Note: This refers to the active slice - and can change over time.
+	template<typename T>
+	FORCEINLINE_DEBUGGABLE const T& GetRef(const FRigVMOperand& InOperand, int32 InSliceIndex = 0) const
+	{
+		return GetRef<T>(InOperand.GetRegisterIndex(), InOperand.GetRegisterOffset(), InSliceIndex);
+	}
+
+	// Returns the current typed pointer for a given register
+	// Note: This refers to the active slice - and can change over time.
+	template<typename T>
+	FORCEINLINE_DEBUGGABLE T* Get(FRigVMRegister& InRegister, int32 InRegisterOffset = INDEX_NONE, int32 InSliceIndex = 0)
+	{
+		return (T*)GetData(InRegister, InRegisterOffset, InSliceIndex);
+	}
+
+	// Returns the current typed pointer for a given register index
+	// Note: This refers to the active slice - and can change over time.
+	template<typename T>
+	FORCEINLINE_DEBUGGABLE T* Get(int32 InRegisterIndex, int32 InRegisterOffset = INDEX_NONE, int32 InSliceIndex = 0)
+	{
+		ensure(Registers.IsValidIndex(InRegisterIndex));
+		FRigVMRegister& Register = Registers[InRegisterIndex];
+		return Get<T>(Register, InRegisterOffset, InSliceIndex);
+	}
+
+	// Returns the current typed pointer for a given argument
+	// Note: This refers to the active slice - and can change over time.
+	template<typename T>
+	FORCEINLINE_DEBUGGABLE T* Get(const FRigVMOperand& InOperand, int32 InSliceIndex = 0)
+	{
+		return Get<T>(InOperand.GetRegisterIndex(), InOperand.GetRegisterOffset(), InSliceIndex);
+	}
+
+	// Returns the current typed reference for a given register
+	// Note: This refers to the active slice - and can change over time.
+	template<typename T>
+	FORCEINLINE_DEBUGGABLE T& GetRef(FRigVMRegister& Register, int32 InRegisterOffset = INDEX_NONE, int32 InSliceIndex = 0)
+	{
+		return *Get<T>(Register, InRegisterOffset, InSliceIndex);
+	}
+
+	// Returns the current typed reference for a given register index
+	// Note: This refers to the active slice - and can change over time.
+	template<typename T>
+	FORCEINLINE_DEBUGGABLE T& GetRef(int32 InRegisterIndex, int32 InRegisterOffset = INDEX_NONE, int32 InSliceIndex = 0)
+	{
+		return *Get<T>(InRegisterIndex, InRegisterOffset, InSliceIndex);
+	}
+
+	// Returns the current typed reference for a given argument
+	// Note: This refers to the active slice - and can change over time.
+	template<typename T>
+	FORCEINLINE_DEBUGGABLE T& GetRef(const FRigVMOperand& InOperand, int32 InSliceIndex = 0)
+	{
+		return GetRef<T>(InOperand.GetRegisterIndex(), InOperand.GetRegisterOffset(), InSliceIndex);
+	}
+
+	// Returns an array view for all elements of the current slice for a given register.
+	template<typename T>
+	FORCEINLINE_DEBUGGABLE FRigVMFixedArray<T> GetFixedArray(FRigVMRegister& InRegister, int32 InRegisterOffset = INDEX_NONE, int32 InSliceIndex = 0)
+	{
+		if (InRegisterOffset == INDEX_NONE)
+		{
+			uint8* Ptr = (uint8*)GetDataPtr(InRegister, InRegisterOffset, InSliceIndex, false);
+			if (InRegister.IsNestedDynamic())
+			{
+				FRigVMByteArray* Storage = (FRigVMByteArray*)Ptr;
+				return FRigVMFixedArray<T>(*Storage);
+			}
+			else if (InRegister.IsDynamic())
+			{
+				FRigVMByteArray* Storage = (FRigVMByteArray*)Ptr;
+				Ptr = Storage->GetData() + InSliceIndex * InRegister.GetNumBytesPerSlice();
+			}
+
+			return FRigVMFixedArray<T>((T*)Ptr, InRegister.ElementCount);
+		}
+
+		TArray<T>* StoredArray = (TArray<T>*)GetData(InRegister, InRegisterOffset, InSliceIndex);
+		return FRigVMFixedArray<T>(StoredArray->GetData(), StoredArray->Num());
+	}
+
+	// Returns an array view for all elements of the current slice for a given register index.
+	template<typename T>
+	FORCEINLINE_DEBUGGABLE FRigVMFixedArray<T> GetFixedArray(int32 InRegisterIndex, int32 InRegisterOffset = INDEX_NONE, int32 InSliceIndex = 0)
+	{
+		ensure(Registers.IsValidIndex(InRegisterIndex));
+		FRigVMRegister& Register = Registers[InRegisterIndex];
+		return GetFixedArray<T>(Register, InRegisterOffset, InSliceIndex);
+	}
+	
+	// Returns an array view for all elements of the current slice for a given argument.
+	template<typename T>
+	FORCEINLINE_DEBUGGABLE FRigVMFixedArray<T> GetFixedArray(const FRigVMOperand& InOperand, int32 InSliceIndex = 0)
+	{
+		return GetFixedArray<T>(InOperand.GetRegisterIndex(), InOperand.GetRegisterOffset(), InSliceIndex);
+	}
+
+	// Returns an array view for all elements of the current slice for a given register.
+	template<typename T>
+	FORCEINLINE_DEBUGGABLE FRigVMDynamicArray<T> GetDynamicArray(FRigVMRegister& InRegister, int32 InSliceIndex = 0)
+	{
+		if (!InRegister.IsDynamic())
+		{
+			return FRigVMDynamicArray<T>(DefaultByteArray);
+		}
+		
+		FRigVMByteArray* ArrayStorage = (FRigVMByteArray*)GetDataPtr(InRegister, INDEX_NONE, 0, false);
+		return FRigVMDynamicArray<T>(*ArrayStorage);
+	}
+
+	// Returns an array view for all elements of the current slice for a given register index.
+	template<typename T>
+	FORCEINLINE_DEBUGGABLE FRigVMDynamicArray<T> GetDynamicArray(int32 InRegisterIndex, int32 InSliceIndex = 0)
+	{
+		ensure(Registers.IsValidIndex(InRegisterIndex));
+		FRigVMRegister& Register = Registers[InRegisterIndex];
+		return GetDynamicArray<T>(Register, InSliceIndex);
+	}
+
+	// Returns an array view for all elements of the current slice for a given argument.
+	template<typename T>
+	FORCEINLINE_DEBUGGABLE FRigVMDynamicArray<T> GetDynamicArray(const FRigVMOperand& InOperand, int32 InSliceIndex = 0)
+	{
+		return GetDynamicArray<T>(InOperand.GetRegisterIndex(), InOperand.GetRegisterOffset(), InSliceIndex);
+	}
+
+	// Returns the script struct used for a given register (can be nullptr for non-struct-registers).
+	FORCEINLINE_DEBUGGABLE UScriptStruct* GetScriptStruct(const FRigVMRegister& Register) const
+	{
+		if (Register.ScriptStructIndex != INDEX_NONE)
+		{
+			ensure(ScriptStructs.IsValidIndex(Register.ScriptStructIndex));
+			return ScriptStructs[Register.ScriptStructIndex];
+		}
+		return nullptr;
+	}
+
+	// Returns the script struct used for a given register index (can be nullptr for non-struct-registers).
+	FORCEINLINE_DEBUGGABLE UScriptStruct* GetScriptStruct(int32 InRegisterIndex, int32 InRegisterOffset = INDEX_NONE) const
+	{
+		if (InRegisterOffset == INDEX_NONE)
+		{
+			ensure(Registers.IsValidIndex(InRegisterIndex));
+			const FRigVMRegister& Register = Registers[InRegisterIndex];
+			return GetScriptStruct(Register);
+		}
+		ensure(RegisterOffsets.IsValidIndex(InRegisterOffset));
+		const FRigVMRegisterOffset& Path = RegisterOffsets[InRegisterOffset];
+		return Path.GetScriptStruct();
+	}
+
+	// Copies the content of a source register to a target register.
+	// The source register can optionally be referencing a specific source memory container.
+	bool Copy(
+		int32 InSourceRegisterIndex,
+		int32 InTargetRegisterIndex,
+		const FRigVMMemoryContainer* InSourceMemory = nullptr,
+		int32 InSourceRegisterOffset = INDEX_NONE,
+		int32 InTargetRegisterOffset = INDEX_NONE,
+		int32 InSourceSliceIndex = 0,
+		int32 InTargetSliceIndex = 0);
+
+	// Copies the content of a source register to a target register.
+	// The source register can optionally be referencing a specific source memory container.
+	// Note: This only works if SupportsNames() == true
+	bool Copy(
+		const FName& InSourceName,
+		const FName& InTargetName,
+		const FRigVMMemoryContainer* InSourceMemory = nullptr,
+		int32 InSourceRegisterOffset = INDEX_NONE,
+		int32 InTargetRegisterOffset = INDEX_NONE,
+		int32 InSourceSliceIndex = 0,
+		int32 InTargetSliceIndex = 0);
+
+	// Copies the content of a source register to a target register.
+	// The source register can optionally be referencing a specific source memory container.
+	bool Copy(
+		const FRigVMOperand& InSourceOperand,
+		const FRigVMOperand& InTargetOperand,
+		const FRigVMMemoryContainer* InSourceMemory = nullptr,
+		int32 InSourceSliceIndex = 0,
+		int32 InTargetSliceIndex = 0);
+
+	// Returns the index of a register based on the register name.
+	// Note: This only works if SupportsNames() == true
+	FORCEINLINE_DEBUGGABLE int32 GetIndex(const FName& InName) const
+	{
+		if (!bUseNameMap)
+		{
+			return INDEX_NONE;
+		}
+
+		if (NameMap.Num() != Registers.Num())
+		{
+			for (int32 Index = 0; Index < Registers.Num(); Index++)
+			{
+				if (Registers[Index].Name == InName)
+				{
+					return Index;
+				}
+			}
+		}
+		else
+		{
+			const int32* Index = NameMap.Find(InName);
+			if (Index != nullptr)
+			{
+				return *Index;
+			}
+		}
+
+		return INDEX_NONE;
+	}
+
+	// Returns true if a given name is available for a new register.
+	FORCEINLINE_DEBUGGABLE bool IsNameAvailable(const FName& InPotentialNewName) const
+	{
+		if (!bUseNameMap)
+		{
+			return false;
+		}
+		return GetIndex(InPotentialNewName) == INDEX_NONE;
+	}
+
+	// Adds a new named register for a typed array from an array view (used by compiler)
+	template<typename T>
+	FORCEINLINE_DEBUGGABLE int32 AddFixedArray(const FName& InNewName, const FRigVMFixedArray<T>& InArrayView, int32 InSliceCount = 1)
+	{
+		return AddRegisterArray<T>(true, InNewName, InArrayView.Num(), true, (const uint8*)InArrayView.GetData(), InSliceCount);
+	}
+
+	// Adds a new unnamed register for a typed array from an array view.
+	template<typename T>
+	FORCEINLINE_DEBUGGABLE int32 AddFixedArray(const FRigVMFixedArray<T>& InArrayView, int32 InSliceCount = 1)
+	{
+		return AddFixedArray<T>(NAME_None, InArrayView, InSliceCount);
+	}
+
+	// Adds a new named register for a typed array from an array view (used by compiler)
+	template<typename T>
+	FORCEINLINE_DEBUGGABLE int32 AddDynamicArray(const FName& InNewName, const FRigVMFixedArray<T>& InArrayView, int32 InSliceCount = 1)
+	{
+		return AddRegisterArray<T>(false, InNewName, InArrayView.Num(), true, (const uint8*)InArrayView.GetData(), InSliceCount);
+	}
+
+	// Adds a new unnamed register for a typed array from an array view.
+	template<typename T>
+	FORCEINLINE_DEBUGGABLE int32 AddDynamicArray(const FRigVMFixedArray<T>& InArrayView, int32 InSliceCount = 1)
+	{
+		return AddDynamicArray<T>(NAME_None, InArrayView, InSliceCount);
+	}
+
+	// Adds a new named register for a typed array from an array view (used by compiler)
+	template<typename T>
+	FORCEINLINE_DEBUGGABLE int32 AddDynamicArray(const FName& InNewName, const TArray<T>& InArray, int32 InSliceCount = 1)
+	{
+		return AddDynamicArray<T>(InNewName, FRigVMFixedArray<T>(InArray), InSliceCount);
+	}
+
+	// Adds a new unnamed register for a typed array from an array view.
+	template<typename T>
+	FORCEINLINE_DEBUGGABLE int32 AddDynamicArray(const TArray<T>& InArray, int32 InSliceCount = 1)
+	{
+		return AddDynamicArray<T>(NAME_None, InArray, InSliceCount);
+	}
+
+	// Adds a new named register for a typed value from a value reference.
+	template<typename T>
+	FORCEINLINE_DEBUGGABLE int32 Add(const FName& InNewName, const T& InValue, int32 InSliceCount = 1)
+	{
+		return AddRegisterArray<T>(true, InNewName, 1, false, (const uint8*)&InValue, InSliceCount);
+	}
+
+	// Adds a new unnamed register for a typed value from a value reference.
+	template<typename T>
+	FORCEINLINE_DEBUGGABLE int32 Add(const T& InValue, int32 InSliceCount = 1)
+	{
+		return Add<T>(NAME_None, InValue, InSliceCount);
+	}
+
+	// Adds a new named register for a typed value from a value reference.
+	template<typename T>
+	FORCEINLINE_DEBUGGABLE int32 AddDynamicValue(const FName& InNewName, const T& InValue, int32 InSliceCount = 1)
+	{
+		return AddRegisterArray<T>(false, InNewName, 1, false, (const uint8*)&InValue, InSliceCount);
+	}
+
+	// Adds a new unnamed register for a typed value from a value reference.
+	template<typename T>
+	FORCEINLINE_DEBUGGABLE int32 AddDynamicValue(const T& InValue, int32 InSliceCount = 1)
+	{
+		return AddDynamicValue<T>(NAME_None, InValue, InSliceCount);
+	}
+
+	// Adds a register path and returns its index
+	int32 GetOrAddRegisterOffset(int32 InRegisterIndex, int32 InArrayElement = 0);
+
+	// Adds a register path and returns its index
+	int32 GetOrAddRegisterOffset(int32 InRegisterIndex, const FString& InSegmentPath, int32 InArrayElement = 0);
+
+	// Adds a register path and returns its index
+	int32 GetOrAddRegisterOffset(int32 InRegisterIndex, UScriptStruct* InScriptStruct, const FString& InSegmentPath, int32 InInitialOffset = 0, int32 InElementSize = 0);
+
+	void SetRegisterValueFromString(const FRigVMOperand& InOperand, const FString& InCPPType, const UObject* InCPPTypeObject, const TArray<FString>& InDefaultValues);
+	TArray<FString> GetRegisterValueAsString(const FRigVMOperand& InOperand, const FString& InCPPType, const UObject* InCPPTypeObject);
+
+	// Returns the register offset given an index
+	const FRigVMRegisterOffset& GetRegisterOffset(int32 InRegisterOffsetIndex) const;
+
+	// Returns the register offset for a given operand
+	const FRigVMRegisterOffset& GetRegisterOffsetForOperand(const FRigVMOperand& InOperand) const;
+
+	// returns the statistics information
+	FRigVMMemoryStatistics GetStatistics() const
+	{
+		FRigVMMemoryStatistics Statistics;
+		Statistics.RegisterCount = Registers.Num();
+		Statistics.DataBytes = Data.GetAllocatedSize();
+		Statistics.TotalBytes = Data.GetAllocatedSize() + Registers.GetAllocatedSize() + RegisterOffsets.GetAllocatedSize();
+		return Statistics;
+	}
+
+private:
+
+	// Copies the source memory into a known register
+	bool Copy(
+		int32 InTargetRegisterIndex,
+		int32 InTargetRegisterOffset,
+		ERigVMRegisterType InTargetType,
+		const uint8* InSourcePtr,
+		uint8* InTargetPtr,
+		uint16 InNumBytes);
+
+	template<
+		typename T,
+		typename TEnableIf<TIsArithmetic<T>::Value>::Type* = nullptr
+	>
+	FORCEINLINE_DEBUGGABLE int32 AddRegisterArray(bool bFixed, const FName& InNewName, int32 InCount, bool bIsArrayPerSlice, const uint8* InDataPtr, int32 InSliceCount)
+	{
+		return AddRegisterArray<T>(bFixed, InNewName, InCount, bIsArrayPerSlice, InDataPtr, InSliceCount, ERigVMRegisterType::Plain, nullptr);
+	}
+
+	template<
+		typename T,
+		typename TEnableIf<TRigVMIsName<T>::Value>::Type* = nullptr
+	>
+	FORCEINLINE_DEBUGGABLE int32 AddRegisterArray(bool bFixed, const FName& InNewName, int32 InCount, bool bIsArrayPerSlice, const uint8* InDataPtr, int32 InSliceCount)
+	{
+		return AddRegisterArray<T>(bFixed, InNewName, InCount, bIsArrayPerSlice, InDataPtr, InSliceCount, ERigVMRegisterType::Name, nullptr);
+	}
+
+	template<
+		typename T,
+		typename TEnableIf<TRigVMIsString<T>::Value>::Type* = nullptr
+	>
+	FORCEINLINE_DEBUGGABLE int32 AddRegisterArray(bool bFixed, const FName& InNewName, int32 InCount, bool bIsArrayPerSlice, const uint8* InDataPtr, int32 InSliceCount)
+	{
+		return AddRegisterArray<T>(bFixed, InNewName, InCount, bIsArrayPerSlice, InDataPtr, InSliceCount, ERigVMRegisterType::String, nullptr);
+	}
+
+	template<
+		typename T,
+		typename TEnableIf<TModels<CRigVMUStruct, T>::Value>::Type* = nullptr
+	>
+	FORCEINLINE_DEBUGGABLE int32 AddRegisterArray(bool bFixed, const FName& InNewName, int32 InCount, bool bIsArrayPerSlice, const uint8* InDataPtr, int32 InSliceCount)
+	{
+		return AddRegisterArray<T>(bFixed, InNewName, InCount, bIsArrayPerSlice, InDataPtr, InSliceCount, ERigVMRegisterType::Struct, T::StaticStruct());
+	}
+
+	template<
+		typename T,
+		typename TEnableIf<TIsEnum<T>::Value>::Type* = nullptr
+	>
+	FORCEINLINE_DEBUGGABLE int32 AddRegisterArray(bool bFixed, const FName& InNewName, int32 InCount, bool bIsArrayPerSlice, const uint8* InDataPtr, int32 InSliceCount)
+	{
+		return AddRegisterArray<T>(bFixed, InNewName, InCount, bIsArrayPerSlice, InDataPtr, InSliceCount, ERigVMRegisterType::Plain, nullptr);
+	}
+
+	template <
+		typename T,
+		typename TEnableIf<TRigVMIsBaseStructure<T>::Value, T>::Type* = nullptr
+	>
+	FORCEINLINE_DEBUGGABLE int32 AddRegisterArray(bool bFixed, const FName& InNewName, int32 InCount, bool bIsArrayPerSlice, const uint8* InDataPtr, int32 InSliceCount)
+	{
+		return AddRegisterArray<T>(bFixed, InNewName, InCount, bIsArrayPerSlice, InDataPtr, InSliceCount, ERigVMRegisterType::Plain, TBaseStructure<T>::Get());
+	}
+
+	template<typename T>
+	FORCEINLINE_DEBUGGABLE int32 AddRegisterArray(bool bFixed, const FName& InNewName, int32 InCount, bool bIsArrayPerSlice, const uint8* InDataPtr, int32 InSliceCount, ERigVMRegisterType InType, UScriptStruct* InScriptStruct)
+	{
+		if(bFixed)
+		{
+			return AddFixedArray(InNewName, sizeof(T), InCount, bIsArrayPerSlice, InDataPtr, InSliceCount, InType, InScriptStruct);
+		}
+		else
+		{
+			return AddDynamicArray(InNewName, sizeof(T), InCount, bIsArrayPerSlice, InDataPtr, InSliceCount, InType, InScriptStruct);
+		}
+	}
+
+	FORCEINLINE_DEBUGGABLE int32 AddRegisterArray(bool bFixed, const FName& InNewName, int32 InElementSize, int32 InCount, bool bIsArrayPerSlice, const uint8* InDataPtr, int32 InSliceCount, ERigVMRegisterType InType, UScriptStruct* InScriptStruct)
+	{
+		if (bFixed)
+		{
+			return AddFixedArray(InNewName, InElementSize, InCount, bIsArrayPerSlice, InDataPtr, InSliceCount, InType, InScriptStruct);
+		}
+		else
+		{
+			return AddDynamicArray(InNewName, InElementSize, InCount, bIsArrayPerSlice, InDataPtr, InSliceCount, InType, InScriptStruct);
+		}
+	}
+
+	// Adds a new named register for a fixed array from a data pointer (used by compiler)
+	FORCEINLINE_DEBUGGABLE int32 AddFixedArray(const FName& InNewName, int32 InElementSize, int32 InCount, bool bIsArrayPerSlice, const uint8* InDataPtr, int32 InSliceCount, ERigVMRegisterType InType, UScriptStruct* InScriptStruct)	{
+		int32 Register = Allocate(InNewName, InElementSize, InCount, InSliceCount, nullptr, false);
+		if (!Registers.IsValidIndex(Register))
+		{
+			return Register;
+		}
+
+		Registers[Register].Type = InType;
+		Registers[Register].ScriptStructIndex = FindOrAddScriptStruct(InScriptStruct);
+		Registers[Register].bIsArray = bIsArrayPerSlice;
+
+		UpdateRegisters();
+		Construct(Register);
+
+		if (InDataPtr)
+		{
+			for (int32 SliceIndex = 0; SliceIndex < InSliceCount; SliceIndex++)
+			{
+				Copy(Register, INDEX_NONE, InType, InDataPtr, GetData(Registers[Register], INDEX_NONE, SliceIndex), InElementSize * InCount);
+			}
+		}
+		return Register;
+	}
+
+	// Adds a new named register for a dynamic array from a data pointer (used by compiler)
+	FORCEINLINE_DEBUGGABLE int32 AddDynamicArray(const FName& InNewName, int32 InElementSize, int32 InCount, bool bIsArrayPerSlice, const uint8* InDataPtr, int32 InSliceCount, ERigVMRegisterType InType, UScriptStruct* InScriptStruct)
+	{
+		int32 Register = INDEX_NONE;
+		
+		if (bIsArrayPerSlice)
+		{
+			Register = Allocate(InNewName, sizeof(FRigVMNestedByteArray), 1, 1, nullptr);
+			if (!Registers.IsValidIndex(Register))
+			{
+				return Register;
+			}
+
+			Registers[Register].Type = InType;
+			Registers[Register].ScriptStructIndex = FindOrAddScriptStruct(InScriptStruct);
+			Registers[Register].bIsDynamic = true;
+			Registers[Register].bIsArray = bIsArrayPerSlice;
+			Registers[Register].ElementSize = InElementSize;
+			Registers[Register].ElementCount = InCount;
+			Registers[Register].SliceCount = InSliceCount;
+
+			uint8* Ptr = (uint8*)&Data[Registers[Register].GetWorkByteIndex()];
+			FRigVMNestedByteArray& Storage = *(FRigVMNestedByteArray*)Ptr;
+			Storage.SetNum(InSliceCount);
+
+			for (int32 SliceIndex = 0; SliceIndex < InSliceCount; SliceIndex++)
+			{
+				Storage[SliceIndex].SetNumZeroed(InCount * InElementSize);
+				if (InDataPtr)
+				{
+					Copy(Register, INDEX_NONE, InType, InDataPtr, Storage[SliceIndex].GetData(), InElementSize * InCount);
+				}
+			}
+
+			if (InDataPtr == nullptr)
+			{
+				Construct(Register);
+			}
+		}
+		else
+		{
+			Register = Allocate(InNewName, sizeof(FRigVMByteArray), 1, 1, nullptr);
+			if (!Registers.IsValidIndex(Register))
+			{
+				return Register;
+			}
+
+			Registers[Register].Type = InType;
+			Registers[Register].ScriptStructIndex = FindOrAddScriptStruct(InScriptStruct);
+			Registers[Register].bIsDynamic = true;
+			Registers[Register].bIsArray = bIsArrayPerSlice;
+			Registers[Register].ElementSize = InElementSize;
+			Registers[Register].ElementCount = InCount;
+			Registers[Register].SliceCount = InSliceCount;
+
+			uint8* Ptr = (uint8*)&Data[Registers[Register].GetWorkByteIndex()];
+			FRigVMByteArray& Storage = *(FRigVMByteArray*)Ptr;
+
+			Storage.SetNumZeroed(InSliceCount * InElementSize);
+			if (InDataPtr)
+			{
+				for (int32 SliceIndex = 0; SliceIndex < InSliceCount; SliceIndex++)
+				{
+					Copy(Register, INDEX_NONE, InType, InDataPtr, &Storage[SliceIndex * InElementSize], InElementSize);
+				}
+			}
+			else
+			{
+				Construct(Register);
+			}
+		}
+
+		return Register;
+	}
+
+	// Updates internal data for topological changes
+	void UpdateRegisters();
+
+	// Allocates a new named register
+	int32 Allocate(const FName& InNewName, int32 InElementSize, int32 InElementCount, int32 InSliceCount, const uint8* InDataPtr = nullptr, bool bUpdateRegisters = true);
+
+	// Allocates a new unnamed register
+	int32 Allocate(int32 InElementSize, int32 InElementCount, int32 InSliceCount, const uint8* InDataPtr = nullptr, bool bUpdateRegisters = true);
+
+	// Performs optional construction of data within a struct register
+	bool Construct(int32 InRegisterIndex, int32 InElementIndex = INDEX_NONE, int32 InSliceIndex = 0);
+
+	// Performs optional destruction of data within a struct register
+	bool Destroy(int32 InRegisterIndex, int32 InElementIndex = INDEX_NONE, int32 InSliceIndex = 0);
+
+	// Fills a register with zero memory
+	void FillWithZeroes(int32 InRegisterIndex);
+
+	// Ensures to add a script struct to the internal map if needed
+	int32 FindOrAddScriptStruct(UScriptStruct* InScriptStruct);
+
+#endif
+	
+	UPROPERTY()
+	bool bUseNameMap;
+
+	UPROPERTY()
+	ERigVMMemoryType MemoryType;
+
+	UPROPERTY()
+	TArray<FRigVMRegister> Registers;
+
+	UPROPERTY()
+	TArray<FRigVMRegisterOffset> RegisterOffsets;
+
+	UPROPERTY(transient)
+	TArray<uint8> Data;
+
+	UPROPERTY(transient)
+	TArray<TObjectPtr<UScriptStruct>> ScriptStructs;
+
+	UPROPERTY(transient)
+	TMap<FName, int32> NameMap;
+
+	UPROPERTY(transient)
+	bool bEncounteredErrorDuringLoad;
+
+#if UE_RIGVM_UCLASS_BASED_STORAGE_DISABLED
+
+	static FRigVMByteArray DefaultByteArray;
+	static FRigVMRegisterOffset InvalidRegisterOffset;
+
+	friend class URigVM;
+	friend class URigVMCompiler;
+	friend struct FRigVMCompilerWorkData;
+
+#endif
+};
+
+#if UE_RIGVM_UCLASS_BASED_STORAGE_DISABLED
+
+typedef FRigVMMemoryContainer* FRigVMMemoryContainerPtr;
+typedef FRigVMFixedArray<FRigVMMemoryContainerPtr> FRigVMMemoryContainerPtrArray; 
+
+#endif
+
diff --git a/Engine/Source/Runtime/RigVM/Public/RigVMCore/RigVMMemoryStorage.h b/Engine/Source/Runtime/RigVM/Public/RigVMCore/RigVMMemoryStorage.h
new file mode 100644
--- /dev/null
+++ b/Engine/Source/Runtime/RigVM/Public/RigVMCore/RigVMMemoryStorage.h
@@ -0,0 +1,55 @@
+// Copyright Epic Games, Inc. All Rights Reserved.
+
+#pragma once
+
+#include "CoreMinimal.h"
+#include "UObject/ObjectMacros.h"
+#include "RigVMTraits.h"
+#include "RigVMStatistics.h"
+#include "RigVMArray.h"
+#include "RigVMMemoryCommon.h"
+#include "RigVMMemoryStorage.generated.h"
+
+#if !UE_RIGVM_UCLASS_BASED_STORAGE_DISABLED
+
+/**
+ * The FRigVMMemoryHandle is used to access the memory used within a URigMemoryStorage.
+ */
+struct FRigVMMemoryHandle
+{
+public:
+
+	FORCEINLINE_DEBUGGABLE FRigVMMemoryHandle()
+		: Ptr(nullptr)
+	{}
+
+private:
+	
+	uint8* Ptr;
+
+	friend class URigVM;
+};
+
+#endif
+
+UCLASS()
+class RIGVM_API URigVMMemoryStorage : public UObject
+{
+	GENERATED_BODY()
+
+public:
+	
+	struct FPropertyDescription
+	{
+		FName Name;
+		const FProperty* Property;
+		FEdGraphPinType PinType;
+		FString DefaultValue;
+		
+		FPropertyDescription(const FProperty* InProperty, const FString& InDefaultValue, const FName& InName = NAME_None);
+		FPropertyDescription(const FName& InName, const FEdGraphPinType& InPinType, const FString& InDefaultValue);
+	};
+
+	static UClass* CreateStorageClass(UObject* InOuter, const TArray<FPropertyDescription>& InProperties);
+	static URigVMMemoryStorage* CreateStorage(UObject* InOuter, const TArray<FPropertyDescription>& InProperties);
+};
diff --git a/Engine/Source/Runtime/RigVM/Public/RigVMDefines.h b/Engine/Source/Runtime/RigVM/Public/RigVMDefines.h
new file mode 100644
--- /dev/null
+++ b/Engine/Source/Runtime/RigVM/Public/RigVMDefines.h
@@ -0,0 +1,11 @@
+// Copyright Epic Games, Inc. All Rights Reserved.
+
+/*=============================================================================
+	RigVMDefines.h: Global defines
+=============================================================================*/
+
+#pragma once
+
+#ifndef UE_RIGVM_UCLASS_BASED_STORAGE_DISABLED
+#define UE_RIGVM_UCLASS_BASED_STORAGE_DISABLED 1
+#endif
diff --git a/Engine/Source/Runtime/RigVM/Public/RigVMModule.h b/Engine/Source/Runtime/RigVM/Public/RigVMModule.h
--- a/Engine/Source/Runtime/RigVM/Public/RigVMModule.h
+++ b/Engine/Source/Runtime/RigVM/Public/RigVMModule.h
@@ -1,11 +1,12 @@
 // Copyright Epic Games, Inc. All Rights Reserved.
 
 /*=============================================================================
-	RigVM.h: Module definitions.
+	RigVMModule.h: Module definitions.
 =============================================================================*/
 
 #pragma once
 
 #include "CoreMinimal.h"
+#include "RigVMDefines.h"
 
 DECLARE_LOG_CATEGORY_EXTERN(LogRigVM, Log, All);
\ No newline at end of file
