RigVM: Fix non-unity build breaks - follow up to 16906921

diff --git a/Engine/Source/Runtime/RigVM/Public/RigVMCore/RigVMArray.h b/Engine/Source/Runtime/RigVM/Public/RigVMCore/RigVMArray.h
--- a/Engine/Source/Runtime/RigVM/Public/RigVMCore/RigVMArray.h
+++ b/Engine/Source/Runtime/RigVM/Public/RigVMCore/RigVMArray.h
@@ -3,6 +3,7 @@
 #pragma once
 
 #include "CoreMinimal.h"
+#include "RigVMDefines.h"
 #include "RigVMTraits.h"
 
 #if UE_RIGVM_UCLASS_BASED_STORAGE_DISABLED
diff --git a/Engine/Source/Runtime/RigVM/Public/RigVMCore/RigVMByteCode.h b/Engine/Source/Runtime/RigVM/Public/RigVMCore/RigVMByteCode.h
--- a/Engine/Source/Runtime/RigVM/Public/RigVMCore/RigVMByteCode.h
+++ b/Engine/Source/Runtime/RigVM/Public/RigVMCore/RigVMByteCode.h
@@ -3,6 +3,7 @@
 #pragma once
 
 #include "CoreMinimal.h"
+#include "RigVMDefines.h"
 #include "RigVMRegistry.h"
 #include "RigVMStatistics.h"
 
diff --git a/Engine/Source/Runtime/RigVM/Public/RigVMCore/RigVMExecuteContext.h b/Engine/Source/Runtime/RigVM/Public/RigVMCore/RigVMExecuteContext.h
--- a/Engine/Source/Runtime/RigVM/Public/RigVMCore/RigVMExecuteContext.h
+++ b/Engine/Source/Runtime/RigVM/Public/RigVMCore/RigVMExecuteContext.h
@@ -3,6 +3,7 @@
 #pragma once
 
 #include "CoreMinimal.h"
+#include "RigVMDefines.h"
 #include "RigVMArray.h"
 #include "RigVMExternalVariable.h"
 #include "RigVMExecuteContext.generated.h"
diff --git a/Engine/Source/Runtime/RigVM/Public/RigVMCore/RigVMMemoryCommon.h b/Engine/Source/Runtime/RigVM/Public/RigVMCore/RigVMMemoryCommon.h
--- a/Engine/Source/Runtime/RigVM/Public/RigVMCore/RigVMMemoryCommon.h
+++ b/Engine/Source/Runtime/RigVM/Public/RigVMCore/RigVMMemoryCommon.h
@@ -4,6 +4,7 @@
 
 #include "CoreMinimal.h"
 #include "UObject/ObjectMacros.h"
+#include "RigVMDefines.h"
 #include "RigVMTraits.h"
 #include "RigVMArray.h"
 #include "RigVMMemoryCommon.generated.h"
diff --git a/Engine/Source/Runtime/RigVM/Public/RigVMCore/RigVMTraits.h b/Engine/Source/Runtime/RigVM/Public/RigVMCore/RigVMTraits.h
--- a/Engine/Source/Runtime/RigVM/Public/RigVMCore/RigVMTraits.h
+++ b/Engine/Source/Runtime/RigVM/Public/RigVMCore/RigVMTraits.h
@@ -4,8 +4,8 @@
 
 #include "CoreMinimal.h"
 #include "UObject/Class.h"
-
 #include "Templates/Models.h"
+#include "RigVMDefines.h"
 
 struct CRigVMUStruct
 {
