RigVM : URigVMMemoryStorage prep work + first unit test

diff --git a/Engine/Source/Runtime/RigVM/Private/RigVMCore/RigVMMemoryStorage.cpp b/Engine/Source/Runtime/RigVM/Private/RigVMCore/RigVMMemoryStorage.cpp
--- a/Engine/Source/Runtime/RigVM/Private/RigVMCore/RigVMMemoryStorage.cpp
+++ b/Engine/Source/Runtime/RigVM/Private/RigVMCore/RigVMMemoryStorage.cpp
@@ -1,34 +1,389 @@
-// Copyright Epic Games, Inc. All Rights Reserved.
+	// Copyright Epic Games, Inc. All Rights Reserved.
 
 #include "RigVMCore/RigVMMemoryStorage.h"
 
 ////////////////////////////////////////////////////////////////////////////////////////////////////////////////////////
 
+const FString URigVMMemoryStorage::FPropertyDescription::ArrayPrefix = TEXT("TArray<");
+const FString URigVMMemoryStorage::FPropertyDescription::MapPrefix = TEXT("TMap<");
+const FString URigVMMemoryStorage::FPropertyDescription::SetPrefix = TEXT("TSet<");
+const FString URigVMMemoryStorage::FPropertyDescription::ContainerSuffix = TEXT(">");
+
 URigVMMemoryStorage::FPropertyDescription::FPropertyDescription(const FProperty* InProperty, const FString& InDefaultValue, const FName& InName)
 	: Name(InName)
 	, Property(InProperty)
-	, PinType()
+	, CPPType()
+	, ContainerType(EPinContainerType::None)
 	, DefaultValue(InDefaultValue)
 {
+	SanitizeName();
 }
 
-URigVMMemoryStorage::FPropertyDescription::FPropertyDescription(const FName& InName, const FEdGraphPinType& InPinType, const FString& InDefaultValue)
+URigVMMemoryStorage::FPropertyDescription::FPropertyDescription(const FName& InName, const FString& InCPPType, UObject* InCPPTypeObject, const FString& InDefaultValue)
 	: Name(InName)
 	, Property(nullptr)
-	, PinType(InPinType)
+	, CPPType(InCPPType)
+	, CPPTypeObject(InCPPTypeObject)
+	, ContainerType(EPinContainerType::None)
 	, DefaultValue(InDefaultValue)
 {
+	SanitizeName();
+
+	FString BaseCPPType = CPPType;
+	if(BaseCPPType.RemoveFromStart(ArrayPrefix))
+	{
+		ContainerType = EPinContainerType::Array;
+	}
+	else if(BaseCPPType.RemoveFromStart(MapPrefix))
+	{
+		ContainerType = EPinContainerType::Map;
+	}
+	else if(BaseCPPType.RemoveFromStart(SetPrefix))
+	{
+		ContainerType = EPinContainerType::Set;
+	}
+
+	// make sure this description doesn't contain another container
+	if(ContainerType != EPinContainerType::None)
+	{
+		BaseCPPType.RemoveFromEnd(ContainerSuffix);
+		check(!BaseCPPType.Contains(TEXT("<")));
+	}
+}
+
+FName URigVMMemoryStorage::FPropertyDescription::SanitizedName(const FName& InName)
+{
+	FString NameString = InName.ToString();
+
+	// Sanitize the name
+	for (int32 i = 0; i < NameString.Len(); ++i)
+	{
+		TCHAR& C = NameString[i];
+
+		const bool bGoodChar = FChar::IsAlpha(C) ||							// Any letter
+			(C == '_') || 													// _ anytime
+			((i > 0) && (FChar::IsDigit(C)));								// 0-9 after the first character
+
+		if (!bGoodChar)
+		{
+			C = '_';
+		}
+	}
+
+	if(NameString != InName.ToString())
+	{
+		return *NameString;
+	}
+
+	return InName;
+}
+
+void URigVMMemoryStorage::FPropertyDescription::SanitizeName()
+{
+	Name = SanitizedName(Name);
+}
+
+FString URigVMMemoryStorage::FPropertyDescription::GetBaseCPPType() const
+{
+	if(IsArray())
+	{
+		return GetArrayElementCPPType();
+	}
+	if(IsMap())
+	{
+		return GetMapValueCPPType();
+	}
+	if(IsSet())
+	{
+		return GetSetElementCPPType(); 
+	}
+	return CPPType;
+}
+
+FString URigVMMemoryStorage::FPropertyDescription::GetArrayElementCPPType() const
+{
+	check(IsArray());
+
+	FString BaseCPPType = CPPType;
+	check(BaseCPPType.RemoveFromStart(ArrayPrefix));
+	check(BaseCPPType.RemoveFromEnd(ContainerSuffix));
+	return BaseCPPType.TrimStartAndEnd();
+}
+
+FString URigVMMemoryStorage::FPropertyDescription::GetSetElementCPPType() const
+{
+	check(IsSet());
+
+	FString BaseCPPType = CPPType;
+	check(BaseCPPType.RemoveFromStart(SetPrefix));
+	check(BaseCPPType.RemoveFromEnd(ContainerSuffix));
+	return BaseCPPType.TrimStartAndEnd();
+}
+
+FString URigVMMemoryStorage::FPropertyDescription::GetMapKeyCPPType() const
+{
+	check(IsMap());
+
+	FString BaseCPPType = CPPType;
+	check(BaseCPPType.RemoveFromStart(MapPrefix));
+	check(BaseCPPType.RemoveFromEnd(ContainerSuffix));
+
+	int32 Comma = 0;
+	check(BaseCPPType.FindChar(',', Comma));
+
+	return BaseCPPType.Left(Comma).TrimStartAndEnd();
+}
+
+FString URigVMMemoryStorage::FPropertyDescription::GetMapValueCPPType() const
+{
+	check(IsMap());
+
+	FString BaseCPPType = CPPType;
+	check(BaseCPPType.RemoveFromStart(MapPrefix));
+	check(BaseCPPType.RemoveFromEnd(ContainerSuffix));
+
+	int32 Comma = 0;
+	check(BaseCPPType.FindChar(',', Comma));
+
+	return BaseCPPType.Mid(Comma + 1).TrimStartAndEnd();
 }
 
 ////////////////////////////////////////////////////////////////////////////////////////////////////////////////////////
 
-UClass* URigVMMemoryStorage::CreateStorageClass(UObject* InOuter, const TArray<FPropertyDescription>& InProperties)
+FCriticalSection RigVMMemoryStorageClassCriticalSection;
+TMap<FString, URigVMMemoryStorage::FClassInfo> URigVMMemoryStorage::PackageToInfo;
+
+UClass* URigVMMemoryStorage::GetStorageClass(UObject* InOuter, ERigVMMemoryType InMemoryType)
 {
-	return nullptr;
+	check(InOuter);
+	UPackage* Package = InOuter->GetOutermost(); 
+
+	FScopeLock Lock(&RigVMMemoryStorageClassCriticalSection);
+
+	FClassInfo& ClassInfo = PackageToInfo.FindOrAdd(InOuter->GetPathName());
+	UClass** ClassPtr = ClassInfo.GetClassPtr(InMemoryType);
+	return *ClassPtr;
 }
 
-URigVMMemoryStorage* URigVMMemoryStorage::CreateStorage(UObject* InOuter,
-	const TArray<FPropertyDescription>& InProperties)
+UClass* URigVMMemoryStorage::CreateStorageClass(UObject* InOuter, ERigVMMemoryType InMemoryType, const TArray<FPropertyDescription>& InProperties)
 {
-	return nullptr;
-}
\ No newline at end of file
+	check(InOuter);
+	UPackage* Package = InOuter->GetOutermost(); 
+	UClass *SuperClass = URigVMMemoryStorage::StaticClass();
+
+	FScopeLock Lock(&RigVMMemoryStorageClassCriticalSection);
+
+	FClassInfo& ClassInfo = PackageToInfo.FindOrAdd(InOuter->GetPathName());
+	UClass** ClassPtr = ClassInfo.GetClassPtr(InMemoryType);
+	check(ClassPtr);
+	
+	if(*ClassPtr == nullptr)
+	{
+		(*ClassPtr) = NewObject<UClass>(
+			InOuter->GetOutermost(),
+			NAME_None,
+			RF_Public | RF_Transactional
+			);
+
+		(*ClassPtr)->AddToRoot();
+	}
+
+	UClass* Class = *ClassPtr;
+
+	Class->PurgeClass(false);
+	Class->PropertyLink = SuperClass->PropertyLink;
+	Class->SetSuperStruct(SuperClass);
+	Class->ClassWithin = UObject::StaticClass();
+
+	// Generate properties
+	FField** LinkToProperty = &Class->ChildProperties;
+	TArray<FProperty*> GeneratedProperties;
+	for(const FPropertyDescription& PropertyDescription : InProperties)
+	{
+		while (*LinkToProperty != nullptr)
+		{
+			LinkToProperty = &(*LinkToProperty)->Next;
+		}
+		
+		if(PropertyDescription.Property)
+		{
+			FProperty* NewProperty = CastFieldChecked<FProperty>(FField::Duplicate(PropertyDescription.Property, Class, PropertyDescription.Name));
+			check(NewProperty);
+
+			GeneratedProperties.Add(NewProperty);
+			*LinkToProperty = NewProperty;
+		}
+		else
+		{
+			FFieldVariant PropertyOwner = Class;
+			FProperty* InnerProperty = nullptr;
+			FProperty* OuterProperty = nullptr;
+			FProperty** KeyPropertyPtr = nullptr;
+			FProperty** ValuePropertyPtr = &InnerProperty;
+
+			switch(PropertyDescription.ContainerType)
+			{
+				case EPinContainerType::Array:
+				{
+					FArrayProperty* ArrayProperty = new FArrayProperty(PropertyOwner, PropertyDescription.Name, RF_Public);
+					ValuePropertyPtr = &ArrayProperty->Inner;
+					PropertyOwner = OuterProperty = ArrayProperty;
+					break;
+				}
+				case EPinContainerType::Map:
+				{
+					checkNoEntry(); // this is not implemented yet
+					FMapProperty* MapProperty = new FMapProperty(PropertyOwner, PropertyDescription.Name, RF_Public);
+					KeyPropertyPtr = &MapProperty->KeyProp;
+					ValuePropertyPtr = &MapProperty->ValueProp;
+					PropertyOwner = OuterProperty = MapProperty;
+					break;
+				}
+				case EPinContainerType::Set:
+				{
+					FSetProperty* SetProperty = new FSetProperty(PropertyOwner, PropertyDescription.Name, RF_Public);
+					ValuePropertyPtr = &SetProperty->ElementProp;
+					PropertyOwner = OuterProperty = SetProperty;
+					break;
+				}
+				case EPinContainerType::None:
+				default:
+				{
+					break;
+				}
+			}
+
+			if(PropertyDescription.CPPTypeObject != nullptr)
+			{
+				if(UEnum* Enum = Cast<UEnum>(PropertyDescription.CPPTypeObject))
+				{
+					FByteProperty* EnumProperty = new FByteProperty(PropertyOwner, PropertyDescription.Name, RF_Public);
+					EnumProperty->Enum = Enum;
+					(*ValuePropertyPtr) = EnumProperty;
+				}
+				else if(UScriptStruct* ScriptStruct = Cast<UScriptStruct>(PropertyDescription.CPPTypeObject))
+				{
+					FStructProperty* StructProperty = new FStructProperty(PropertyOwner, PropertyDescription.Name, RF_Public);
+					StructProperty->Struct = ScriptStruct;
+					(*ValuePropertyPtr) = StructProperty;
+				}
+				else if(UClass* PropertyClass = Cast<UClass>(PropertyDescription.CPPTypeObject))
+				{
+					FObjectProperty* ObjectProperty = new FObjectProperty(PropertyOwner, PropertyDescription.Name, RF_Public);
+					ObjectProperty->SetPropertyClass(PropertyClass);
+					(*ValuePropertyPtr) = ObjectProperty;
+				}
+				else
+				{
+					checkNoEntry();
+				}
+			}
+			else // take care of default types...
+			{
+				static FString BoolString = TEXT("bool");
+				static FString Int32String = TEXT("int32");
+				static FString IntString = TEXT("int");
+				static FString FloatString = TEXT("float");
+				static FString DoubleString = TEXT("double");
+				static FString StringString = TEXT("FString");
+				static FString NameString = TEXT("FName");
+
+				FString BaseCPPType = PropertyDescription.GetBaseCPPType();
+				if(BaseCPPType.Equals(BoolString, ESearchCase::IgnoreCase))
+				{
+					(*ValuePropertyPtr) = new FBoolProperty(PropertyOwner, PropertyDescription.Name, RF_Public);;
+				}
+				else if(BaseCPPType.Equals(Int32String, ESearchCase::IgnoreCase) ||
+					BaseCPPType.Equals(IntString, ESearchCase::IgnoreCase))
+				{
+					(*ValuePropertyPtr) = new FIntProperty(PropertyOwner, PropertyDescription.Name, RF_Public);;
+				}
+				else if(BaseCPPType.Equals(FloatString, ESearchCase::IgnoreCase))
+				{
+					(*ValuePropertyPtr) = new FFloatProperty(PropertyOwner, PropertyDescription.Name, RF_Public);;
+				}
+				else if(BaseCPPType.Equals(DoubleString, ESearchCase::IgnoreCase))
+				{
+					(*ValuePropertyPtr) = new FDoubleProperty(PropertyOwner, PropertyDescription.Name, RF_Public);;
+				}
+				else if(BaseCPPType.Equals(StringString, ESearchCase::IgnoreCase))
+				{
+					(*ValuePropertyPtr) = new FStrProperty(PropertyOwner, PropertyDescription.Name, RF_Public);;
+				}
+				else if(BaseCPPType.Equals(NameString, ESearchCase::IgnoreCase))
+				{
+					(*ValuePropertyPtr) = new FNameProperty(PropertyOwner, PropertyDescription.Name, RF_Public);;
+				}
+				else
+				{
+					checkNoEntry();
+				}
+			}
+
+			if(OuterProperty)
+			{
+				GeneratedProperties.Add(OuterProperty);
+				(*LinkToProperty) = OuterProperty;
+			}
+			else
+			{
+				check(*ValuePropertyPtr);
+				GeneratedProperties.Add(*ValuePropertyPtr);
+				(*LinkToProperty) = *ValuePropertyPtr;
+			}
+		}
+	}
+
+	// Update the class
+	Class->Bind();
+	Class->StaticLink(true);
+
+	// Similar to FConfigPropertyHelperDetails::CustomizeDetails, this is required for GC to work properly
+	Class->AssembleReferenceTokenStream();
+
+	check(GeneratedProperties.Num() == InProperties.Num());
+	
+	// Create default object and store default values.
+	URigVMMemoryStorage* CDO = Cast<URigVMMemoryStorage>(Class->GetDefaultObject(true));
+	for(int32 PropertyIndex = 0; PropertyIndex < GeneratedProperties.Num(); PropertyIndex++)
+	{
+		const FString& DefaultValue = InProperties[PropertyIndex].DefaultValue;
+		if(DefaultValue.IsEmpty())
+		{
+			continue;
+		}
+
+		FProperty* Property = GeneratedProperties[PropertyIndex];
+		uint8* ValuePtr = Property->ContainerPtrToValuePtr<uint8>(CDO);
+
+		Property->ImportText(*DefaultValue, ValuePtr, EPropertyPortFlags::PPF_None, nullptr);
+	}
+	
+	return Class;
+}
+
+URigVMMemoryStorage* URigVMMemoryStorage::CreateStorage(UObject* InOuter, ERigVMMemoryType InMemoryType)
+{
+	UClass* Class = GetStorageClass(InOuter, InMemoryType);
+
+	if(InMemoryType == ERigVMMemoryType::Literal)
+	{
+		return Cast<URigVMMemoryStorage>(Class->GetDefaultObject(true));
+	}
+
+	return NewObject<URigVMMemoryStorage>(InOuter, Class, NAME_None, RF_Public | RF_Transactional);
+}
+
+void URigVMMemoryStorage::RefreshCache()
+{
+	Cache.Reset();
+
+	FField** LinkToProperty = &GetClass()->ChildProperties;
+
+	while (*LinkToProperty != nullptr)
+	{
+		FProperty* Property = CastField<FProperty>(*LinkToProperty);
+		Cache.Add(Property->ContainerPtrToValuePtr<uint8>(this));
+		LinkToProperty = &(*LinkToProperty)->Next;
+	}
+}
diff --git a/Engine/Source/Runtime/RigVM/Public/RigVMCore/RigVMMemoryStorage.h b/Engine/Source/Runtime/RigVM/Public/RigVMCore/RigVMMemoryStorage.h
--- a/Engine/Source/Runtime/RigVM/Public/RigVMCore/RigVMMemoryStorage.h
+++ b/Engine/Source/Runtime/RigVM/Public/RigVMCore/RigVMMemoryStorage.h
@@ -38,18 +38,90 @@ class RIGVM_API URigVMMemoryStorage : public UObject
 	GENERATED_BODY()
 
 public:
-	
-	struct FPropertyDescription
+
+	//////////////////////////////////////////////////////////////////////////////
+	/// Property Management
+	//////////////////////////////////////////////////////////////////////////////
+	struct RIGVM_API FPropertyDescription
 	{
 		FName Name;
 		const FProperty* Property;
-		FEdGraphPinType PinType;
+		FString CPPType;
+		UObject* CPPTypeObject;
+		EPinContainerType ContainerType;
 		FString DefaultValue;
 		
 		FPropertyDescription(const FProperty* InProperty, const FString& InDefaultValue, const FName& InName = NAME_None);
-		FPropertyDescription(const FName& InName, const FEdGraphPinType& InPinType, const FString& InDefaultValue);
+		FPropertyDescription(const FName& InName, const FString& InCPPType, UObject* InCPPTypeObject, const FString& InDefaultValue);
+
+		static FName SanitizedName(const FName& InName);
+		void SanitizeName();
+
+		bool HasContainer() const { return ContainerType != EPinContainerType::None; }
+		bool IsArray() const { return ContainerType == EPinContainerType::Array; }
+		bool IsMap() const { return ContainerType == EPinContainerType::Map; }
+		bool IsSet() const { return ContainerType == EPinContainerType::Set; }
+
+		FString GetBaseCPPType() const;
+		FString GetArrayElementCPPType() const;
+		FString GetSetElementCPPType() const;
+		FString GetMapKeyCPPType() const;
+		FString GetMapValueCPPType() const;
+
+		static const FString ArrayPrefix;
+		static const FString MapPrefix;
+		static const FString SetPrefix;
+		static const FString ContainerSuffix;
 	};
 
-	static UClass* CreateStorageClass(UObject* InOuter, const TArray<FPropertyDescription>& InProperties);
-	static URigVMMemoryStorage* CreateStorage(UObject* InOuter, const TArray<FPropertyDescription>& InProperties);
+	static UClass* GetStorageClass(UObject* InOuter, ERigVMMemoryType InMemoryType);
+	static UClass* CreateStorageClass(UObject* InOuter, ERigVMMemoryType InMemoryType, const TArray<FPropertyDescription>& InProperties);
+	static URigVMMemoryStorage* CreateStorage(UObject* InOuter, ERigVMMemoryType InMemoryType);
+
+	//////////////////////////////////////////////////////////////////////////////
+	/// Memory Access
+	//////////////////////////////////////////////////////////////////////////////
+
+protected:
+
+	void RefreshCache();
+	TArray<uint8*> Cache;
+
+private:
+	
+	struct FClassInfo
+	{
+		UClass* LiteralStorageClass;
+		UClass* WorkStorageClass;
+		UClass* DebugStorageClass;
+		
+		FClassInfo()
+			: LiteralStorageClass(nullptr)
+			, WorkStorageClass(nullptr)
+			, DebugStorageClass(nullptr)
+		{
+		}
+
+		UClass** GetClassPtr(ERigVMMemoryType InMemoryType)
+		{
+			switch(InMemoryType)
+			{
+				case ERigVMMemoryType::Literal:
+				{
+					return &LiteralStorageClass;
+				}
+				case ERigVMMemoryType::Debug:
+				{
+					return &DebugStorageClass;
+				}
+				default:
+				{
+					break;
+				}
+			}
+			return &WorkStorageClass;
+		}
+	};
+	
+	static TMap<FString, FClassInfo> PackageToInfo;
 };
