RigVM: Extending memory storage API + unit test

diff --git a/Engine/Source/Runtime/RigVM/Private/RigVMCore/RigVMMemoryStorage.cpp b/Engine/Source/Runtime/RigVM/Private/RigVMCore/RigVMMemoryStorage.cpp
--- a/Engine/Source/Runtime/RigVM/Private/RigVMCore/RigVMMemoryStorage.cpp
+++ b/Engine/Source/Runtime/RigVM/Private/RigVMCore/RigVMMemoryStorage.cpp
@@ -51,7 +51,7 @@ URigVMMemoryStorage::FPropertyDescription::FPropertyDescription(const FName& InN
 	}
 }
 
-FName URigVMMemoryStorage::FPropertyDescription::SanitizedName(const FName& InName)
+FName URigVMMemoryStorage::FPropertyDescription::SanitizeName(const FName& InName)
 {
 	FString NameString = InName.ToString();
 
@@ -80,7 +80,7 @@ FName URigVMMemoryStorage::FPropertyDescription::SanitizedName(const FName& InNa
 
 void URigVMMemoryStorage::FPropertyDescription::SanitizeName()
 {
-	Name = SanitizedName(Name);
+	Name = SanitizeName(Name);
 }
 
 FString URigVMMemoryStorage::FPropertyDescription::GetBaseCPPType() const
@@ -197,7 +197,7 @@ UClass* URigVMMemoryStorage::CreateStorageClass(UObject* InOuter, ERigVMMemoryTy
 
 	// Generate properties
 	FField** LinkToProperty = &Class->ChildProperties;
-	TArray<FProperty*> GeneratedProperties;
+	TArray<const FProperty*> CachedProperties;
 	for(const FPropertyDescription& PropertyDescription : InProperties)
 	{
 		while (*LinkToProperty != nullptr)
@@ -210,7 +210,7 @@ UClass* URigVMMemoryStorage::CreateStorageClass(UObject* InOuter, ERigVMMemoryTy
 			FProperty* NewProperty = CastFieldChecked<FProperty>(FField::Duplicate(PropertyDescription.Property, Class, PropertyDescription.Name));
 			check(NewProperty);
 
-			GeneratedProperties.Add(NewProperty);
+			CachedProperties.Add(NewProperty);
 			*LinkToProperty = NewProperty;
 		}
 		else
@@ -322,13 +322,13 @@ UClass* URigVMMemoryStorage::CreateStorageClass(UObject* InOuter, ERigVMMemoryTy
 
 			if(OuterProperty)
 			{
-				GeneratedProperties.Add(OuterProperty);
+				CachedProperties.Add(OuterProperty);
 				(*LinkToProperty) = OuterProperty;
 			}
 			else
 			{
 				check(*ValuePropertyPtr);
-				GeneratedProperties.Add(*ValuePropertyPtr);
+				CachedProperties.Add(*ValuePropertyPtr);
 				(*LinkToProperty) = *ValuePropertyPtr;
 			}
 		}
@@ -341,11 +341,14 @@ UClass* URigVMMemoryStorage::CreateStorageClass(UObject* InOuter, ERigVMMemoryTy
 	// Similar to FConfigPropertyHelperDetails::CustomizeDetails, this is required for GC to work properly
 	Class->AssembleReferenceTokenStream();
 
-	check(GeneratedProperties.Num() == InProperties.Num());
+	check(CachedProperties.Num() == InProperties.Num());
 	
-	// Create default object and store default values.
+	// Create default object
 	URigVMMemoryStorage* CDO = Cast<URigVMMemoryStorage>(Class->GetDefaultObject(true));
-	for(int32 PropertyIndex = 0; PropertyIndex < GeneratedProperties.Num(); PropertyIndex++)
+	CDO->CachedProperties = CachedProperties;
+
+	// and store default values.
+	for(int32 PropertyIndex = 0; PropertyIndex < CachedProperties.Num(); PropertyIndex++)
 	{
 		const FString& DefaultValue = InProperties[PropertyIndex].DefaultValue;
 		if(DefaultValue.IsEmpty())
@@ -353,7 +356,7 @@ UClass* URigVMMemoryStorage::CreateStorageClass(UObject* InOuter, ERigVMMemoryTy
 			continue;
 		}
 
-		FProperty* Property = GeneratedProperties[PropertyIndex];
+		const FProperty* Property = CachedProperties[PropertyIndex];
 		uint8* ValuePtr = Property->ContainerPtrToValuePtr<uint8>(CDO);
 
 		Property->ImportText(*DefaultValue, ValuePtr, EPropertyPortFlags::PPF_None, nullptr);
@@ -374,16 +377,31 @@ URigVMMemoryStorage* URigVMMemoryStorage::CreateStorage(UObject* InOuter, ERigVM
 	return NewObject<URigVMMemoryStorage>(InOuter, Class, NAME_None, RF_Public | RF_Transactional);
 }
 
-void URigVMMemoryStorage::RefreshCache()
+const TArray<const FProperty*>& URigVMMemoryStorage::GetProperties() const
 {
-	Cache.Reset();
-
-	FField** LinkToProperty = &GetClass()->ChildProperties;
-
-	while (*LinkToProperty != nullptr)
+	if(!HasAnyFlags((RF_ClassDefaultObject)))
 	{
-		FProperty* Property = CastField<FProperty>(*LinkToProperty);
-		Cache.Add(Property->ContainerPtrToValuePtr<uint8>(this));
-		LinkToProperty = &(*LinkToProperty)->Next;
+		if(URigVMMemoryStorage* CDO = Cast<URigVMMemoryStorage>(GetClass()->GetDefaultObject()))
+		{
+			return CDO->GetProperties();
+		}
 	}
+	return CachedProperties;
+}
+
+int32 URigVMMemoryStorage::GetPropertyIndex(const FProperty* InProperty) const
+{
+	return GetProperties().Find(InProperty);
+}
+
+int32 URigVMMemoryStorage::GetPropertyIndexByName(const FName& InName) const
+{
+	const FProperty* Property = FindPropertyByName(InName);
+	return GetPropertyIndex(Property);
+}
+
+const FProperty* URigVMMemoryStorage::FindPropertyByName(const FName& InName) const
+{
+	const FName SanitizedName = FPropertyDescription::SanitizeName(InName);
+	return GetClass()->FindPropertyByName(SanitizedName);
 }
diff --git a/Engine/Source/Runtime/RigVM/Public/RigVMCore/RigVMMemoryStorage.h b/Engine/Source/Runtime/RigVM/Public/RigVMCore/RigVMMemoryStorage.h
--- a/Engine/Source/Runtime/RigVM/Public/RigVMCore/RigVMMemoryStorage.h
+++ b/Engine/Source/Runtime/RigVM/Public/RigVMCore/RigVMMemoryStorage.h
@@ -54,7 +54,7 @@ public:
 		FPropertyDescription(const FProperty* InProperty, const FString& InDefaultValue, const FName& InName = NAME_None);
 		FPropertyDescription(const FName& InName, const FString& InCPPType, UObject* InCPPTypeObject, const FString& InDefaultValue);
 
-		static FName SanitizedName(const FName& InName);
+		static FName SanitizeName(const FName& InName);
 		void SanitizeName();
 
 		bool HasContainer() const { return ContainerType != EPinContainerType::None; }
@@ -82,10 +82,32 @@ public:
 	/// Memory Access
 	//////////////////////////////////////////////////////////////////////////////
 
+	const TArray<const FProperty*>& GetProperties() const;
+	int32 GetPropertyIndex(const FProperty* InProperty) const;
+	int32 GetPropertyIndexByName(const FName& InName) const;
+	const FProperty* FindPropertyByName(const FName& InName) const;
+
+	template<typename T>
+	T* ContainerPtrToValuePtr(int32 InPropertyIndex)
+	{
+		const TArray<const FProperty*>& Properties = GetProperties();
+		if(Properties.IsValidIndex(InPropertyIndex))
+		{
+			return Properties[InPropertyIndex]->ContainerPtrToValuePtr<T>(this);
+		}
+		return nullptr;
+	}
+
+	template<typename T>
+	T* ContainerPtrToValuePtrByName(const FName& InName)
+	{
+		const int32 PropertyIndex = GetPropertyIndexByName(InName);
+		return ContainerPtrToValuePtr<T>(PropertyIndex);
+	}
+
 protected:
 
-	void RefreshCache();
-	TArray<uint8*> Cache;
+	TArray<const FProperty*> CachedProperties;
 
 private:
 	
