RigVM: Memory storage WIP - fixed horde non unity breakage

diff --git a/Engine/Source/Runtime/RigVM/Private/RigVMCore/RigVMMemoryStorage.cpp b/Engine/Source/Runtime/RigVM/Private/RigVMCore/RigVMMemoryStorage.cpp
--- a/Engine/Source/Runtime/RigVM/Private/RigVMCore/RigVMMemoryStorage.cpp
+++ b/Engine/Source/Runtime/RigVM/Private/RigVMCore/RigVMMemoryStorage.cpp
@@ -13,10 +13,35 @@ URigVMMemoryStorage::FPropertyDescription::FPropertyDescription(const FProperty*
 	: Name(InName)
 	, Property(InProperty)
 	, CPPType()
-	, ContainerType(EPinContainerType::None)
+	, Containers()
 	, DefaultValue(InDefaultValue)
 {
 	SanitizeName();
+
+	const FProperty* ChildProperty = InProperty;
+	do
+	{
+		if(const FArrayProperty* ArrayProperty = CastField<FArrayProperty>(ChildProperty))
+		{
+			Containers.Add(EPinContainerType::Array);
+			ChildProperty = ArrayProperty->Inner;
+		}
+		else if(const FSetProperty* SetProperty = CastField<FSetProperty>(ChildProperty))
+		{
+			Containers.Add(EPinContainerType::Set);
+			ChildProperty = SetProperty->ElementProp;
+		}
+		else if(const FMapProperty* MapProperty = CastField<FMapProperty>(ChildProperty))
+		{
+			Containers.Add(EPinContainerType::Map);
+			ChildProperty = MapProperty->ValueProp;
+		}
+		else
+		{
+			ChildProperty = nullptr;
+		}
+	}
+	while (ChildProperty);
 }
 
 URigVMMemoryStorage::FPropertyDescription::FPropertyDescription(const FName& InName, const FString& InCPPType, UObject* InCPPTypeObject, const FString& InDefaultValue)
@@ -24,31 +49,36 @@ URigVMMemoryStorage::FPropertyDescription::FPropertyDescription(const FName& InN
 	, Property(nullptr)
 	, CPPType(InCPPType)
 	, CPPTypeObject(InCPPTypeObject)
-	, ContainerType(EPinContainerType::None)
+	, Containers()
 	, DefaultValue(InDefaultValue)
 {
 	SanitizeName();
 
 	FString BaseCPPType = CPPType;
-	if(BaseCPPType.RemoveFromStart(ArrayPrefix))
-	{
-		ContainerType = EPinContainerType::Array;
-	}
-	else if(BaseCPPType.RemoveFromStart(MapPrefix))
-	{
-		ContainerType = EPinContainerType::Map;
-	}
-	else if(BaseCPPType.RemoveFromStart(SetPrefix))
-	{
-		ContainerType = EPinContainerType::Set;
-	}
 
-	// make sure this description doesn't contain another container
-	if(ContainerType != EPinContainerType::None)
+	do
 	{
-		BaseCPPType.RemoveFromEnd(ContainerSuffix);
-		check(!BaseCPPType.Contains(TEXT("<")));
+		if(BaseCPPType.RemoveFromStart(ArrayPrefix))
+		{
+			Containers.Add(EPinContainerType::Array);
+			check(BaseCPPType.RemoveFromEnd(ContainerSuffix));
+		}
+		else if(BaseCPPType.RemoveFromStart(MapPrefix))
+		{
+			Containers.Add(EPinContainerType::Map);
+			check(BaseCPPType.RemoveFromEnd(ContainerSuffix));
+		}
+		else if(BaseCPPType.RemoveFromStart(SetPrefix))
+		{
+			Containers.Add(EPinContainerType::Set);
+			check(BaseCPPType.RemoveFromEnd(ContainerSuffix));
+		}
+		else
+		{
+			break;
+		}
 	}
+	while(!BaseCPPType.IsEmpty());
 }
 
 FName URigVMMemoryStorage::FPropertyDescription::SanitizeName(const FName& InName)
@@ -85,67 +115,39 @@ void URigVMMemoryStorage::FPropertyDescription::SanitizeName()
 
 FString URigVMMemoryStorage::FPropertyDescription::GetBaseCPPType() const
 {
-	if(IsArray())
-	{
-		return GetArrayElementCPPType();
-	}
-	if(IsMap())
-	{
-		return GetMapValueCPPType();
-	}
-	if(IsSet())
-	{
-		return GetSetElementCPPType(); 
-	}
-	return CPPType;
-}
-
-FString URigVMMemoryStorage::FPropertyDescription::GetArrayElementCPPType() const
-{
-	check(IsArray());
-
-	FString BaseCPPType = CPPType;
-	check(BaseCPPType.RemoveFromStart(ArrayPrefix));
-	check(BaseCPPType.RemoveFromEnd(ContainerSuffix));
-	return BaseCPPType.TrimStartAndEnd();
-}
-
-FString URigVMMemoryStorage::FPropertyDescription::GetSetElementCPPType() const
-{
-	check(IsSet());
-
 	FString BaseCPPType = CPPType;
-	check(BaseCPPType.RemoveFromStart(SetPrefix));
-	check(BaseCPPType.RemoveFromEnd(ContainerSuffix));
-	return BaseCPPType.TrimStartAndEnd();
-}
-
-FString URigVMMemoryStorage::FPropertyDescription::GetMapKeyCPPType() const
-{
-	check(IsMap());
-
-	FString BaseCPPType = CPPType;
-	check(BaseCPPType.RemoveFromStart(MapPrefix));
-	check(BaseCPPType.RemoveFromEnd(ContainerSuffix));
-
-	int32 Comma = 0;
-	check(BaseCPPType.FindChar(',', Comma));
-
-	return BaseCPPType.Left(Comma).TrimStartAndEnd();
-}
-
-FString URigVMMemoryStorage::FPropertyDescription::GetMapValueCPPType() const
-{
-	check(IsMap());
 
-	FString BaseCPPType = CPPType;
-	check(BaseCPPType.RemoveFromStart(MapPrefix));
-	check(BaseCPPType.RemoveFromEnd(ContainerSuffix));
-
-	int32 Comma = 0;
-	check(BaseCPPType.FindChar(',', Comma));
-
-	return BaseCPPType.Mid(Comma + 1).TrimStartAndEnd();
+	for(EPinContainerType Container : Containers)
+	{
+		switch(Container)
+		{
+			case EPinContainerType::Array:
+			{
+				check(BaseCPPType.RemoveFromStart(ArrayPrefix))
+				check(BaseCPPType.RemoveFromEnd(ContainerSuffix));
+				break;
+			}		
+			case EPinContainerType::Map:
+			{
+				check(BaseCPPType.RemoveFromStart(MapPrefix))
+				check(BaseCPPType.RemoveFromEnd(ContainerSuffix));
+				break;
+			}		
+			case EPinContainerType::Set:
+			{
+				check(BaseCPPType.RemoveFromStart(SetPrefix))
+				check(BaseCPPType.RemoveFromEnd(ContainerSuffix));
+				break;
+			}		
+			case EPinContainerType::None:
+			default:
+			{
+				break;
+			}
+		}
+	}
+	
+	return BaseCPPType;
 }
 
 ////////////////////////////////////////////////////////////////////////////////////////////////////////////////////////
@@ -197,53 +199,125 @@ UClass* URigVMMemoryStorage::CreateStorageClass(UObject* InOuter, ERigVMMemoryTy
 
 	// Generate properties
 	FField** LinkToProperty = &Class->ChildProperties;
+
 	TArray<const FProperty*> CachedProperties;
 	for(const FPropertyDescription& PropertyDescription : InProperties)
 	{
-		while (*LinkToProperty != nullptr)
+		FProperty* CachedProperty = AddProperty(Class, PropertyDescription, false, false, LinkToProperty);
+		check(CachedProperty);
+		CachedProperties.Add(CachedProperty);
+	}
+
+	// Update the class
+	Class->Bind();
+	Class->StaticLink(true);
+
+	// Similar to FConfigPropertyHelperDetails::CustomizeDetails, this is required for GC to work properly
+	Class->AssembleReferenceTokenStream();
+
+	check(CachedProperties.Num() == InProperties.Num());
+	
+	// Create default object
+	URigVMMemoryStorage* CDO = Cast<URigVMMemoryStorage>(Class->GetDefaultObject(true));
+	CDO->CachedProperties = CachedProperties;
+
+	// and store default values.
+	for(int32 PropertyIndex = 0; PropertyIndex < CachedProperties.Num(); PropertyIndex++)
+	{
+		const FString& DefaultValue = InProperties[PropertyIndex].DefaultValue;
+		if(DefaultValue.IsEmpty())
 		{
-			LinkToProperty = &(*LinkToProperty)->Next;
+			continue;
 		}
-		
-		if(PropertyDescription.Property)
-		{
-			FProperty* NewProperty = CastFieldChecked<FProperty>(FField::Duplicate(PropertyDescription.Property, Class, PropertyDescription.Name));
-			check(NewProperty);
 
-			CachedProperties.Add(NewProperty);
-			*LinkToProperty = NewProperty;
-		}
-		else
-		{
-			FFieldVariant PropertyOwner = Class;
-			FProperty* InnerProperty = nullptr;
-			FProperty* OuterProperty = nullptr;
-			FProperty** KeyPropertyPtr = nullptr;
-			FProperty** ValuePropertyPtr = &InnerProperty;
+		const FProperty* Property = CachedProperties[PropertyIndex];
+		uint8* ValuePtr = Property->ContainerPtrToValuePtr<uint8>(CDO);
+
+		Property->ImportText(*DefaultValue, ValuePtr, EPropertyPortFlags::PPF_None, nullptr);
+	}
+	
+	return Class;
+}
+
+URigVMMemoryStorage* URigVMMemoryStorage::CreateStorage(UObject* InOuter, ERigVMMemoryType InMemoryType)
+{
+	UClass* Class = GetStorageClass(InOuter, InMemoryType);
 
-			switch(PropertyDescription.ContainerType)
+	if(InMemoryType == ERigVMMemoryType::Literal)
+	{
+		return Cast<URigVMMemoryStorage>(Class->GetDefaultObject(true));
+	}
+
+	return NewObject<URigVMMemoryStorage>(InOuter, Class, NAME_None, RF_Public | RF_Transactional);
+}
+
+FProperty* URigVMMemoryStorage::AddProperty(UClass* InClass, const FPropertyDescription& InProperty, bool bPurge, bool bLink, FField** LinkToProperty)
+{
+	UClass *SuperClass = URigVMMemoryStorage::StaticClass();
+	
+	check(InClass);
+	check(InClass->GetSuperClass() == SuperClass);
+
+	if(bPurge)
+	{
+		InClass->PurgeClass(false);
+		InClass->PropertyLink = SuperClass->PropertyLink;
+		InClass->SetSuperStruct(SuperClass);
+		InClass->ClassWithin = UObject::StaticClass();
+	}
+	
+	if(LinkToProperty == nullptr)
+	{
+		LinkToProperty = &InClass->ChildProperties;
+	}
+	while (*LinkToProperty != nullptr)
+	{
+		LinkToProperty = &(*LinkToProperty)->Next;
+	}
+
+	FProperty* Result = nullptr;
+	if(InProperty.Property)
+	{
+		FProperty* NewProperty = CastFieldChecked<FProperty>(FField::Duplicate(InProperty.Property, InClass, InProperty.Name));
+		check(NewProperty);
+
+		Result = NewProperty;
+		*LinkToProperty = NewProperty;
+	}
+	else
+	{
+		FFieldVariant PropertyOwner = InClass;
+		FProperty** KeyPropertyPtr = nullptr;
+		FProperty** ValuePropertyPtr = &Result;
+
+		for(EPinContainerType Container : InProperty.Containers)
+		{
+			switch(Container)
 			{
 				case EPinContainerType::Array:
 				{
-					FArrayProperty* ArrayProperty = new FArrayProperty(PropertyOwner, PropertyDescription.Name, RF_Public);
+					FArrayProperty* ArrayProperty = new FArrayProperty(PropertyOwner, InProperty.Name, RF_Public);
+					*ValuePropertyPtr = ArrayProperty;
 					ValuePropertyPtr = &ArrayProperty->Inner;
-					PropertyOwner = OuterProperty = ArrayProperty;
+					PropertyOwner = ArrayProperty;
 					break;
 				}
 				case EPinContainerType::Map:
 				{
 					checkNoEntry(); // this is not implemented yet
-					FMapProperty* MapProperty = new FMapProperty(PropertyOwner, PropertyDescription.Name, RF_Public);
+					FMapProperty* MapProperty = new FMapProperty(PropertyOwner, InProperty.Name, RF_Public);
+					*ValuePropertyPtr = MapProperty;
 					KeyPropertyPtr = &MapProperty->KeyProp;
 					ValuePropertyPtr = &MapProperty->ValueProp;
-					PropertyOwner = OuterProperty = MapProperty;
+					PropertyOwner = MapProperty;
 					break;
 				}
 				case EPinContainerType::Set:
 				{
-					FSetProperty* SetProperty = new FSetProperty(PropertyOwner, PropertyDescription.Name, RF_Public);
+					FSetProperty* SetProperty = new FSetProperty(PropertyOwner, InProperty.Name, RF_Public);
+					*ValuePropertyPtr = SetProperty;
 					ValuePropertyPtr = &SetProperty->ElementProp;
-					PropertyOwner = OuterProperty = SetProperty;
+					PropertyOwner = SetProperty;
 					break;
 				}
 				case EPinContainerType::None:
@@ -252,129 +326,89 @@ UClass* URigVMMemoryStorage::CreateStorageClass(UObject* InOuter, ERigVMMemoryTy
 					break;
 				}
 			}
+		}
 
-			if(PropertyDescription.CPPTypeObject != nullptr)
+		if(InProperty.CPPTypeObject != nullptr)
+		{
+			if(UEnum* Enum = Cast<UEnum>(InProperty.CPPTypeObject))
 			{
-				if(UEnum* Enum = Cast<UEnum>(PropertyDescription.CPPTypeObject))
-				{
-					FByteProperty* EnumProperty = new FByteProperty(PropertyOwner, PropertyDescription.Name, RF_Public);
-					EnumProperty->Enum = Enum;
-					(*ValuePropertyPtr) = EnumProperty;
-				}
-				else if(UScriptStruct* ScriptStruct = Cast<UScriptStruct>(PropertyDescription.CPPTypeObject))
-				{
-					FStructProperty* StructProperty = new FStructProperty(PropertyOwner, PropertyDescription.Name, RF_Public);
-					StructProperty->Struct = ScriptStruct;
-					(*ValuePropertyPtr) = StructProperty;
-				}
-				else if(UClass* PropertyClass = Cast<UClass>(PropertyDescription.CPPTypeObject))
-				{
-					FObjectProperty* ObjectProperty = new FObjectProperty(PropertyOwner, PropertyDescription.Name, RF_Public);
-					ObjectProperty->SetPropertyClass(PropertyClass);
-					(*ValuePropertyPtr) = ObjectProperty;
-				}
-				else
-				{
-					checkNoEntry();
-				}
+				FByteProperty* EnumProperty = new FByteProperty(PropertyOwner, InProperty.Name, RF_Public);
+				EnumProperty->Enum = Enum;
+				(*ValuePropertyPtr) = EnumProperty;
 			}
-			else // take care of default types...
+			else if(UScriptStruct* ScriptStruct = Cast<UScriptStruct>(InProperty.CPPTypeObject))
 			{
-				static FString BoolString = TEXT("bool");
-				static FString Int32String = TEXT("int32");
-				static FString IntString = TEXT("int");
-				static FString FloatString = TEXT("float");
-				static FString DoubleString = TEXT("double");
-				static FString StringString = TEXT("FString");
-				static FString NameString = TEXT("FName");
-
-				FString BaseCPPType = PropertyDescription.GetBaseCPPType();
-				if(BaseCPPType.Equals(BoolString, ESearchCase::IgnoreCase))
-				{
-					(*ValuePropertyPtr) = new FBoolProperty(PropertyOwner, PropertyDescription.Name, RF_Public);;
-				}
-				else if(BaseCPPType.Equals(Int32String, ESearchCase::IgnoreCase) ||
-					BaseCPPType.Equals(IntString, ESearchCase::IgnoreCase))
-				{
-					(*ValuePropertyPtr) = new FIntProperty(PropertyOwner, PropertyDescription.Name, RF_Public);;
-				}
-				else if(BaseCPPType.Equals(FloatString, ESearchCase::IgnoreCase))
-				{
-					(*ValuePropertyPtr) = new FFloatProperty(PropertyOwner, PropertyDescription.Name, RF_Public);;
-				}
-				else if(BaseCPPType.Equals(DoubleString, ESearchCase::IgnoreCase))
-				{
-					(*ValuePropertyPtr) = new FDoubleProperty(PropertyOwner, PropertyDescription.Name, RF_Public);;
-				}
-				else if(BaseCPPType.Equals(StringString, ESearchCase::IgnoreCase))
-				{
-					(*ValuePropertyPtr) = new FStrProperty(PropertyOwner, PropertyDescription.Name, RF_Public);;
-				}
-				else if(BaseCPPType.Equals(NameString, ESearchCase::IgnoreCase))
-				{
-					(*ValuePropertyPtr) = new FNameProperty(PropertyOwner, PropertyDescription.Name, RF_Public);;
-				}
-				else
-				{
-					checkNoEntry();
-				}
+				FStructProperty* StructProperty = new FStructProperty(PropertyOwner, InProperty.Name, RF_Public);
+				StructProperty->Struct = ScriptStruct;
+				(*ValuePropertyPtr) = StructProperty;
 			}
-
-			if(OuterProperty)
+			else if(UClass* PropertyClass = Cast<UClass>(InProperty.CPPTypeObject))
 			{
-				CachedProperties.Add(OuterProperty);
-				(*LinkToProperty) = OuterProperty;
+				FObjectProperty* ObjectProperty = new FObjectProperty(PropertyOwner, InProperty.Name, RF_Public);
+				ObjectProperty->SetPropertyClass(PropertyClass);
+				(*ValuePropertyPtr) = ObjectProperty;
 			}
 			else
 			{
-				check(*ValuePropertyPtr);
-				CachedProperties.Add(*ValuePropertyPtr);
-				(*LinkToProperty) = *ValuePropertyPtr;
+				checkNoEntry();
 			}
 		}
-	}
-
-	// Update the class
-	Class->Bind();
-	Class->StaticLink(true);
-
-	// Similar to FConfigPropertyHelperDetails::CustomizeDetails, this is required for GC to work properly
-	Class->AssembleReferenceTokenStream();
-
-	check(CachedProperties.Num() == InProperties.Num());
-	
-	// Create default object
-	URigVMMemoryStorage* CDO = Cast<URigVMMemoryStorage>(Class->GetDefaultObject(true));
-	CDO->CachedProperties = CachedProperties;
-
-	// and store default values.
-	for(int32 PropertyIndex = 0; PropertyIndex < CachedProperties.Num(); PropertyIndex++)
-	{
-		const FString& DefaultValue = InProperties[PropertyIndex].DefaultValue;
-		if(DefaultValue.IsEmpty())
+		else // take care of default types...
 		{
-			continue;
+			static FString BoolString = TEXT("bool");
+			static FString Int32String = TEXT("int32");
+			static FString IntString = TEXT("int");
+			static FString FloatString = TEXT("float");
+			static FString DoubleString = TEXT("double");
+			static FString StringString = TEXT("FString");
+			static FString NameString = TEXT("FName");
+
+			const FString BaseCPPType = InProperty.GetBaseCPPType();
+			if(BaseCPPType.Equals(BoolString, ESearchCase::IgnoreCase))
+			{
+				(*ValuePropertyPtr) = new FBoolProperty(PropertyOwner, InProperty.Name, RF_Public);;
+			}
+			else if(BaseCPPType.Equals(Int32String, ESearchCase::IgnoreCase) ||
+				BaseCPPType.Equals(IntString, ESearchCase::IgnoreCase))
+			{
+				(*ValuePropertyPtr) = new FIntProperty(PropertyOwner, InProperty.Name, RF_Public);;
+			}
+			else if(BaseCPPType.Equals(FloatString, ESearchCase::IgnoreCase))
+			{
+				(*ValuePropertyPtr) = new FFloatProperty(PropertyOwner, InProperty.Name, RF_Public);;
+			}
+			else if(BaseCPPType.Equals(DoubleString, ESearchCase::IgnoreCase))
+			{
+				(*ValuePropertyPtr) = new FDoubleProperty(PropertyOwner, InProperty.Name, RF_Public);;
+			}
+			else if(BaseCPPType.Equals(StringString, ESearchCase::IgnoreCase))
+			{
+				(*ValuePropertyPtr) = new FStrProperty(PropertyOwner, InProperty.Name, RF_Public);;
+			}
+			else if(BaseCPPType.Equals(NameString, ESearchCase::IgnoreCase))
+			{
+				(*ValuePropertyPtr) = new FNameProperty(PropertyOwner, InProperty.Name, RF_Public);;
+			}
+			else
+			{
+				checkNoEntry();
+			}
 		}
 
-		const FProperty* Property = CachedProperties[PropertyIndex];
-		uint8* ValuePtr = Property->ContainerPtrToValuePtr<uint8>(CDO);
-
-		Property->ImportText(*DefaultValue, ValuePtr, EPropertyPortFlags::PPF_None, nullptr);
+		(*LinkToProperty) = Result;
 	}
-	
-	return Class;
-}
-
-URigVMMemoryStorage* URigVMMemoryStorage::CreateStorage(UObject* InOuter, ERigVMMemoryType InMemoryType)
-{
-	UClass* Class = GetStorageClass(InOuter, InMemoryType);
 
-	if(InMemoryType == ERigVMMemoryType::Literal)
+	if(bLink)
 	{
-		return Cast<URigVMMemoryStorage>(Class->GetDefaultObject(true));
-	}
+		// Update the class
+		InClass->Bind();
+		InClass->StaticLink(true);
 
-	return NewObject<URigVMMemoryStorage>(InOuter, Class, NAME_None, RF_Public | RF_Transactional);
+		// Similar to FConfigPropertyHelperDetails::CustomizeDetails, this is required for GC to work properly
+		InClass->AssembleReferenceTokenStream();
+	};
+	
+	return Result;
 }
 
 const TArray<const FProperty*>& URigVMMemoryStorage::GetProperties() const
diff --git a/Engine/Source/Runtime/RigVM/Public/RigVMCore/RigVMMemoryStorage.h b/Engine/Source/Runtime/RigVM/Public/RigVMCore/RigVMMemoryStorage.h
--- a/Engine/Source/Runtime/RigVM/Public/RigVMCore/RigVMMemoryStorage.h
+++ b/Engine/Source/Runtime/RigVM/Public/RigVMCore/RigVMMemoryStorage.h
@@ -8,6 +8,7 @@
 #include "RigVMStatistics.h"
 #include "RigVMArray.h"
 #include "RigVMMemoryCommon.h"
+#include "EdGraph/EdGraphNode.h"
 #include "RigVMMemoryStorage.generated.h"
 
 #if !UE_RIGVM_UCLASS_BASED_STORAGE_DISABLED
@@ -23,6 +24,16 @@ public:
 		: Ptr(nullptr)
 	{}
 
+	FORCEINLINE_DEBUGGABLE uint8* GetData()
+	{
+		return Ptr;
+	}
+
+	FORCEINLINE_DEBUGGABLE const uint8* GetData() const
+	{
+		return Ptr;
+	}
+
 private:
 	
 	uint8* Ptr;
@@ -48,7 +59,7 @@ public:
 		const FProperty* Property;
 		FString CPPType;
 		UObject* CPPTypeObject;
-		EPinContainerType ContainerType;
+		TArray<EPinContainerType> Containers;
 		FString DefaultValue;
 		
 		FPropertyDescription(const FProperty* InProperty, const FString& InDefaultValue, const FName& InName = NAME_None);
@@ -57,16 +68,13 @@ public:
 		static FName SanitizeName(const FName& InName);
 		void SanitizeName();
 
-		bool HasContainer() const { return ContainerType != EPinContainerType::None; }
-		bool IsArray() const { return ContainerType == EPinContainerType::Array; }
-		bool IsMap() const { return ContainerType == EPinContainerType::Map; }
-		bool IsSet() const { return ContainerType == EPinContainerType::Set; }
+		int32 NumContainers() const { return Containers.Num(); }
+		bool HasContainer() const { return NumContainers() > 0; }
+		bool IsArray(int32 InContainerIndex = 0) const { return Containers[InContainerIndex] == EPinContainerType::Array; }
+		bool IsMap(int32 InContainerIndex = 0) const { return Containers[InContainerIndex] == EPinContainerType::Map; }
+		bool IsSet(int32 InContainerIndex = 0) const { return Containers[InContainerIndex] == EPinContainerType::Set; }
 
 		FString GetBaseCPPType() const;
-		FString GetArrayElementCPPType() const;
-		FString GetSetElementCPPType() const;
-		FString GetMapKeyCPPType() const;
-		FString GetMapValueCPPType() const;
 
 		static const FString ArrayPrefix;
 		static const FString MapPrefix;
@@ -77,6 +85,7 @@ public:
 	static UClass* GetStorageClass(UObject* InOuter, ERigVMMemoryType InMemoryType);
 	static UClass* CreateStorageClass(UObject* InOuter, ERigVMMemoryType InMemoryType, const TArray<FPropertyDescription>& InProperties);
 	static URigVMMemoryStorage* CreateStorage(UObject* InOuter, ERigVMMemoryType InMemoryType);
+	static FProperty* AddProperty(UClass* InClass, const FPropertyDescription& InProperty, bool bPurge = false, bool bLink = true, FField** LinkToProperty = nullptr);
 
 	//////////////////////////////////////////////////////////////////////////////
 	/// Memory Access
