RigVM: Property Path replacement for RegisterOffset

diff --git a/Engine/Source/Runtime/RigVM/Private/RigVMCore/RigVMMemoryStorage.cpp b/Engine/Source/Runtime/RigVM/Private/RigVMCore/RigVMMemoryStorage.cpp
--- a/Engine/Source/Runtime/RigVM/Private/RigVMCore/RigVMMemoryStorage.cpp
+++ b/Engine/Source/Runtime/RigVM/Private/RigVMCore/RigVMMemoryStorage.cpp
@@ -395,6 +395,7 @@ FProperty* URigVMMemoryStorage::AddProperty(UClass* InClass, const FPropertyDesc
 			}
 		}
 
+		Result->SetPropertyFlags(CPF_Edit);
 		(*LinkToProperty) = Result;
 	}
 
@@ -406,6 +407,16 @@ FProperty* URigVMMemoryStorage::AddProperty(UClass* InClass, const FPropertyDesc
 
 		// Similar to FConfigPropertyHelperDetails::CustomizeDetails, this is required for GC to work properly
 		InClass->AssembleReferenceTokenStream();
+
+		URigVMMemoryStorage* CDO = Cast<URigVMMemoryStorage>(InClass->GetDefaultObject(true));
+		CDO->CachedProperties.Add(Result);
+		
+		const FString& DefaultValue = InProperty.DefaultValue;
+		if(!DefaultValue.IsEmpty())
+		{
+			uint8* ValuePtr = Result->ContainerPtrToValuePtr<uint8>(CDO);
+			Result->ImportText(*DefaultValue, ValuePtr, EPropertyPortFlags::PPF_None, nullptr);
+		}		
 	};
 	
 	return Result;
diff --git a/Engine/Source/Runtime/RigVM/Private/RigVMCore/RigVMPropertyPath.cpp b/Engine/Source/Runtime/RigVM/Private/RigVMCore/RigVMPropertyPath.cpp
new file mode 100644
--- /dev/null
+++ b/Engine/Source/Runtime/RigVM/Private/RigVMCore/RigVMPropertyPath.cpp
@@ -0,0 +1,156 @@
+// Copyright Epic Games, Inc. All Rights Reserved.
+
+#include "RigVMCore/RigVMPropertyPath.h"
+
+////////////////////////////////////////////////////////////////////////////////////////////////////////////////////////
+
+FRigVMPropertyPath::FRigVMPropertyPath()
+: Path()
+, Segments()
+{
+}
+
+FRigVMPropertyPath::FRigVMPropertyPath(const FProperty* InProperty, const FString& InSegmentPath)
+: Path()
+, Segments()
+{
+	check(InProperty);
+	check(!InSegmentPath.IsEmpty())
+
+	const FProperty* Property = InProperty;
+	FString WorkPath = InSegmentPath;
+	WorkPath = WorkPath.Replace(TEXT("["), TEXT("."));
+	WorkPath = WorkPath.Replace(TEXT("]"), TEXT("."));
+	WorkPath = WorkPath.Replace(TEXT(".."), TEXT("."));
+	WorkPath.TrimCharInline('.', nullptr);
+
+	while(!WorkPath.IsEmpty())
+	{
+		FString PathSegment, PathRemainder;
+		if(!WorkPath.Split(TEXT("."), &PathSegment, &PathRemainder))
+		{
+			PathSegment = WorkPath;
+			PathRemainder.Empty();
+		}
+
+		FRigVMPropertyPathSegment Segment;
+
+		if(const FStructProperty* StructProperty = CastField<FStructProperty>(Property))
+		{
+			Segment.Type = ERigVMPropertyPathSegmentType::StructMember;
+
+			if(const FProperty* MemberProperty = StructProperty->Struct->FindPropertyByName(*PathSegment))
+			{
+				Segment.Name = MemberProperty->GetFName();
+				Segment.Index = MemberProperty->GetOffset_ForInternal();
+				Property = Segment.Property = MemberProperty;
+			}
+			else
+			{
+				Segments.Empty();
+				return;
+			}
+		}
+		else if(const FArrayProperty* ArrayProperty = CastField<FArrayProperty>(Property))
+		{
+			Segment.Type = ERigVMPropertyPathSegmentType::ArrayElement;
+			const FProperty* ElementProperty = ArrayProperty->Inner;
+			Segment.Name = ElementProperty->GetFName();
+			Segment.Index = FCString::Atoi(*PathSegment);
+			Segment.Property = ArrayProperty;
+			Property = ElementProperty;
+		}
+		else if(const FMapProperty* MapProperty = CastField<FMapProperty>(Property))
+		{
+			check(MapProperty->KeyProp->IsA<FNameProperty>());
+			
+			Segment.Type = ERigVMPropertyPathSegmentType::MapValue;
+			const FProperty* ValueProperty = MapProperty->ValueProp;
+			Segment.Name = *PathSegment;
+			Segment.Index = INDEX_NONE;
+			Segment.Property = MapProperty;
+			Property = ValueProperty;
+		}
+		else
+		{
+			Segments.Empty();
+			return;
+		}
+
+		Segments.Add(Segment);
+		WorkPath = PathRemainder;
+	}
+
+	Path.Empty();
+
+	for(const FRigVMPropertyPathSegment& Segment : Segments)
+	{
+		switch(Segment.Type)
+		{
+			case ERigVMPropertyPathSegmentType::StructMember:
+			{
+				if(!Path.IsEmpty())
+				{
+					Path += TEXT(".");
+				}
+				Path += Segment.Name.ToString();
+				break;
+			}
+			case ERigVMPropertyPathSegmentType::ArrayElement:
+			{
+				Path += TEXT("[") + FString::FromInt(Segment.Index) + TEXT("]");
+				break;
+			}
+			case ERigVMPropertyPathSegmentType::MapValue:
+			{
+				Path += TEXT("[") + Segment.Name.ToString() + TEXT("]");
+				break;
+			}
+			default:
+			{
+				break;
+			}
+		}
+	}
+}
+
+FRigVMPropertyPath::FRigVMPropertyPath(const FRigVMPropertyPath& InOther)
+: Path(InOther.Path)
+, Segments(InOther.Segments)
+{
+}
+
+uint8* FRigVMPropertyPath::GetData_Internal(uint8* InPtr) const
+{
+	for(const FRigVMPropertyPathSegment& Segment : Segments)
+	{
+		switch(Segment.Type)
+		{
+			case ERigVMPropertyPathSegmentType::StructMember:
+			{
+				// index represents the offset needed
+				// from container to struct member
+				InPtr += Segment.Property->GetOffset_ForInternal();
+				break;
+			}
+			case ERigVMPropertyPathSegmentType::ArrayElement:
+			{
+				FScriptArrayHelper ArrayHelper(CastFieldChecked<FArrayProperty>(Segment.Property), InPtr);
+				InPtr = ArrayHelper.GetRawPtr(Segment.Index);
+				break;
+			}
+			case ERigVMPropertyPathSegmentType::MapValue:
+			{
+				FScriptMapHelper MapHelper(CastFieldChecked<FMapProperty>(Segment.Property), InPtr);
+				InPtr = MapHelper.FindValueFromHash(&Segment.Name);;
+				break;
+			}
+			default:
+			{
+				break;
+			}
+		}
+	}
+
+	return InPtr;
+}
diff --git a/Engine/Source/Runtime/RigVM/Public/RigVMCore/RigVMMemoryStorage.h b/Engine/Source/Runtime/RigVM/Public/RigVMCore/RigVMMemoryStorage.h
--- a/Engine/Source/Runtime/RigVM/Public/RigVMCore/RigVMMemoryStorage.h
+++ b/Engine/Source/Runtime/RigVM/Public/RigVMCore/RigVMMemoryStorage.h
@@ -9,6 +9,7 @@
 #include "RigVMArray.h"
 #include "RigVMMemoryCommon.h"
 #include "EdGraph/EdGraphNode.h"
+#include "RigVMPropertyPath.h"
 #include "RigVMMemoryStorage.generated.h"
 
 #if !UE_RIGVM_UCLASS_BASED_STORAGE_DISABLED
@@ -85,19 +86,73 @@ public:
 	static UClass* GetStorageClass(UObject* InOuter, ERigVMMemoryType InMemoryType);
 	static UClass* CreateStorageClass(UObject* InOuter, ERigVMMemoryType InMemoryType, const TArray<FPropertyDescription>& InProperties);
 	static URigVMMemoryStorage* CreateStorage(UObject* InOuter, ERigVMMemoryType InMemoryType);
-	static FProperty* AddProperty(UClass* InClass, const FPropertyDescription& InProperty, bool bPurge = false, bool bLink = true, FField** LinkToProperty = nullptr);
 
 	//////////////////////////////////////////////////////////////////////////////
 	/// Memory Access
 	//////////////////////////////////////////////////////////////////////////////
 
+	FORCEINLINE int32 Num() const
+	{
+		return GetProperties().Num();
+	}
+
+	FORCEINLINE bool IsValidIndex(int32 InIndex) const
+	{
+		return GetProperties().IsValidIndex(InIndex);
+	}
+	
 	const TArray<const FProperty*>& GetProperties() const;
+
 	int32 GetPropertyIndex(const FProperty* InProperty) const;
+
 	int32 GetPropertyIndexByName(const FName& InName) const;
+
 	const FProperty* FindPropertyByName(const FName& InName) const;
+	
+	FORCEINLINE bool IsArray(int32 InPropertyIndex) const
+	{
+		return GetProperties()[InPropertyIndex]->IsA<FArrayProperty>();
+	}
+	
+	FORCEINLINE bool IsArrayByName(const FName& InName) const
+	{
+		if(const FProperty* Property = FindPropertyByName(InName))
+		{
+			return Property->IsA<FArrayProperty>();
+		}
+		return false;
+	}
+	
+	FORCEINLINE bool IsMap(int32 InPropertyIndex) const
+	{
+		return GetProperties()[InPropertyIndex]->IsA<FMapProperty>();
+	}
+	
+	FORCEINLINE bool IsMapByName(const FName& InName) const
+	{
+		if(const FProperty* Property = FindPropertyByName(InName))
+		{
+			return Property->IsA<FMapProperty>();
+		}
+		return false;
+	}
+	
+	FORCEINLINE bool IsSet(int32 InPropertyIndex) const
+	{
+		return GetProperties()[InPropertyIndex]->IsA<FSetProperty>();
+	}
+	
+	FORCEINLINE bool IsSetByName(const FName& InName) const
+	{
+		if(const FProperty* Property = FindPropertyByName(InName))
+		{
+			return Property->IsA<FSetProperty>();
+		}
+		return false;
+	}
 
 	template<typename T>
-	T* ContainerPtrToValuePtr(int32 InPropertyIndex)
+	FORCEINLINE T* GetData(int32 InPropertyIndex)
 	{
 		const TArray<const FProperty*>& Properties = GetProperties();
 		if(Properties.IsValidIndex(InPropertyIndex))
@@ -108,14 +163,29 @@ public:
 	}
 
 	template<typename T>
-	T* ContainerPtrToValuePtrByName(const FName& InName)
+	FORCEINLINE T* GetDataByName(const FName& InName)
 	{
 		const int32 PropertyIndex = GetPropertyIndexByName(InName);
-		return ContainerPtrToValuePtr<T>(PropertyIndex);
+		return GetData<T>(PropertyIndex);
+	}
+
+	template<typename T>
+	FORCEINLINE T* GetData(int32 InPropertyIndex, const FRigVMPropertyPath& InPropertyPath)
+	{
+		return InPropertyPath.GetData<T>(GetData<uint8>(InPropertyIndex));
+	}
+
+	template<typename T>
+	FORCEINLINE T* GetDataByName(const FName& InName, const FRigVMPropertyPath& InPropertyPath)
+	{
+		const int32 PropertyIndex = GetPropertyIndexByName(InName);
+		return GetData<T>(PropertyIndex, InPropertyPath);
 	}
 
 protected:
 
+	static FProperty* AddProperty(UClass* InClass, const FPropertyDescription& InProperty, bool bPurge = false, bool bLink = true, FField** LinkToProperty = nullptr);
+
 	TArray<const FProperty*> CachedProperties;
 
 private:
diff --git a/Engine/Source/Runtime/RigVM/Public/RigVMCore/RigVMPropertyPath.h b/Engine/Source/Runtime/RigVM/Public/RigVMCore/RigVMPropertyPath.h
new file mode 100644
--- /dev/null
+++ b/Engine/Source/Runtime/RigVM/Public/RigVMCore/RigVMPropertyPath.h
@@ -0,0 +1,67 @@
+// Copyright Epic Games, Inc. All Rights Reserved.
+
+#pragma once
+
+#include "CoreMinimal.h"
+
+enum RIGVM_API ERigVMPropertyPathSegmentType
+{
+	StructMember,
+	ArrayElement,
+	MapValue
+};
+
+struct RIGVM_API FRigVMPropertyPathSegment
+{
+	ERigVMPropertyPathSegmentType Type;
+	FName Name;
+	int32 Index;
+	const FProperty* Property;
+};
+
+class RIGVM_API FRigVMPropertyPath
+{
+public:
+
+	FRigVMPropertyPath();
+	FRigVMPropertyPath(const FProperty* InProperty, const FString& InSegmentPath);
+	FRigVMPropertyPath(const FRigVMPropertyPath& InOther);
+
+	FORCEINLINE const FString& ToString() const { return Path; }
+	FORCEINLINE int32 Num() const { return Segments.Num(); }
+	FORCEINLINE bool IsValid() const { return Num() > 0; }
+	FORCEINLINE const FRigVMPropertyPathSegment& operator[](int32 InIndex) const { return Segments[InIndex]; }
+
+	FORCEINLINE bool operator ==(const FRigVMPropertyPath& Other) const
+	{
+		return ToString() == Other.ToString();
+	}
+
+	FORCEINLINE bool operator !=(const FRigVMPropertyPath& Other) const
+	{
+		return ToString() != Other.ToString();
+	}
+
+	FORCEINLINE bool operator >(const FRigVMPropertyPath& Other) const
+	{
+		return ToString() > Other.ToString();
+	}
+
+	FORCEINLINE bool operator <(const FRigVMPropertyPath& Other) const
+	{
+		return ToString() < Other.ToString();
+	}
+
+	uint8* GetData_Internal(uint8* InPtr) const;
+
+	template<typename T>
+	FORCEINLINE T* GetData(uint8* InPtr) const
+	{
+		return (T*)GetData_Internal(InPtr);
+	}
+
+private:
+
+	FString Path;
+	TArray<FRigVMPropertyPathSegment> Segments;
+};
\ No newline at end of file
