RigVM: Property Copy services

diff --git a/Engine/Source/Runtime/RigVM/Private/RigVMCore/RigVMMemoryStorage.cpp b/Engine/Source/Runtime/RigVM/Private/RigVMCore/RigVMMemoryStorage.cpp
--- a/Engine/Source/Runtime/RigVM/Private/RigVMCore/RigVMMemoryStorage.cpp
+++ b/Engine/Source/Runtime/RigVM/Private/RigVMCore/RigVMMemoryStorage.cpp
@@ -251,6 +251,71 @@ URigVMMemoryStorage* URigVMMemoryStorage::CreateStorage(UObject* InOuter, ERigVM
 	return NewObject<URigVMMemoryStorage>(InOuter, Class, NAME_None, RF_Public | RF_Transactional);
 }
 
+bool URigVMMemoryStorage::CopyProperty(
+	URigVMMemoryStorage* InTargetStorage,
+	int32 InTargetPropertyIndex,
+	const FRigVMPropertyPath& InTargetPropertyPath,
+	URigVMMemoryStorage* InSourceStorage,
+	int32 InSourcePropertyIndex,
+	const FRigVMPropertyPath& InSourcePropertyPath)
+{
+	check(InTargetStorage != nullptr);
+	check(InSourceStorage != nullptr);
+
+	const FProperty* TargetProperty = InTargetStorage->GetProperties()[InTargetPropertyIndex];
+	const FProperty* SourceProperty = InSourceStorage->GetProperties()[InSourcePropertyIndex];
+	uint8* TargetPtr = TargetProperty->ContainerPtrToValuePtr<uint8>(InTargetStorage);
+	uint8* SourcePtr = SourceProperty->ContainerPtrToValuePtr<uint8>(InSourceStorage);
+
+	auto TraversePropertyPath = [](const FProperty*& Property, uint8*& MemoryPtr, const FRigVMPropertyPath& PropertyPath)
+	{
+		if(PropertyPath.IsEmpty())
+		{
+			return;
+		}
+
+		MemoryPtr = PropertyPath.GetData<uint8>(MemoryPtr);
+
+		const FRigVMPropertyPathSegment& LastSegment = PropertyPath[PropertyPath.Num() - 1];
+		switch(LastSegment.Type)
+		{
+			case ERigVMPropertyPathSegmentType::ArrayElement:
+			{
+				const FArrayProperty* ArrayProperty = CastFieldChecked<FArrayProperty>(LastSegment.Property);
+				Property = ArrayProperty->Inner;
+				break;
+			}
+			case ERigVMPropertyPathSegmentType::MapValue:
+			{
+				const FMapProperty* MapProperty = CastFieldChecked<FMapProperty>(LastSegment.Property);
+				Property = MapProperty->ValueProp;
+				break;
+			}
+			case ERigVMPropertyPathSegmentType::StructMember:
+			{
+				Property = LastSegment.Property;
+				break;
+			}
+			default:
+			{
+				checkNoEntry();
+				break;
+			}
+		}
+	};
+
+	TraversePropertyPath(TargetProperty, TargetPtr, InTargetPropertyPath);
+	TraversePropertyPath(SourceProperty, SourcePtr, InSourcePropertyPath);
+	
+	if(!TargetProperty->SameType(SourceProperty))
+	{
+		return false;
+	}
+
+	TargetProperty->CopyCompleteValue(TargetPtr, SourcePtr);
+	return true;
+}
+
 FProperty* URigVMMemoryStorage::AddProperty(UClass* InClass, const FPropertyDescription& InProperty, bool bPurge, bool bLink, FField** LinkToProperty)
 {
 	UClass *SuperClass = URigVMMemoryStorage::StaticClass();
@@ -287,7 +352,7 @@ FProperty* URigVMMemoryStorage::AddProperty(UClass* InClass, const FPropertyDesc
 	else
 	{
 		FFieldVariant PropertyOwner = InClass;
-		FProperty** KeyPropertyPtr = nullptr;
+		// FProperty** KeyPropertyPtr = nullptr;
 		FProperty** ValuePropertyPtr = &Result;
 
 		for(EPinContainerType Container : InProperty.Containers)
@@ -307,7 +372,7 @@ FProperty* URigVMMemoryStorage::AddProperty(UClass* InClass, const FPropertyDesc
 					checkNoEntry(); // this is not implemented yet
 					FMapProperty* MapProperty = new FMapProperty(PropertyOwner, InProperty.Name, RF_Public);
 					*ValuePropertyPtr = MapProperty;
-					KeyPropertyPtr = &MapProperty->KeyProp;
+					// KeyPropertyPtr = &MapProperty->KeyProp;
 					ValuePropertyPtr = &MapProperty->ValueProp;
 					PropertyOwner = MapProperty;
 					break;
diff --git a/Engine/Source/Runtime/RigVM/Private/RigVMCore/RigVMPropertyPath.cpp b/Engine/Source/Runtime/RigVM/Private/RigVMCore/RigVMPropertyPath.cpp
--- a/Engine/Source/Runtime/RigVM/Private/RigVMCore/RigVMPropertyPath.cpp
+++ b/Engine/Source/Runtime/RigVM/Private/RigVMCore/RigVMPropertyPath.cpp
@@ -4,6 +4,8 @@
 
 ////////////////////////////////////////////////////////////////////////////////////////////////////////////////////////
 
+FRigVMPropertyPath FRigVMPropertyPath::Empty;
+
 FRigVMPropertyPath::FRigVMPropertyPath()
 : Path()
 , Segments()
@@ -154,3 +156,4 @@ uint8* FRigVMPropertyPath::GetData_Internal(uint8* InPtr) const
 
 	return InPtr;
 }
+
diff --git a/Engine/Source/Runtime/RigVM/Public/RigVMCore/RigVMMemoryStorage.h b/Engine/Source/Runtime/RigVM/Public/RigVMCore/RigVMMemoryStorage.h
--- a/Engine/Source/Runtime/RigVM/Public/RigVMCore/RigVMMemoryStorage.h
+++ b/Engine/Source/Runtime/RigVM/Public/RigVMCore/RigVMMemoryStorage.h
@@ -182,6 +182,14 @@ public:
 		return GetData<T>(PropertyIndex, InPropertyPath);
 	}
 
+	static bool CopyProperty(
+		URigVMMemoryStorage* InTargetStorage,
+		int32 InTargetPropertyIndex,
+		const FRigVMPropertyPath& InTargetPropertyPath,
+		URigVMMemoryStorage* InSourceStorage,
+		int32 InSourcePropertyIndex,
+		const FRigVMPropertyPath& InSourcePropertyPath);
+
 protected:
 
 	static FProperty* AddProperty(UClass* InClass, const FPropertyDescription& InProperty, bool bPurge = false, bool bLink = true, FField** LinkToProperty = nullptr);
diff --git a/Engine/Source/Runtime/RigVM/Public/RigVMCore/RigVMPropertyPath.h b/Engine/Source/Runtime/RigVM/Public/RigVMCore/RigVMPropertyPath.h
--- a/Engine/Source/Runtime/RigVM/Public/RigVMCore/RigVMPropertyPath.h
+++ b/Engine/Source/Runtime/RigVM/Public/RigVMCore/RigVMPropertyPath.h
@@ -30,6 +30,7 @@ public:
 	FORCEINLINE const FString& ToString() const { return Path; }
 	FORCEINLINE int32 Num() const { return Segments.Num(); }
 	FORCEINLINE bool IsValid() const { return Num() > 0; }
+	FORCEINLINE bool IsEmpty() const { return Num() == 0; }
 	FORCEINLINE const FRigVMPropertyPathSegment& operator[](int32 InIndex) const { return Segments[InIndex]; }
 
 	FORCEINLINE bool operator ==(const FRigVMPropertyPath& Other) const
@@ -60,6 +61,8 @@ public:
 		return (T*)GetData_Internal(InPtr);
 	}
 
+	static FRigVMPropertyPath Empty;
+
 private:
 
 	FString Path;
