RigVM: Property Copy services

Follow up for 16956055

diff --git a/Engine/Source/Runtime/RigVM/Private/RigVMCore/RigVMMemoryStorage.cpp b/Engine/Source/Runtime/RigVM/Private/RigVMCore/RigVMMemoryStorage.cpp
--- a/Engine/Source/Runtime/RigVM/Private/RigVMCore/RigVMMemoryStorage.cpp
+++ b/Engine/Source/Runtime/RigVM/Private/RigVMCore/RigVMMemoryStorage.cpp
@@ -252,20 +252,37 @@ URigVMMemoryStorage* URigVMMemoryStorage::CreateStorage(UObject* InOuter, ERigVM
 }
 
 bool URigVMMemoryStorage::CopyProperty(
-	URigVMMemoryStorage* InTargetStorage,
-	int32 InTargetPropertyIndex,
+	const FProperty* InTargetProperty,
+	uint8* InTargetPtr,
+	const FProperty* InSourceProperty,
+	uint8* InSourcePtr)
+{
+	check(InTargetProperty != nullptr);
+	check(InSourceProperty != nullptr);
+	check(InTargetPtr != nullptr);
+	check(InSourcePtr != nullptr);
+
+	if(!InTargetProperty->SameType(InSourceProperty))
+	{
+		return false;
+	}
+
+	InTargetProperty->CopyCompleteValue(InTargetPtr, InSourcePtr);
+	return true;
+}
+
+bool URigVMMemoryStorage::CopyProperty(
+	const FProperty* InTargetProperty,
+	uint8* InTargetPtr,
 	const FRigVMPropertyPath& InTargetPropertyPath,
-	URigVMMemoryStorage* InSourceStorage,
-	int32 InSourcePropertyIndex,
+	const FProperty* InSourceProperty,
+	uint8* InSourcePtr,
 	const FRigVMPropertyPath& InSourcePropertyPath)
 {
-	check(InTargetStorage != nullptr);
-	check(InSourceStorage != nullptr);
-
-	const FProperty* TargetProperty = InTargetStorage->GetProperties()[InTargetPropertyIndex];
-	const FProperty* SourceProperty = InSourceStorage->GetProperties()[InSourcePropertyIndex];
-	uint8* TargetPtr = TargetProperty->ContainerPtrToValuePtr<uint8>(InTargetStorage);
-	uint8* SourcePtr = SourceProperty->ContainerPtrToValuePtr<uint8>(InSourceStorage);
+	check(InTargetProperty != nullptr);
+	check(InSourceProperty != nullptr);
+	check(InTargetPtr != nullptr);
+	check(InSourcePtr != nullptr);
 
 	auto TraversePropertyPath = [](const FProperty*& Property, uint8*& MemoryPtr, const FRigVMPropertyPath& PropertyPath)
 	{
@@ -279,24 +296,24 @@ bool URigVMMemoryStorage::CopyProperty(
 		const FRigVMPropertyPathSegment& LastSegment = PropertyPath[PropertyPath.Num() - 1];
 		switch(LastSegment.Type)
 		{
-			case ERigVMPropertyPathSegmentType::ArrayElement:
+		case ERigVMPropertyPathSegmentType::ArrayElement:
 			{
 				const FArrayProperty* ArrayProperty = CastFieldChecked<FArrayProperty>(LastSegment.Property);
 				Property = ArrayProperty->Inner;
 				break;
 			}
-			case ERigVMPropertyPathSegmentType::MapValue:
+		case ERigVMPropertyPathSegmentType::MapValue:
 			{
 				const FMapProperty* MapProperty = CastFieldChecked<FMapProperty>(LastSegment.Property);
 				Property = MapProperty->ValueProp;
 				break;
 			}
-			case ERigVMPropertyPathSegmentType::StructMember:
+		case ERigVMPropertyPathSegmentType::StructMember:
 			{
 				Property = LastSegment.Property;
 				break;
 			}
-			default:
+		default:
 			{
 				checkNoEntry();
 				break;
@@ -304,16 +321,29 @@ bool URigVMMemoryStorage::CopyProperty(
 		}
 	};
 
-	TraversePropertyPath(TargetProperty, TargetPtr, InTargetPropertyPath);
-	TraversePropertyPath(SourceProperty, SourcePtr, InSourcePropertyPath);
-	
-	if(!TargetProperty->SameType(SourceProperty))
-	{
-		return false;
-	}
+	TraversePropertyPath(InTargetProperty, InTargetPtr, InTargetPropertyPath);
+	TraversePropertyPath(InSourceProperty, InSourcePtr, InSourcePropertyPath);
 
-	TargetProperty->CopyCompleteValue(TargetPtr, SourcePtr);
-	return true;
+	return CopyProperty(InTargetProperty, InTargetPtr, InSourceProperty, InSourcePtr);
+}
+
+bool URigVMMemoryStorage::CopyProperty(
+	URigVMMemoryStorage* InTargetStorage,
+	int32 InTargetPropertyIndex,
+	const FRigVMPropertyPath& InTargetPropertyPath,
+	URigVMMemoryStorage* InSourceStorage,
+	int32 InSourcePropertyIndex,
+	const FRigVMPropertyPath& InSourcePropertyPath)
+{
+	check(InTargetStorage != nullptr);
+	check(InSourceStorage != nullptr);
+
+	const FProperty* TargetProperty = InTargetStorage->GetProperties()[InTargetPropertyIndex];
+	const FProperty* SourceProperty = InSourceStorage->GetProperties()[InSourcePropertyIndex];
+	uint8* TargetPtr = TargetProperty->ContainerPtrToValuePtr<uint8>(InTargetStorage);
+	uint8* SourcePtr = SourceProperty->ContainerPtrToValuePtr<uint8>(InSourceStorage);
+
+	return CopyProperty(TargetProperty, TargetPtr, InTargetPropertyPath, SourceProperty, SourcePtr, InSourcePropertyPath);
 }
 
 FProperty* URigVMMemoryStorage::AddProperty(UClass* InClass, const FPropertyDescription& InProperty, bool bPurge, bool bLink, FField** LinkToProperty)
diff --git a/Engine/Source/Runtime/RigVM/Public/RigVMCore/RigVMMemoryStorage.h b/Engine/Source/Runtime/RigVM/Public/RigVMCore/RigVMMemoryStorage.h
--- a/Engine/Source/Runtime/RigVM/Public/RigVMCore/RigVMMemoryStorage.h
+++ b/Engine/Source/Runtime/RigVM/Public/RigVMCore/RigVMMemoryStorage.h
@@ -182,6 +182,20 @@ public:
 		return GetData<T>(PropertyIndex, InPropertyPath);
 	}
 
+	static bool CopyProperty(
+		const FProperty* InTargetProperty,
+		uint8* InTargetPtr,
+		const FProperty* InSourceProperty,
+		uint8* InSourcePtr);
+
+	static bool CopyProperty(
+		const FProperty* InTargetProperty,
+		uint8* InTargetPtr,
+		const FRigVMPropertyPath& InTargetPropertyPath,
+		const FProperty* InSourceProperty,
+		uint8* InSourcePtr,
+		const FRigVMPropertyPath& InSourcePropertyPath);
+
 	static bool CopyProperty(
 		URigVMMemoryStorage* InTargetStorage,
 		int32 InTargetPropertyIndex,
