RigVM: MemoryStorage to support copy operation

diff --git a/Engine/Source/Runtime/RigVM/Private/RigVMCore/RigVMMemoryStorage.cpp b/Engine/Source/Runtime/RigVM/Private/RigVMCore/RigVMMemoryStorage.cpp
--- a/Engine/Source/Runtime/RigVM/Private/RigVMCore/RigVMMemoryStorage.cpp
+++ b/Engine/Source/Runtime/RigVM/Private/RigVMCore/RigVMMemoryStorage.cpp
@@ -1,6 +1,7 @@
 // Copyright Epic Games, Inc. All Rights Reserved.
 
 #include "RigVMCore/RigVMMemoryStorage.h"
+#include "RigVMModule.h"
 
 ////////////////////////////////////////////////////////////////////////////////////////////////////////////////////////
 
@@ -251,11 +252,54 @@ URigVMMemoryStorage* URigVMMemoryStorage::CreateStorage(UObject* InOuter, ERigVM
 	return NewObject<URigVMMemoryStorage>(InOuter, Class, NAME_None, RF_Public | RF_Transactional);
 }
 
+FString URigVMMemoryStorage::GetDataAsString(int32 InPropertyIndex)
+{
+	check(IsValidIndex(InPropertyIndex));
+	const uint8* Data = GetData<uint8>(InPropertyIndex);
+
+	FString Value;
+	GetProperties()[InPropertyIndex]->ExportTextItem(Value, Data, nullptr, nullptr, PPF_None);
+	return Value;
+}
+
+class FRigVMMemoryStorageImportErrorContext : public FOutputDevice
+{
+public:
+
+	int32 NumErrors;
+
+	FRigVMMemoryStorageImportErrorContext()
+		: FOutputDevice()
+		, NumErrors(0)
+	{
+	}
+
+	FORCEINLINE_DEBUGGABLE void Serialize(const TCHAR* V, ELogVerbosity::Type Verbosity, const class FName& Category) override
+	{
+#if WITH_EDITOR
+		UE_LOG(LogRigVM, Display, TEXT("Skipping Importing To MemoryStorage: %s"), V);
+#else
+		UE_LOG(LogRigVM, Error, TEXT("Error Importing To MemoryStorage: %s"), V);
+#endif
+		NumErrors++;
+	}
+};
+
+bool URigVMMemoryStorage::SetDataFromString(int32 InPropertyIndex, const FString& InValue)
+{
+	check(IsValidIndex(InPropertyIndex));
+	uint8* Data = GetData<uint8>(InPropertyIndex);
+	
+	FRigVMMemoryStorageImportErrorContext ErrorPipe;
+	GetProperties()[InPropertyIndex]->ImportText(*InValue, Data, EPropertyPortFlags::PPF_None, nullptr, &ErrorPipe);
+	return ErrorPipe.NumErrors == 0;
+}
+
 bool URigVMMemoryStorage::CopyProperty(
 	const FProperty* InTargetProperty,
 	uint8* InTargetPtr,
 	const FProperty* InSourceProperty,
-	uint8* InSourcePtr)
+	const uint8* InSourcePtr)
 {
 	check(InTargetProperty != nullptr);
 	check(InSourceProperty != nullptr);
@@ -276,7 +320,7 @@ bool URigVMMemoryStorage::CopyProperty(
 	uint8* InTargetPtr,
 	const FRigVMPropertyPath& InTargetPropertyPath,
 	const FProperty* InSourceProperty,
-	uint8* InSourcePtr,
+	const uint8* InSourcePtr,
 	const FRigVMPropertyPath& InSourcePropertyPath)
 {
 	check(InTargetProperty != nullptr);
@@ -321,10 +365,11 @@ bool URigVMMemoryStorage::CopyProperty(
 		}
 	};
 
+	uint8* SourcePtr = (uint8*)InSourcePtr;
 	TraversePropertyPath(InTargetProperty, InTargetPtr, InTargetPropertyPath);
-	TraversePropertyPath(InSourceProperty, InSourcePtr, InSourcePropertyPath);
+	TraversePropertyPath(InSourceProperty, SourcePtr, InSourcePropertyPath);
 
-	return CopyProperty(InTargetProperty, InTargetPtr, InSourceProperty, InSourcePtr);
+	return CopyProperty(InTargetProperty, InTargetPtr, InSourceProperty, SourcePtr);
 }
 
 bool URigVMMemoryStorage::CopyProperty(
@@ -346,6 +391,21 @@ bool URigVMMemoryStorage::CopyProperty(
 	return CopyProperty(TargetProperty, TargetPtr, InTargetPropertyPath, SourceProperty, SourcePtr, InSourcePropertyPath);
 }
 
+#if !UE_RIGVM_UCLASS_BASED_STORAGE_DISABLED
+
+bool URigVMMemoryStorage::CopyProperty(
+	FRigVMMemoryHandle& TargetHandle,
+	FRigVMMemoryHandle& SourceHandle)
+{
+	return CopyProperty(
+		TargetHandle.GetProperty(),
+		TargetHandle.GetData(true),
+		SourceHandle.GetProperty(),
+		SourceHandle.GetData(true));
+}
+
+#endif
+
 FProperty* URigVMMemoryStorage::AddProperty(UClass* InClass, const FPropertyDescription& InProperty, bool bPurge, bool bLink, FField** LinkToProperty)
 {
 	UClass *SuperClass = URigVMMemoryStorage::StaticClass();
diff --git a/Engine/Source/Runtime/RigVM/Private/RigVMCore/RigVMPropertyPath.cpp b/Engine/Source/Runtime/RigVM/Private/RigVMCore/RigVMPropertyPath.cpp
--- a/Engine/Source/Runtime/RigVM/Private/RigVMCore/RigVMPropertyPath.cpp
+++ b/Engine/Source/Runtime/RigVM/Private/RigVMCore/RigVMPropertyPath.cpp
@@ -122,6 +122,26 @@ FRigVMPropertyPath::FRigVMPropertyPath(const FRigVMPropertyPath& InOther)
 {
 }
 
+bool FRigVMPropertyPath::IsDirect() const
+{
+	for(const FRigVMPropertyPathSegment& Segment : Segments)
+	{
+		switch(Segment.Type)
+		{
+			case ERigVMPropertyPathSegmentType::ArrayElement:
+			case ERigVMPropertyPathSegmentType::MapValue:
+			{
+				return true;
+			}
+			default:
+			{
+				break;
+			}
+		}
+	}
+	return false;
+}
+
 uint8* FRigVMPropertyPath::GetData_Internal(uint8* InPtr) const
 {
 	for(const FRigVMPropertyPathSegment& Segment : Segments)
diff --git a/Engine/Source/Runtime/RigVM/Public/RigVMCore/RigVMMemoryStorage.h b/Engine/Source/Runtime/RigVM/Public/RigVMCore/RigVMMemoryStorage.h
--- a/Engine/Source/Runtime/RigVM/Public/RigVMCore/RigVMMemoryStorage.h
+++ b/Engine/Source/Runtime/RigVM/Public/RigVMCore/RigVMMemoryStorage.h
@@ -23,21 +23,59 @@ public:
 
 	FORCEINLINE_DEBUGGABLE FRigVMMemoryHandle()
 		: Ptr(nullptr)
+		, Property(nullptr) 
+		, PropertyPath(nullptr)
 	{}
 
-	FORCEINLINE_DEBUGGABLE uint8* GetData()
+	FORCEINLINE_DEBUGGABLE FRigVMMemoryHandle(uint8* InPtr, const FProperty* InProperty,  const FRigVMPropertyPath* InPropertyPath = nullptr)
+		: Ptr(InPtr)
+		, Property(InProperty)
+		, PropertyPath(InPropertyPath)
 	{
-		return Ptr;
+		if(PropertyPath)
+		{
+			if(PropertyPath->IsDirect())
+			{
+				Ptr = PropertyPath->GetData<uint8>(Ptr);
+				PropertyPath = nullptr;
+			}
+		}
 	}
 
-	FORCEINLINE_DEBUGGABLE const uint8* GetData() const
+	FORCEINLINE uint8* GetData(bool bFollowPropertyPath = false)
 	{
-		return Ptr;
+		return GetData_Internal(bFollowPropertyPath);
+	}
+
+	FORCEINLINE const uint8* GetData(bool bFollowPropertyPath = false) const
+	{
+		return GetData_Internal(bFollowPropertyPath);
+	}
+
+	FORCEINLINE const FProperty* GetProperty() const
+	{
+		return Property;
+	}
+
+	FORCEINLINE const FRigVMPropertyPath* GetPropertyPath() const
+	{
+		return PropertyPath;
 	}
 
 private:
+
+	FORCEINLINE uint8* GetData_Internal(bool bFollowPropertyPath = false) const
+	{
+		if(bFollowPropertyPath && PropertyPath != nullptr)
+		{
+			return PropertyPath->GetData<uint8>(Ptr);
+		}
+		return Ptr;
+	}
 	
 	uint8* Ptr;
+	const FProperty* Property;
+	const FRigVMPropertyPath* PropertyPath;
 
 	friend class URigVM;
 };
@@ -182,18 +220,34 @@ public:
 		return GetData<T>(PropertyIndex, InPropertyPath);
 	}
 
+	FString GetDataAsString(int32 InPropertyIndex);
+	
+	FORCEINLINE FString GetDataAsStringByName(const FName& InName)
+	{
+		const int32 PropertyIndex = GetPropertyIndexByName(InName);
+		return GetDataAsString(PropertyIndex);
+	}
+
+	bool SetDataFromString(int32 InPropertyIndex, const FString& InValue);
+
+	FORCEINLINE bool SetDataFromStringByName(const FName& InName, const FString& InValue)
+	{
+		const int32 PropertyIndex = GetPropertyIndexByName(InName);
+		return SetDataFromString(PropertyIndex, InValue);
+	}
+
 	static bool CopyProperty(
 		const FProperty* InTargetProperty,
 		uint8* InTargetPtr,
 		const FProperty* InSourceProperty,
-		uint8* InSourcePtr);
+		const uint8* InSourcePtr);
 
 	static bool CopyProperty(
 		const FProperty* InTargetProperty,
 		uint8* InTargetPtr,
 		const FRigVMPropertyPath& InTargetPropertyPath,
 		const FProperty* InSourceProperty,
-		uint8* InSourcePtr,
+		const uint8* InSourcePtr,
 		const FRigVMPropertyPath& InSourcePropertyPath);
 
 	static bool CopyProperty(
@@ -204,6 +258,14 @@ public:
 		int32 InSourcePropertyIndex,
 		const FRigVMPropertyPath& InSourcePropertyPath);
 
+#if !UE_RIGVM_UCLASS_BASED_STORAGE_DISABLED
+	
+	static bool CopyProperty(
+		FRigVMMemoryHandle& TargetHandle,
+		FRigVMMemoryHandle& SourceHandle);
+	
+#endif
+
 protected:
 
 	static FProperty* AddProperty(UClass* InClass, const FPropertyDescription& InProperty, bool bPurge = false, bool bLink = true, FField** LinkToProperty = nullptr);
diff --git a/Engine/Source/Runtime/RigVM/Public/RigVMCore/RigVMPropertyPath.h b/Engine/Source/Runtime/RigVM/Public/RigVMCore/RigVMPropertyPath.h
--- a/Engine/Source/Runtime/RigVM/Public/RigVMCore/RigVMPropertyPath.h
+++ b/Engine/Source/Runtime/RigVM/Public/RigVMCore/RigVMPropertyPath.h
@@ -32,6 +32,7 @@ public:
 	FORCEINLINE bool IsValid() const { return Num() > 0; }
 	FORCEINLINE bool IsEmpty() const { return Num() == 0; }
 	FORCEINLINE const FRigVMPropertyPathSegment& operator[](int32 InIndex) const { return Segments[InIndex]; }
+	bool IsDirect() const;
 
 	FORCEINLINE bool operator ==(const FRigVMPropertyPath& Other) const
 	{
