Control Rig : RigVM: Fix non unity build for certain platforms

diff --git a/Engine/Source/Runtime/RigVM/Private/RigVMCore/RigVMMemoryStorage.cpp b/Engine/Source/Runtime/RigVM/Private/RigVMCore/RigVMMemoryStorage.cpp
--- a/Engine/Source/Runtime/RigVM/Private/RigVMCore/RigVMMemoryStorage.cpp
+++ b/Engine/Source/Runtime/RigVM/Private/RigVMCore/RigVMMemoryStorage.cpp
@@ -306,7 +306,7 @@ bool URigVMMemoryStorage::CopyProperty(
 	check(InTargetPtr != nullptr);
 	check(InSourcePtr != nullptr);
 
-	if(!InTargetProperty->SameType(InSourceProperty))
+	if(!ensure(InTargetProperty->SameType(InSourceProperty)))
 	{
 		return false;
 	}
diff --git a/Engine/Source/Runtime/RigVM/Public/RigVMCore/RigVMPropertyPath.h b/Engine/Source/Runtime/RigVM/Public/RigVMCore/RigVMPropertyPath.h
--- a/Engine/Source/Runtime/RigVM/Public/RigVMCore/RigVMPropertyPath.h
+++ b/Engine/Source/Runtime/RigVM/Public/RigVMCore/RigVMPropertyPath.h
@@ -3,6 +3,7 @@
 #pragma once
 
 #include "CoreMinimal.h"
+#include "UObject/UnrealType.h"
 
 enum RIGVM_API ERigVMPropertyPathSegmentType
 {
