Control Rig: URigVMMemoryStorage GetHandle accessor

diff --git a/Engine/Source/Runtime/RigVM/Private/RigVMCore/RigVMMemoryStorage.cpp b/Engine/Source/Runtime/RigVM/Private/RigVMCore/RigVMMemoryStorage.cpp
--- a/Engine/Source/Runtime/RigVM/Private/RigVMCore/RigVMMemoryStorage.cpp
+++ b/Engine/Source/Runtime/RigVM/Private/RigVMCore/RigVMMemoryStorage.cpp
@@ -295,6 +295,20 @@ bool URigVMMemoryStorage::SetDataFromString(int32 InPropertyIndex, const FString
 	return ErrorPipe.NumErrors == 0;
 }
 
+#if !UE_RIGVM_UCLASS_BASED_STORAGE_DISABLED
+
+FRigVMMemoryHandle URigVMMemoryStorage::GetHandle(int32 InPropertyIndex, const FRigVMPropertyPath* InPropertyPath)
+{
+	check(IsValidIndex(InPropertyIndex));
+
+	const FProperty* Property = GetProperties()[InPropertyIndex];
+	uint8* Data = GetData<uint8>(InPropertyIndex);
+
+	return FRigVMMemoryHandle(Data, Property, InPropertyPath);
+}
+
+#endif
+
 bool URigVMMemoryStorage::CopyProperty(
 	const FProperty* InTargetProperty,
 	uint8* InTargetPtr,
diff --git a/Engine/Source/Runtime/RigVM/Public/RigVMCore/RigVMMemoryStorage.h b/Engine/Source/Runtime/RigVM/Public/RigVMCore/RigVMMemoryStorage.h
--- a/Engine/Source/Runtime/RigVM/Public/RigVMCore/RigVMMemoryStorage.h
+++ b/Engine/Source/Runtime/RigVM/Public/RigVMCore/RigVMMemoryStorage.h
@@ -236,6 +236,18 @@ public:
 		return SetDataFromString(PropertyIndex, InValue);
 	}
 
+#if !UE_RIGVM_UCLASS_BASED_STORAGE_DISABLED
+	
+	FRigVMMemoryHandle GetHandle(int32 InPropertyIndex, const FRigVMPropertyPath* InPropertyPath = nullptr);
+
+	FORCEINLINE FRigVMMemoryHandle GetHandleByName(const FName& InName, const FRigVMPropertyPath* InPropertyPath = nullptr)
+	{
+		const int32 PropertyIndex = GetPropertyIndexByName(InName);
+		return GetHandle(PropertyIndex, InPropertyPath);
+	}
+	
+#endif
+	
 	static bool CopyProperty(
 		const FProperty* InTargetProperty,
 		uint8* InTargetPtr,
