ControlRig: PropertyPath / MemoryStorage API improvements

diff --git a/Engine/Source/Runtime/RigVM/Public/RigVMCore/RigVMMemoryStorage.h b/Engine/Source/Runtime/RigVM/Public/RigVMCore/RigVMMemoryStorage.h
--- a/Engine/Source/Runtime/RigVM/Public/RigVMCore/RigVMMemoryStorage.h
+++ b/Engine/Source/Runtime/RigVM/Public/RigVMCore/RigVMMemoryStorage.h
@@ -100,13 +100,22 @@ public:
 		UObject* CPPTypeObject;
 		TArray<EPinContainerType> Containers;
 		FString DefaultValue;
-		
+
+		FPropertyDescription()
+			: Name(NAME_None)
+			, Property(nullptr)
+			, CPPType()
+			, CPPTypeObject(nullptr)
+			, Containers()
+			, DefaultValue()
+		{}
 		FPropertyDescription(const FProperty* InProperty, const FString& InDefaultValue, const FName& InName = NAME_None);
 		FPropertyDescription(const FName& InName, const FString& InCPPType, UObject* InCPPTypeObject, const FString& InDefaultValue);
 
 		static FName SanitizeName(const FName& InName);
 		void SanitizeName();
 
+		bool IsValid() const { return !Name.IsNone(); }
 		int32 NumContainers() const { return Containers.Num(); }
 		bool HasContainer() const { return NumContainers() > 0; }
 		bool IsArray(int32 InContainerIndex = 0) const { return Containers[InContainerIndex] == EPinContainerType::Array; }
diff --git a/Engine/Source/Runtime/RigVM/Public/RigVMCore/RigVMPropertyPath.h b/Engine/Source/Runtime/RigVM/Public/RigVMCore/RigVMPropertyPath.h
--- a/Engine/Source/Runtime/RigVM/Public/RigVMCore/RigVMPropertyPath.h
+++ b/Engine/Source/Runtime/RigVM/Public/RigVMCore/RigVMPropertyPath.h
@@ -37,22 +37,22 @@ public:
 
 	FORCEINLINE bool operator ==(const FRigVMPropertyPath& Other) const
 	{
-		return ToString() == Other.ToString();
+		return GetTypeHash(this) == GetTypeHash(Other);
 	}
 
 	FORCEINLINE bool operator !=(const FRigVMPropertyPath& Other) const
 	{
-		return ToString() != Other.ToString();
+		return GetTypeHash(this) != GetTypeHash(Other);
 	}
 
 	FORCEINLINE bool operator >(const FRigVMPropertyPath& Other) const
 	{
-		return ToString() > Other.ToString();
+		return GetTypeHash(this) > GetTypeHash(Other);
 	}
 
 	FORCEINLINE bool operator <(const FRigVMPropertyPath& Other) const
 	{
-		return ToString() < Other.ToString();
+		return GetTypeHash(this) < GetTypeHash(Other);
 	}
 
 	uint8* GetData_Internal(uint8* InPtr) const;
@@ -63,6 +63,19 @@ public:
 		return (T*)GetData_Internal(InPtr);
 	}
 
+	friend FORCEINLINE uint32 GetTypeHash(const FRigVMPropertyPath& InPropertyPath)
+	{
+		if(InPropertyPath.IsEmpty())
+		{
+			return 0;
+		}
+		
+		return HashCombine(
+			GetTypeHash(InPropertyPath[0].Property),
+			GetTypeHash(InPropertyPath.ToString())
+		);
+	}
+
 	static FRigVMPropertyPath Empty;
 
 private:
