RigVM: Implementing support for new memory storage

All code is still based on the UE_RIGVM_UCLASS_BASED_STORAGE_DISABLED define, which is turned on.
So the feature / changes don't affect anything just yet.

diff --git a/Engine/Plugins/Animation/ControlRig/Source/ControlRig/Private/ControlRig.cpp b/Engine/Plugins/Animation/ControlRig/Source/ControlRig/Private/ControlRig.cpp
--- a/Engine/Plugins/Animation/ControlRig/Source/ControlRig/Private/ControlRig.cpp
+++ b/Engine/Plugins/Animation/ControlRig/Source/ControlRig/Private/ControlRig.cpp
@@ -1056,13 +1056,21 @@ void UControlRig::ExecuteUnits(FRigUnitContext& InOutContext, const FName& InEve
 {
 	if (VM)
 	{
+#if UE_RIGVM_UCLASS_BASED_STORAGE_DISABLED
 		FRigVMMemoryContainer* LocalMemory[] = { VM->WorkMemoryPtr, VM->LiteralMemoryPtr, VM->DebugMemoryPtr };
+#else
+		TArray<URigVMMemoryStorage*> LocalMemory = VM->GetLocalMemoryArray();
+#endif
 		TArray<void*> AdditionalArguments;
 		AdditionalArguments.Add(&InOutContext);
 
 		if (InOutContext.State == EControlRigState::Init)
 		{
+#if UE_RIGVM_UCLASS_BASED_STORAGE_DISABLED
 			VM->Initialize(FRigVMMemoryContainerPtrArray(LocalMemory, 3), AdditionalArguments);
+#else
+			VM->Initialize(LocalMemory, AdditionalArguments);
+#endif
 		}
 		else
 		{
@@ -1079,7 +1087,12 @@ void UControlRig::ExecuteUnits(FRigUnitContext& InOutContext, const FName& InEve
 				}
 			}
 #endif
+
+#if UE_RIGVM_UCLASS_BASED_STORAGE_DISABLED
 			VM->Execute(FRigVMMemoryContainerPtrArray(LocalMemory, 3), AdditionalArguments, InEventName);
+#else
+			VM->Execute(LocalMemory, AdditionalArguments, InEventName);
+#endif
 		}
 	}
 }
diff --git a/Engine/Plugins/Animation/ControlRig/Source/ControlRig/Private/Drawing/ControlRigDrawInterface.cpp b/Engine/Plugins/Animation/ControlRig/Source/ControlRig/Private/Drawing/ControlRigDrawInterface.cpp
--- a/Engine/Plugins/Animation/ControlRig/Source/ControlRig/Private/Drawing/ControlRigDrawInterface.cpp
+++ b/Engine/Plugins/Animation/ControlRig/Source/ControlRig/Private/Drawing/ControlRigDrawInterface.cpp
@@ -19,7 +19,7 @@ void FControlRigDrawInterface::DrawPoint(const FTransform& WorldOffset, const FV
 	Instructions.Add(Instruction);
 }
 
-void FControlRigDrawInterface::DrawPoints(const FTransform& WorldOffset, const FRigVMFixedArray<FVector>& Points, float Size, const FLinearColor& Color)
+void FControlRigDrawInterface::DrawPoints(const FTransform& WorldOffset, const TArrayView<const FVector>& Points, float Size, const FLinearColor& Color)
 {
 	if (!IsEnabled())
 	{
@@ -44,7 +44,7 @@ void FControlRigDrawInterface::DrawLine(const FTransform& WorldOffset, const FVe
 	Instructions.Add(Instruction);
 }
 
-void FControlRigDrawInterface::DrawLines(const FTransform& WorldOffset, const FRigVMFixedArray<FVector>& Positions, const FLinearColor& Color, float Thickness)
+void FControlRigDrawInterface::DrawLines(const FTransform& WorldOffset, const TArrayView<const FVector>& Positions, const FLinearColor& Color, float Thickness)
 {
 	if (!IsEnabled())
 	{
@@ -56,7 +56,7 @@ void FControlRigDrawInterface::DrawLines(const FTransform& WorldOffset, const FR
 	Instructions.Add(Instruction);
 }
 
-void FControlRigDrawInterface::DrawLineStrip(const FTransform& WorldOffset, const FRigVMFixedArray<FVector>& Positions, const FLinearColor& Color, float Thickness)
+void FControlRigDrawInterface::DrawLineStrip(const FTransform& WorldOffset, const TArrayView<const FVector>& Positions, const FLinearColor& Color, float Thickness)
 {
 	if (!IsEnabled())
 	{
@@ -121,7 +121,7 @@ void FControlRigDrawInterface::DrawAxes(const FTransform& WorldOffset, const FTr
 	DrawLine(WorldOffset, Transform.GetLocation(), Transform.TransformPosition(FVector(0.f, 0.f, Size)), FLinearColor::Blue, Thickness);
 }
 
-void FControlRigDrawInterface::DrawAxes(const FTransform& WorldOffset, const FRigVMFixedArray<FTransform>& Transforms, float Size, float Thickness)
+void FControlRigDrawInterface::DrawAxes(const FTransform& WorldOffset, const TArrayView<const FTransform>& Transforms, float Size, float Thickness)
 {
 	if (!IsEnabled())
 	{
@@ -177,7 +177,7 @@ void FControlRigDrawInterface::DrawAxes(const FTransform& WorldOffset, const FTr
 	Instructions.Add(Instruction);
 }
 
-void FControlRigDrawInterface::DrawAxes(const FTransform& WorldOffset, const FRigVMFixedArray<FTransform>& Transforms,
+void FControlRigDrawInterface::DrawAxes(const FTransform& WorldOffset, const TArrayView<const FTransform>& Transforms,
 	const FLinearColor& InColor, float Size, float Thickness)
 {
 	if (!IsEnabled())
diff --git a/Engine/Plugins/Animation/ControlRig/Source/ControlRig/Private/Units/Debug/RigUnit_DebugLineStrip.cpp b/Engine/Plugins/Animation/ControlRig/Source/ControlRig/Private/Units/Debug/RigUnit_DebugLineStrip.cpp
--- a/Engine/Plugins/Animation/ControlRig/Source/ControlRig/Private/Units/Debug/RigUnit_DebugLineStrip.cpp
+++ b/Engine/Plugins/Animation/ControlRig/Source/ControlRig/Private/Units/Debug/RigUnit_DebugLineStrip.cpp
@@ -43,6 +43,6 @@ FRigUnit_DebugLineStripItemSpace_Execute()
 	}
 	else
 	{
-		Context.DrawInterface->DrawLineStrip(WorldOffset, Points, Color, Thickness);
+		Context.DrawInterface->DrawLineStrip(WorldOffset, TArrayView<const FVector>(Points.GetData(), Points.Num()), Color, Thickness);
 	}
 }
diff --git a/Engine/Plugins/Animation/ControlRig/Source/ControlRig/Private/Units/Highlevel/Hierarchy/RigUnit_MultiFABRIK.cpp b/Engine/Plugins/Animation/ControlRig/Source/ControlRig/Private/Units/Highlevel/Hierarchy/RigUnit_MultiFABRIK.cpp
--- a/Engine/Plugins/Animation/ControlRig/Source/ControlRig/Private/Units/Highlevel/Hierarchy/RigUnit_MultiFABRIK.cpp
+++ b/Engine/Plugins/Animation/ControlRig/Source/ControlRig/Private/Units/Highlevel/Hierarchy/RigUnit_MultiFABRIK.cpp
@@ -84,7 +84,7 @@ namespace MultiFABRIK
 	}
 
 	// Forward Evaluator
-	void SolveForwardRecursively(FRigUnit_MultiFABRIK_ChainGroup& InOutCurrentGroup, TArray<FRigUnit_MultiFABRIK_BoneWorkingData>& InBoneTree, const FRigVMFixedArray<FRigUnit_MultiFABRIK_EndEffector>& InEffectors, float InPrecision)
+	void SolveForwardRecursively(FRigUnit_MultiFABRIK_ChainGroup& InOutCurrentGroup, TArray<FRigUnit_MultiFABRIK_BoneWorkingData>& InBoneTree, const TArrayView<const FRigUnit_MultiFABRIK_EndEffector>& InEffectors, float InPrecision)
 	{
 		const int32 ChildNum = InOutCurrentGroup.Children.Num();
 		// solve child first, and then follow
@@ -419,11 +419,13 @@ FRigUnit_MultiFABRIK_Execute()
 			FVector RootLocation = BoneTree[0].Location;
 			// now iterate 
 			// we start from end to root group prioritizing getting to effector
+
+			const TArrayView<const FRigUnit_MultiFABRIK_EndEffector> EffectorsView(Effectors.GetData(), Effectors.Num());
 			for (int32 IterIndex = 0; IterIndex < MaxIterations; ++IterIndex)
 			{
 				// update transform first
 				MultiFABRIK::UpdateTransfomRecusively(ChainGroup, BoneTree);
-				MultiFABRIK::SolveForwardRecursively(ChainGroup, BoneTree, Effectors, Precision);
+				MultiFABRIK::SolveForwardRecursively(ChainGroup, BoneTree, EffectorsView, Precision);
 
 				// root position can change while forward reaching
 				BoneTree[0].Location = RootLocation;
diff --git a/Engine/Plugins/Animation/ControlRig/Source/ControlRig/Private/Units/Highlevel/Hierarchy/RigUnit_TransformConstraint.cpp b/Engine/Plugins/Animation/ControlRig/Source/ControlRig/Private/Units/Highlevel/Hierarchy/RigUnit_TransformConstraint.cpp
--- a/Engine/Plugins/Animation/ControlRig/Source/ControlRig/Private/Units/Highlevel/Hierarchy/RigUnit_TransformConstraint.cpp
+++ b/Engine/Plugins/Animation/ControlRig/Source/ControlRig/Private/Units/Highlevel/Hierarchy/RigUnit_TransformConstraint.cpp
@@ -59,6 +59,7 @@ FRigUnit_TransformConstraintPerItem_Execute()
 						);
 
 
+					const TArrayView<const FConstraintTarget> TargetsView(Targets.GetData(), Targets.Num());
 					for (int32 TargetIndex = 0; TargetIndex < TargetNum; ++TargetIndex)
 					{
 						// talk to Rob about the implication of support both of this
@@ -68,23 +69,23 @@ FRigUnit_TransformConstraintPerItem_Execute()
 
 						if (bTranslationFilterValid && bRotationFilterValid && bScaleFilterValid)
 						{
-							AddConstraintData(Targets, ETransformConstraintType::Parent, TargetIndex, SourceTransform, InputBaseTransform, ConstraintData, ConstraintDataToTargets);
+							AddConstraintData(TargetsView, ETransformConstraintType::Parent, TargetIndex, SourceTransform, InputBaseTransform, ConstraintData, ConstraintDataToTargets);
 						}
 						else
 						{
 							if (bTranslationFilterValid)
 							{
-								AddConstraintData(Targets, ETransformConstraintType::Translation, TargetIndex, SourceTransform, InputBaseTransform, ConstraintData, ConstraintDataToTargets);
+								AddConstraintData(TargetsView, ETransformConstraintType::Translation, TargetIndex, SourceTransform, InputBaseTransform, ConstraintData, ConstraintDataToTargets);
 							}
 
 							if (bRotationFilterValid)
 							{
-								AddConstraintData(Targets, ETransformConstraintType::Rotation, TargetIndex, SourceTransform, InputBaseTransform, ConstraintData, ConstraintDataToTargets);
+								AddConstraintData(TargetsView, ETransformConstraintType::Rotation, TargetIndex, SourceTransform, InputBaseTransform, ConstraintData, ConstraintDataToTargets);
 							}
 
 							if (bScaleFilterValid)
 							{
-								AddConstraintData(Targets, ETransformConstraintType::Scale, TargetIndex, SourceTransform, InputBaseTransform, ConstraintData, ConstraintDataToTargets);
+								AddConstraintData(TargetsView, ETransformConstraintType::Scale, TargetIndex, SourceTransform, InputBaseTransform, ConstraintData, ConstraintDataToTargets);
 							}
 						}
 					}
@@ -146,7 +147,7 @@ FRigUnit_TransformConstraintPerItem_Execute()
 	}
 }
 
-void FRigUnit_TransformConstraintPerItem::AddConstraintData(const FRigVMFixedArray<FConstraintTarget>& Targets, ETransformConstraintType ConstraintType, const int32 TargetIndex, const FTransform& SourceTransform, const FTransform& InBaseTransform, TArray<FConstraintData>& OutConstraintData, TMap<int32, int32>& OutConstraintDataToTargets)
+void FRigUnit_TransformConstraintPerItem::AddConstraintData(const TArrayView<const FConstraintTarget>& Targets, ETransformConstraintType ConstraintType, const int32 TargetIndex, const FTransform& SourceTransform, const FTransform& InBaseTransform, TArray<FConstraintData>& OutConstraintData, TMap<int32, int32>& OutConstraintDataToTargets)
 {
 	const FConstraintTarget& Target = Targets[TargetIndex];
 
diff --git a/Engine/Plugins/Animation/ControlRig/Source/ControlRig/Private/Units/Highlevel/Hierarchy/RigUnit_TransformConstraint.h b/Engine/Plugins/Animation/ControlRig/Source/ControlRig/Private/Units/Highlevel/Hierarchy/RigUnit_TransformConstraint.h
--- a/Engine/Plugins/Animation/ControlRig/Source/ControlRig/Private/Units/Highlevel/Hierarchy/RigUnit_TransformConstraint.h
+++ b/Engine/Plugins/Animation/ControlRig/Source/ControlRig/Private/Units/Highlevel/Hierarchy/RigUnit_TransformConstraint.h
@@ -171,7 +171,7 @@ struct CONTROLRIG_API FRigUnit_TransformConstraintPerItem : public FRigUnit_High
 	bool bUseInitialTransforms;
 
 private:
-	static void AddConstraintData(const FRigVMFixedArray<FConstraintTarget>& Targets, ETransformConstraintType ConstraintType, const int32 TargetIndex, const FTransform& SourceTransform, const FTransform& InBaseTransform, TArray<FConstraintData>& OutConstraintData, TMap<int32, int32>& OutConstraintDataToTargets);
+	static void AddConstraintData(const TArrayView<const FConstraintTarget>& Targets, ETransformConstraintType ConstraintType, const int32 TargetIndex, const FTransform& SourceTransform, const FTransform& InBaseTransform, TArray<FConstraintData>& OutConstraintData, TMap<int32, int32>& OutConstraintDataToTargets);
 
 	UPROPERTY(transient)
 	FRigUnit_TransformConstraint_WorkData WorkData;
diff --git a/Engine/Plugins/Animation/ControlRig/Source/ControlRig/Private/Units/Math/RigUnit_MathRBFInterpolate.cpp b/Engine/Plugins/Animation/ControlRig/Source/ControlRig/Private/Units/Math/RigUnit_MathRBFInterpolate.cpp
--- a/Engine/Plugins/Animation/ControlRig/Source/ControlRig/Private/Units/Math/RigUnit_MathRBFInterpolate.cpp
+++ b/Engine/Plugins/Animation/ControlRig/Source/ControlRig/Private/Units/Math/RigUnit_MathRBFInterpolate.cpp
@@ -170,7 +170,7 @@ template<typename T>
 void FRigUnit_MathRBFInterpolateQuatBase::GetInterpolatedWeights(
 	EControlRigState State,
 	FRigUnit_MathRBFInterpolateQuatWorkData& WorkData,
-	const FRigVMFixedArray<T>& Targets,
+	const TArrayView<const T>& Targets,
 	const FQuat &Input,
 	ERBFQuatDistanceType DistanceFunction,
 	ERBFKernelType SmoothingFunction,
@@ -194,7 +194,7 @@ void FRigUnit_MathRBFInterpolateQuatBase::GetInterpolatedWeights(
 	if (State == EControlRigState::Init)
 	{
 		WorkData.Targets = Quats;
-		WorkData.Hash = HashTargets(FRigVMFixedArray<FQuat>(Quats.GetData(), Quats.Num()));
+		WorkData.Hash = HashTargets(TArrayView<const FQuat>(Quats.GetData(), Quats.Num()));
 		WorkData.Interpolator = TRBFInterpolator<FQuat>(WorkData.Targets, WeightFunc);
 		WorkData.bAreTargetsConstant = true;
 	}
@@ -202,7 +202,7 @@ void FRigUnit_MathRBFInterpolateQuatBase::GetInterpolatedWeights(
 	{
 		// Are the targets still constant? If the hash changed, then, no, and they shall stay non-constant 
 		// for the lifetime of this unit's instance.
-		uint64 Hash = HashTargets(FRigVMFixedArray<FQuat>(Quats.GetData(), Quats.Num()));
+		uint64 Hash = HashTargets(TArrayView<const FQuat>(Quats.GetData(), Quats.Num()));
 		if (Hash != WorkData.Hash)
 		{
 			WorkData.bAreTargetsConstant = false;
@@ -218,7 +218,7 @@ void FRigUnit_MathRBFInterpolateQuatBase::GetInterpolatedWeights(
 	WorkData.Interpolator.Interpolate(Weights, Input, /*bClip=*/true, bNormalizeOutput);
 }
 
-uint64 FRigUnit_MathRBFInterpolateQuatBase::HashTargets(const FRigVMFixedArray<FQuat>& Targets)
+uint64 FRigUnit_MathRBFInterpolateQuatBase::HashTargets(const TArrayView<const FQuat>& Targets)
 {
 	uint64 Hash = 0;
 	for (const FQuat& T : Targets)
@@ -234,7 +234,7 @@ template<typename T>
 void FRigUnit_MathRBFInterpolateVectorBase::GetInterpolatedWeights(
 	EControlRigState State,
 	FRigUnit_MathRBFInterpolateVectorWorkData& WorkData,
-	const FRigVMFixedArray<T>& Targets,
+	const TArrayView<const T>& Targets,
 	const FVector& Input,
 	ERBFVectorDistanceType DistanceFunction,
 	ERBFKernelType SmoothingFunction,
@@ -257,14 +257,14 @@ void FRigUnit_MathRBFInterpolateVectorBase::GetInterpolatedWeights(
 	if (State == EControlRigState::Init)
 	{
 		WorkData.Targets = Vectors;
-		WorkData.Hash = HashTargets(FRigVMFixedArray<FVector>(Vectors.GetData(), Vectors.Num()));
+		WorkData.Hash = HashTargets(TArrayView<const FVector>(Vectors.GetData(), Vectors.Num()));
 		WorkData.Interpolator = TRBFInterpolator<FVector>(WorkData.Targets, WeightFunc);
 	}
 	else if (WorkData.bAreTargetsConstant)
 	{
 		// Are the targets still constant? If the hash changed, then, no, and they shall stay non-constant 
 		// for the lifetime of this unit's instance.
-		uint64 Hash = HashTargets(FRigVMFixedArray<FVector>(Vectors.GetData(), Vectors.Num()));
+		uint64 Hash = HashTargets(TArrayView<const FVector>(Vectors.GetData(), Vectors.Num()));
 		if (Hash != WorkData.Hash)
 		{
 			WorkData.bAreTargetsConstant = false;
@@ -280,7 +280,7 @@ void FRigUnit_MathRBFInterpolateVectorBase::GetInterpolatedWeights(
 	WorkData.Interpolator.Interpolate(Weights, Input, /*bClip=*/true, bNormalizeOutput);
 }
 
-uint64 FRigUnit_MathRBFInterpolateVectorBase::HashTargets(const FRigVMFixedArray<FVector>& Targets)
+uint64 FRigUnit_MathRBFInterpolateVectorBase::HashTargets(const TArrayView<const FVector>& Targets)
 {
 	uint64 Hash = 0;
 	for (const FVector& T : Targets)
@@ -296,9 +296,10 @@ FRigUnit_MathRBFInterpolateQuatFloat_Execute()
 {
 	DECLARE_SCOPE_HIERARCHICAL_COUNTER_RIGUNIT()
 
+	const TArrayView<const FMathRBFInterpolateQuatFloat_Target> TargetsView(Targets.GetData(), Targets.Num());
 	TArray<float, TMemStackAllocator<>> Weights;
 	GetInterpolatedWeights<FMathRBFInterpolateQuatFloat_Target>(
-		Context.State, WorkData, Targets, Input,  
+		Context.State, WorkData, TargetsView, Input,  
 		DistanceFunction, SmoothingFunction, SmoothingAngle, 
 		bNormalizeOutput, TwistAxis, Weights);
 
@@ -314,9 +315,10 @@ FRigUnit_MathRBFInterpolateQuatVector_Execute()
 {
 	DECLARE_SCOPE_HIERARCHICAL_COUNTER_RIGUNIT()
 
+	const TArrayView<const FMathRBFInterpolateQuatVector_Target> TargetsView(Targets.GetData(), Targets.Num());
 	TArray<float, TMemStackAllocator<>> Weights;
 	GetInterpolatedWeights<FMathRBFInterpolateQuatVector_Target>(
-		Context.State, WorkData, Targets, Input,
+		Context.State, WorkData, TargetsView, Input,
 		DistanceFunction, SmoothingFunction, SmoothingAngle, 
 		bNormalizeOutput, TwistAxis, Weights);
 
@@ -332,9 +334,10 @@ FRigUnit_MathRBFInterpolateQuatColor_Execute()
 {
 	DECLARE_SCOPE_HIERARCHICAL_COUNTER_RIGUNIT()
 
+	const TArrayView<const FMathRBFInterpolateQuatColor_Target> TargetsView(Targets.GetData(), Targets.Num());
 	TArray<float, TMemStackAllocator<>> Weights;
 	GetInterpolatedWeights<FMathRBFInterpolateQuatColor_Target>(
-		Context.State, WorkData, Targets, Input,
+		Context.State, WorkData, TargetsView, Input,
 		DistanceFunction, SmoothingFunction, SmoothingAngle, 
 		bNormalizeOutput, TwistAxis, Weights);
 
@@ -349,9 +352,10 @@ FRigUnit_MathRBFInterpolateQuatQuat_Execute()
 {
 	DECLARE_SCOPE_HIERARCHICAL_COUNTER_RIGUNIT()
 
+	const TArrayView<const FMathRBFInterpolateQuatQuat_Target> TargetsView(Targets.GetData(), Targets.Num());
 	TArray<float, TMemStackAllocator<>> Weights;
 	GetInterpolatedWeights<FMathRBFInterpolateQuatQuat_Target>(
-		Context.State, WorkData, Targets, Input,
+		Context.State, WorkData, TargetsView, Input,
 		DistanceFunction, SmoothingFunction, SmoothingAngle,
 		bNormalizeOutput, TwistAxis, Weights);
 	
@@ -368,9 +372,10 @@ FRigUnit_MathRBFInterpolateQuatXform_Execute()
 {
 	DECLARE_SCOPE_HIERARCHICAL_COUNTER_RIGUNIT()
 
+	const TArrayView<const FMathRBFInterpolateQuatXform_Target> TargetsView(Targets.GetData(), Targets.Num());
 	TArray<float, TMemStackAllocator<>> Weights;
 	GetInterpolatedWeights<FMathRBFInterpolateQuatXform_Target>(
-		Context.State, WorkData, Targets, Input,
+		Context.State, WorkData, TargetsView, Input,
 		DistanceFunction, SmoothingFunction, SmoothingAngle,
 		bNormalizeOutput, TwistAxis, Weights);
 
@@ -397,9 +402,10 @@ FRigUnit_MathRBFInterpolateVectorFloat_Execute()
 {
 	DECLARE_SCOPE_HIERARCHICAL_COUNTER_RIGUNIT()
 
+	const TArrayView<const FMathRBFInterpolateVectorFloat_Target> TargetsView(Targets.GetData(), Targets.Num());
 	TArray<float, TMemStackAllocator<>> Weights;
 	GetInterpolatedWeights<FMathRBFInterpolateVectorFloat_Target>(
-		Context.State, WorkData, Targets, Input,
+		Context.State, WorkData, TargetsView, Input,
 		DistanceFunction, SmoothingFunction, SmoothingRadius, 
 		bNormalizeOutput, Weights);
 
@@ -415,9 +421,10 @@ FRigUnit_MathRBFInterpolateVectorVector_Execute()
 {
 	DECLARE_SCOPE_HIERARCHICAL_COUNTER_RIGUNIT()
 
+	const TArrayView<const FMathRBFInterpolateVectorVector_Target> TargetsView(Targets.GetData(), Targets.Num());
 	TArray<float, TMemStackAllocator<>> Weights;
 	GetInterpolatedWeights<FMathRBFInterpolateVectorVector_Target>(
-		Context.State, WorkData, Targets, Input,
+		Context.State, WorkData, TargetsView, Input,
 		DistanceFunction, SmoothingFunction, SmoothingRadius, 
 		bNormalizeOutput, Weights);
 
@@ -433,9 +440,10 @@ FRigUnit_MathRBFInterpolateVectorColor_Execute()
 {
 	DECLARE_SCOPE_HIERARCHICAL_COUNTER_RIGUNIT()
 
+	const TArrayView<const FMathRBFInterpolateVectorColor_Target> TargetsView(Targets.GetData(), Targets.Num());
 	TArray<float, TMemStackAllocator<>> Weights;
 	GetInterpolatedWeights<FMathRBFInterpolateVectorColor_Target>(
-		Context.State, WorkData, Targets, Input,
+		Context.State, WorkData, TargetsView, Input,
 		DistanceFunction, SmoothingFunction, SmoothingRadius, 
 		bNormalizeOutput, Weights);
 
@@ -451,9 +459,10 @@ FRigUnit_MathRBFInterpolateVectorQuat_Execute()
 {
 	DECLARE_SCOPE_HIERARCHICAL_COUNTER_RIGUNIT()
 
+	const TArrayView<const FMathRBFInterpolateVectorQuat_Target> TargetsView(Targets.GetData(), Targets.Num());
 	TArray<float, TMemStackAllocator<>> Weights;
 	GetInterpolatedWeights<FMathRBFInterpolateVectorQuat_Target>(
-		Context.State, WorkData, Targets, Input,
+		Context.State, WorkData, TargetsView, Input,
 		DistanceFunction, SmoothingFunction, SmoothingRadius, 
 		bNormalizeOutput, Weights);
 
@@ -470,9 +479,10 @@ FRigUnit_MathRBFInterpolateVectorXform_Execute()
 {
 	DECLARE_SCOPE_HIERARCHICAL_COUNTER_RIGUNIT()
 
+	const TArrayView<const FMathRBFInterpolateVectorXform_Target> TargetsView(Targets.GetData(), Targets.Num());
 	TArray<float, TMemStackAllocator<>> Weights;
 	GetInterpolatedWeights<FMathRBFInterpolateVectorXform_Target>(
-		Context.State, WorkData, Targets, Input,
+		Context.State, WorkData, TargetsView, Input,
 		DistanceFunction, SmoothingFunction, SmoothingRadius, 
 		bNormalizeOutput, Weights);
 
diff --git a/Engine/Plugins/Animation/ControlRig/Source/ControlRig/Private/Units/Math/RigUnit_MathRBFInterpolate.h b/Engine/Plugins/Animation/ControlRig/Source/ControlRig/Private/Units/Math/RigUnit_MathRBFInterpolate.h
--- a/Engine/Plugins/Animation/ControlRig/Source/ControlRig/Private/Units/Math/RigUnit_MathRBFInterpolate.h
+++ b/Engine/Plugins/Animation/ControlRig/Source/ControlRig/Private/Units/Math/RigUnit_MathRBFInterpolate.h
@@ -107,11 +107,12 @@ struct CONTROLRIG_API FRigUnit_MathRBFInterpolateQuatBase :
 	FRigUnit_MathRBFInterpolateQuatWorkData WorkData;
 
 protected:
+	
 	template<typename T>
 	static void GetInterpolatedWeights(
 		EControlRigState State,
 		FRigUnit_MathRBFInterpolateQuatWorkData& WorkData,
-		const FRigVMFixedArray<T>& Targets,
+		const TArrayView<const T>& Targets,
 		const FQuat& Input,
 		ERBFQuatDistanceType DistanceFunction,
 		ERBFKernelType SmoothingFunction,
@@ -121,7 +122,7 @@ protected:
 		TArray<float, TMemStackAllocator<>>& Weights
 	);
 
-	static uint64 HashTargets(const FRigVMFixedArray<FQuat>& Targets);
+	static uint64 HashTargets(const TArrayView<const FQuat>& Targets);
 };
 
 USTRUCT(meta = (Abstract, Keywords = "RBF,Interpolate,Vector"))
@@ -153,7 +154,7 @@ protected:
 	static void GetInterpolatedWeights(
 		EControlRigState State,
 		FRigUnit_MathRBFInterpolateVectorWorkData& WorkData,
-		const FRigVMFixedArray<T>& Targets,
+		const TArrayView<const T>& Targets,
 		const FVector& Input,
 		ERBFVectorDistanceType DistanceFunction,
 		ERBFKernelType SmoothingFunction,
@@ -162,7 +163,7 @@ protected:
 		TArray<float, TMemStackAllocator<>>& Weights
 	);
 
-	static uint64 HashTargets(const FRigVMFixedArray<FVector>& Targets);
+	static uint64 HashTargets(const TArrayView<const FVector>& Targets);
 };
 
 // The actual unit implementation declarations.
diff --git a/Engine/Plugins/Animation/ControlRig/Source/ControlRig/Public/Drawing/ControlRigDrawInterface.h b/Engine/Plugins/Animation/ControlRig/Source/ControlRig/Public/Drawing/ControlRigDrawInterface.h
--- a/Engine/Plugins/Animation/ControlRig/Source/ControlRig/Public/Drawing/ControlRigDrawInterface.h
+++ b/Engine/Plugins/Animation/ControlRig/Source/ControlRig/Public/Drawing/ControlRigDrawInterface.h
@@ -30,15 +30,15 @@ public:
 	GENERATED_BODY()
 
 	void DrawPoint(const FTransform& WorldOffset, const FVector& Position, float Size, const FLinearColor& Color);
-	void DrawPoints(const FTransform& WorldOffset, const FRigVMFixedArray<FVector>& Points, float Size, const FLinearColor& Color);
+	void DrawPoints(const FTransform& WorldOffset, const TArrayView<const FVector>& Points, float Size, const FLinearColor& Color);
 	void DrawLine(const FTransform& WorldOffset, const FVector& LineStart, const FVector& LineEnd, const FLinearColor& Color, float Thickness = 0.f);
-	void DrawLines(const FTransform& WorldOffset, const FRigVMFixedArray<FVector>& Positions, const FLinearColor& Color, float Thickness = 0.f);
-	void DrawLineStrip(const FTransform& WorldOffset, const FRigVMFixedArray<FVector>& Positions, const FLinearColor& Color, float Thickness = 0.f);
+	void DrawLines(const FTransform& WorldOffset, const TArrayView<const FVector>& Positions, const FLinearColor& Color, float Thickness = 0.f);
+	void DrawLineStrip(const FTransform& WorldOffset, const TArrayView<const FVector>& Positions, const FLinearColor& Color, float Thickness = 0.f);
 	void DrawBox(const FTransform& WorldOffset, const FTransform& Transform, const FLinearColor& Color, float Thickness = 0.f);
 	void DrawAxes(const FTransform& WorldOffset, const FTransform& Transform, float Size, float Thickness = 0.f);
-	void DrawAxes(const FTransform& WorldOffset, const FRigVMFixedArray<FTransform>& Transforms, float Size, float Thickness = 0.f);
+	void DrawAxes(const FTransform& WorldOffset, const TArrayView<const FTransform>& Transforms, float Size, float Thickness = 0.f);
 	void DrawAxes(const FTransform& WorldOffset, const FTransform& Transform, const FLinearColor& InColor, float Size, float Thickness = 0.f);
-	void DrawAxes(const FTransform& WorldOffset, const FRigVMFixedArray<FTransform>& Transforms, const FLinearColor& InColor, float Size, float Thickness = 0.f);
+	void DrawAxes(const FTransform& WorldOffset, const TArrayView<const FTransform>& Transforms, const FLinearColor& InColor, float Size, float Thickness = 0.f);
 	void DrawRectangle(const FTransform& WorldOffset, const FTransform& Transform, float Size, const FLinearColor& Color, float Thickness);
 	void DrawArc(const FTransform& WorldOffset, const FTransform& Transform, float Radius, float MinimumAngle, float MaximumAngle, const FLinearColor& Color, float Thickness, int32 Detail);
 	void DrawBezier(const FTransform& WorldOffset, const FCRFourPointBezier& InBezier, float MinimumU, float MaximumU, const FLinearColor& Color, float Thickness, int32 Detail);
diff --git a/Engine/Plugins/Animation/ControlRig/Source/ControlRigDeveloper/Private/ControlRigBlueprint.cpp b/Engine/Plugins/Animation/ControlRig/Source/ControlRigDeveloper/Private/ControlRigBlueprint.cpp
--- a/Engine/Plugins/Animation/ControlRig/Source/ControlRigDeveloper/Private/ControlRigBlueprint.cpp
+++ b/Engine/Plugins/Animation/ControlRig/Source/ControlRigDeveloper/Private/ControlRigBlueprint.cpp
@@ -2635,6 +2635,7 @@ void UControlRigBlueprint::HandleModifiedEvent(ERigVMGraphNotifType InNotifType,
 						
 						if(!bRequiresRecompile)
 						{
+#if UE_RIGVM_UCLASS_BASED_STORAGE_DISABLED
 							TArray<FString> DefaultValues;
 							if (RootPin->IsArray())
 							{
@@ -2647,12 +2648,19 @@ void UControlRigBlueprint::HandleModifiedEvent(ERigVMGraphNotifType InNotifType,
 							{
 								DefaultValues.Add(RootPin->GetDefaultValue());
 							}
+#else
+							const FString DefaultValue = RootPin->GetDefaultValue();
+#endif
 
 							UControlRigBlueprintGeneratedClass* RigClass = GetControlRigBlueprintGeneratedClass();
 							UControlRig* CDO = Cast<UControlRig>(RigClass->GetDefaultObject(true /* create if needed */));
 							if (CDO->VM != nullptr)
 							{
+#if UE_RIGVM_UCLASS_BASED_STORAGE_DISABLED
 								CDO->VM->SetRegisterValueFromString(*Operand, RootPin->GetCPPType(), RootPin->GetCPPTypeObject(), DefaultValues);
+#else
+								CDO->VM->SetPropertyValueFromString(*Operand, DefaultValue);
+#endif
 							}
 
 							TArray<UObject*> ArchetypeInstances;
@@ -2664,7 +2672,11 @@ void UControlRigBlueprint::HandleModifiedEvent(ERigVMGraphNotifType InNotifType,
 								{
 									if (InstancedControlRig->VM)
 									{
+#if UE_RIGVM_UCLASS_BASED_STORAGE_DISABLED
 										InstancedControlRig->VM->SetRegisterValueFromString(*Operand, RootPin->GetCPPType(), RootPin->GetCPPTypeObject(), DefaultValues);
+#else
+										InstancedControlRig->VM->SetPropertyValueFromString(*Operand, DefaultValue);
+#endif
 									}
 								}
 							}
diff --git a/Engine/Plugins/Animation/ControlRig/Source/ControlRigEditor/Private/Graph/SControlRigGraphNode.cpp b/Engine/Plugins/Animation/ControlRig/Source/ControlRigEditor/Private/Graph/SControlRigGraphNode.cpp
--- a/Engine/Plugins/Animation/ControlRig/Source/ControlRigEditor/Private/Graph/SControlRigGraphNode.cpp
+++ b/Engine/Plugins/Animation/ControlRig/Source/ControlRigEditor/Private/Graph/SControlRigGraphNode.cpp
@@ -872,9 +872,15 @@ void SControlRigGraphNode::GetNodeInfoPopups(FNodeInfoContext* Context, TArray<F
 						FString PinHash = URigVMCompiler::GetPinHash(ModelPin, nullptr, true);
 						if (const FRigVMOperand* WatchOperand = RigBlueprint->PinToOperandMap.Find(PinHash))
 						{
+#if UE_RIGVM_UCLASS_BASED_STORAGE_DISABLED
 							FRigVMMemoryContainer& Memory = ActiveObject->GetVM()->GetDebugMemory();
-
 							TArray<FString> DefaultValues = Memory.GetRegisterValueAsString(*WatchOperand, ModelPin->GetCPPType(), ModelPin->GetCPPTypeObject());
+#else
+							URigVMMemoryStorage* Memory = ActiveObject->GetVM()->GetDebugMemory();
+							FString DebugValue = Memory->GetDataAsString(WatchOperand->GetRegisterIndex());
+							TArray<FString> DefaultValues = URigVMPin::SplitDefaultValue(DebugValue);
+#endif
+							
 							if (DefaultValues.Num() == 1)
 							{
 								WatchText = DefaultValues[0];
diff --git a/Engine/Plugins/Experimental/ControlRigSpline/Source/ControlRigSpline/Private/ControlRigSplineTypes.cpp b/Engine/Plugins/Experimental/ControlRigSpline/Source/ControlRigSpline/Private/ControlRigSplineTypes.cpp
--- a/Engine/Plugins/Experimental/ControlRigSpline/Source/ControlRigSpline/Private/ControlRigSplineTypes.cpp
+++ b/Engine/Plugins/Experimental/ControlRigSpline/Source/ControlRigSpline/Private/ControlRigSplineTypes.cpp
@@ -2,7 +2,7 @@
 
 #include "ControlRigSplineTypes.h"
 
-void FControlRigSpline::SetControlPoints(const TArray<FVector>& InPoints, const ESplineType SplineMode, const int32 SamplesPerSegment)
+void FControlRigSpline::SetControlPoints(const TArrayView<const FVector>& InPoints, const ESplineType SplineMode, const int32 SamplesPerSegment)
 {
 	const int32 ControlPointsCount = InPoints.Num();
 	if (ControlPointsCount < 4)
diff --git a/Engine/Plugins/Experimental/ControlRigSpline/Source/ControlRigSpline/Private/ControlRigSplineUnits.cpp b/Engine/Plugins/Experimental/ControlRigSpline/Source/ControlRigSpline/Private/ControlRigSplineUnits.cpp
--- a/Engine/Plugins/Experimental/ControlRigSpline/Source/ControlRigSpline/Private/ControlRigSplineUnits.cpp
+++ b/Engine/Plugins/Experimental/ControlRigSpline/Source/ControlRigSpline/Private/ControlRigSplineUnits.cpp
@@ -13,7 +13,8 @@ FRigUnit_ControlRigSplineFromPoints_Execute()
 		case EControlRigState::Init:
 		case EControlRigState::Update:
 		{
-			Spline.SetControlPoints(Points, SplineMode, SamplesPerSegment);
+			const TArrayView<const FVector> PointsView(Points.GetData(), Points.Num());
+			Spline.SetControlPoints(PointsView, SplineMode, SamplesPerSegment);
 			break;
 		}
 		default:
@@ -43,7 +44,8 @@ FRigUnit_SetSplinePoints_Execute()
 		case EControlRigState::Init:
 		case EControlRigState::Update:
 		{
-			Spline.SetControlPoints(Points, Spline.SplineData->SplineMode, Spline.SplineData->SamplesPerSegment);
+			const TArrayView<const FVector> PointsView(Points.GetData(), Points.Num());
+			Spline.SetControlPoints(PointsView, Spline.SplineData->SplineMode, Spline.SplineData->SamplesPerSegment);
 			break;
 		}
 		default:
diff --git a/Engine/Plugins/Experimental/ControlRigSpline/Source/ControlRigSpline/Public/ControlRigSplineTypes.h b/Engine/Plugins/Experimental/ControlRigSpline/Source/ControlRigSpline/Public/ControlRigSplineTypes.h
--- a/Engine/Plugins/Experimental/ControlRigSpline/Source/ControlRigSpline/Public/ControlRigSplineTypes.h
+++ b/Engine/Plugins/Experimental/ControlRigSpline/Source/ControlRigSpline/Public/ControlRigSplineTypes.h
@@ -73,7 +73,7 @@ struct CONTROLRIGSPLINE_API FControlRigSpline
 	* @param InPoints	The control points to set.
 	* @param forceRebuild	If true, will build the spline from scratch.
 	*/
-	void SetControlPoints(const TArray<FVector>& InPoints, const ESplineType SplineMode = ESplineType::BSpline, const int32 SamplesPerSegment = 16);
+	void SetControlPoints(const TArrayView<const FVector>& InPoints, const ESplineType SplineMode = ESplineType::BSpline, const int32 SamplesPerSegment = 16);
 
 	/**
 	* Populates OutPoints and returns the number of control points in the spline.
diff --git a/Engine/Plugins/Experimental/FullBodyIK/Source/FullBodyIK/Private/FBIKConstraintLib.cpp b/Engine/Plugins/Experimental/FullBodyIK/Source/FullBodyIK/Private/FBIKConstraintLib.cpp
--- a/Engine/Plugins/Experimental/FullBodyIK/Source/FullBodyIK/Private/FBIKConstraintLib.cpp
+++ b/Engine/Plugins/Experimental/FullBodyIK/Source/FullBodyIK/Private/FBIKConstraintLib.cpp
@@ -40,7 +40,7 @@ namespace FBIKConstraintLib
 	}
 
 	// for now allow linkdata to be modified here
-	void BuildConstraints(const TArray<FFBIKConstraintOption>& Constraints, TArray<ConstraintType>& OutConstraints,
+	void BuildConstraints(const TArrayView<const FFBIKConstraintOption>& Constraints, TArray<ConstraintType>& OutConstraints,
 		const URigHierarchy* Hierarchy, TArray<FFBIKLinkData>& InOutLinkData, const TMap<int32, FRigElementKey>& LinkDataToHierarchyIndices,
 		const TMap<FRigElementKey, int32>& HierarchyToLinkDataMap)
 	{
diff --git a/Engine/Plugins/Experimental/FullBodyIK/Source/FullBodyIK/Private/RigUnit_FullbodyIK.cpp b/Engine/Plugins/Experimental/FullBodyIK/Source/FullBodyIK/Private/RigUnit_FullbodyIK.cpp
--- a/Engine/Plugins/Experimental/FullBodyIK/Source/FullBodyIK/Private/RigUnit_FullbodyIK.cpp
+++ b/Engine/Plugins/Experimental/FullBodyIK/Source/FullBodyIK/Private/RigUnit_FullbodyIK.cpp
@@ -79,7 +79,7 @@ static void AddToEffectorTarget(int32 EffectorIndex, const FRigElementKey& Effec
 	}
 }
 
-static void AddEffectors(URigHierarchy* Hierarchy, const FRigElementKey& Root, const TArray<FFBIKEndEffector>& Effectors,
+static void AddEffectors(URigHierarchy* Hierarchy, const FRigElementKey& Root, const TArrayView<const FFBIKEndEffector>& Effectors,
 	TArray<FFBIKLinkData>& LinkData, TMap<int32, FFBIKEffectorTarget>& EffectorTargets, TArray<int32>& EffectorLinkIndices, 
 	TMap<int32, FRigElementKey>& LinkDataToHierarchyIndices, TMap<FRigElementKey, int32>& HierarchyToLinkDataMap, const FSolverInput& SolverProperty)
 {
@@ -208,7 +208,8 @@ FRigUnit_FullbodyIK_Execute()
 		HierarchyToLinkDataMap.Reset();
 
 		// verify the chain
-		AddEffectors(Hierarchy, Root, Effectors, LinkData, EffectorTargets, EffectorLinkIndices, LinkDataToHierarchyIndices, HierarchyToLinkDataMap, SolverProperty);
+		const TArrayView<const FFBIKEndEffector> EffectorsView(Effectors.GetData(), Effectors.Num());
+		AddEffectors(Hierarchy, Root, EffectorsView, LinkData, EffectorTargets, EffectorLinkIndices, LinkDataToHierarchyIndices, HierarchyToLinkDataMap, SolverProperty);
 	}
 
 	if (Context.State == EControlRigState::Update)
@@ -222,7 +223,8 @@ FRigUnit_FullbodyIK_Execute()
 			{
 				DECLARE_SCOPE_HIERARCHICAL_COUNTER(TEXT("Build Constraint"))
 				//Build constraints
-				FBIKConstraintLib::BuildConstraints(Constraints, InternalConstraints, Hierarchy, LinkData, LinkDataToHierarchyIndices, HierarchyToLinkDataMap);
+				const TArrayView<const FFBIKConstraintOption> ConstraintsView(Constraints.GetData(), Constraints.Num());
+				FBIKConstraintLib::BuildConstraints(ConstraintsView, InternalConstraints, Hierarchy, LinkData, LinkDataToHierarchyIndices, HierarchyToLinkDataMap);
 			}
 
 			// during only editor and update
diff --git a/Engine/Plugins/Experimental/FullBodyIK/Source/FullBodyIK/Public/FBIKConstraintLib.h b/Engine/Plugins/Experimental/FullBodyIK/Source/FullBodyIK/Public/FBIKConstraintLib.h
--- a/Engine/Plugins/Experimental/FullBodyIK/Source/FullBodyIK/Public/FBIKConstraintLib.h
+++ b/Engine/Plugins/Experimental/FullBodyIK/Source/FullBodyIK/Public/FBIKConstraintLib.h
@@ -14,7 +14,7 @@ namespace FBIKConstraintLib
 {
 	// constraint related
 	void FULLBODYIK_API ApplyConstraint(TArray<FFBIKLinkData>& InOutLinkData, TArray<ConstraintType>* Constraints);
-	void FULLBODYIK_API BuildConstraints(const TArray<FFBIKConstraintOption>& Constraints, TArray<ConstraintType>& OutConstraints,
+	void FULLBODYIK_API BuildConstraints(const TArrayView<const FFBIKConstraintOption>& Constraints, TArray<ConstraintType>& OutConstraints,
 		const URigHierarchy* Hierarchy, TArray<FFBIKLinkData>& InOutLinkData, const TMap<int32, FRigElementKey>& LinkDataToHierarchyIndices,
 		const TMap<FRigElementKey, int32>& HierarchyToLinkDataMap);
 };
diff --git a/Engine/Source/Developer/RigVMDeveloper/Private/RigVMCompiler/RigVMAST.cpp b/Engine/Source/Developer/RigVMDeveloper/Private/RigVMCompiler/RigVMAST.cpp
--- a/Engine/Source/Developer/RigVMDeveloper/Private/RigVMCompiler/RigVMAST.cpp
+++ b/Engine/Source/Developer/RigVMDeveloper/Private/RigVMCompiler/RigVMAST.cpp
@@ -1884,7 +1884,6 @@ bool FRigVMParserAST::FoldConstantValuesToLiterals(URigVMGraph* InGraph, URigVMC
 	URigVM* TempVM = NewObject<URigVM>(GetTransientPackage());
 	
 	URigVMCompiler* TempCompiler = NewObject<URigVMCompiler>(GetTransientPackage());
-	TempCompiler->Settings.ConsolidateWorkRegisters = false;
 	TempCompiler->Settings.SetupNodeInstructionIndex = false;
 	TempCompiler->Settings.IsPreprocessorPhase = true;
 	TempCompiler->Settings.EnablePinWatches = false;
@@ -1892,13 +1891,27 @@ bool FRigVMParserAST::FoldConstantValuesToLiterals(URigVMGraph* InGraph, URigVMC
 
 	TempCompiler->Compile(InGraph, InController, TempVM, InExternalVariables, InRigVMUserData, &Operands, TempAST);
 
+#if UE_RIGVM_UCLASS_BASED_STORAGE_DISABLED
+	
 	FRigVMMemoryContainer* Memory[] = { TempVM->WorkMemoryPtr, TempVM->LiteralMemoryPtr, TempVM->DebugMemoryPtr };
-
 	for (const FRigVMUserDataArray& RigVMUserData : InRigVMUserData)
 	{
 		TempVM->Execute(FRigVMMemoryContainerPtrArray(Memory, 3), RigVMUserData);
 	}
 
+#else
+
+	TArray<URigVMMemoryStorage*> Memory;
+	Memory.Add(TempVM->GetWorkMemory());
+	Memory.Add(TempVM->GetLiteralMemory());
+	Memory.Add(TempVM->GetDebugMemory());
+	for (const FRigVMUserDataArray& RigVMUserData : InRigVMUserData)
+	{
+		TempVM->Execute(Memory, RigVMUserData);
+	}
+
+#endif
+
 	// copy the values out of the temp VM and set them on the cached value
 	for (const FRigVMASTProxy& PinToComputeProxy : PinsToCompute)
 	{
@@ -1921,15 +1934,30 @@ bool FRigVMParserAST::FoldConstantValuesToLiterals(URigVMGraph* InGraph, URigVMC
 
 		FString PinHash = URigVMCompiler::GetPinHash(RootPin, RootVarExpr, false);
 		const FRigVMOperand& Operand = Operands.FindChecked(PinHash);
+
+#if UE_RIGVM_UCLASS_BASED_STORAGE_DISABLED
+		
 		TArray<FString> DefaultValues = TempVM->GetWorkMemory().GetRegisterValueAsString(Operand, RootPin->GetCPPType(), RootPin->GetCPPTypeObject());
 
+#else
+
+		FString DefaultValue = TempVM->GetWorkMemory()->GetDataAsString(Operand.GetRegisterIndex());
+		if(DefaultValue.IsEmpty())
+		{
+			continue;
+		}
+
+#endif
+
 		// FString TempDefaultValue = FString::Printf(TEXT("(%s)"), *FString::Join(DefaultValues, TEXT(",")));
 		// UE_LOG(LogRigVMDeveloper, Display, TEXT("Computed constant value '%s' = '%s'"), *PinHash, *TempDefaultValue);
 
+#if UE_RIGVM_UCLASS_BASED_STORAGE_DISABLED
 		if (DefaultValues.Num() == 0)
 		{
 			continue;
 		}
+#endif
 
 		TArray<FString> SegmentNames;
 		if (!URigVMPin::SplitPinPath(PinToCompute->GetSegmentPath(), SegmentNames))
@@ -1937,11 +1965,13 @@ bool FRigVMParserAST::FoldConstantValuesToLiterals(URigVMGraph* InGraph, URigVMC
 			SegmentNames.Add(PinToCompute->GetName());
 		}
 
+#if UE_RIGVM_UCLASS_BASED_STORAGE_DISABLED
 		FString DefaultValue = DefaultValues[0];
 		if (RootPin->IsArray())
 		{
 			DefaultValue = FString::Printf(TEXT("(%s)"), *FString::Join(DefaultValues, TEXT(",")));
 		}
+#endif
 
 		URigVMPin* PinForDefaultValue = RootPin;
 		while (PinForDefaultValue != PinToCompute && SegmentNames.Num() > 0)
diff --git a/Engine/Source/Developer/RigVMDeveloper/Private/RigVMCompiler/RigVMCompiler.cpp b/Engine/Source/Developer/RigVMDeveloper/Private/RigVMCompiler/RigVMCompiler.cpp
--- a/Engine/Source/Developer/RigVMDeveloper/Private/RigVMCompiler/RigVMCompiler.cpp
+++ b/Engine/Source/Developer/RigVMDeveloper/Private/RigVMCompiler/RigVMCompiler.cpp
@@ -51,46 +51,76 @@ FRigVMCompileSettings::FRigVMCompileSettings()
 	, SurpressWarnings(false)
 	, SurpressErrors(false)
 	, EnablePinWatches(true)
-	, ConsolidateWorkRegisters(false)
 	, IsPreprocessorPhase(false)
 	, ASTSettings(FRigVMParserASTSettings::Optimized())
 	, SetupNodeInstructionIndex(true)
 {
 }
 
-int32 FRigVMCompilerWorkData::IncRefRegister(int32 InRegister, int32 InIncrement)
+#if !UE_RIGVM_UCLASS_BASED_STORAGE_DISABLED
+
+FRigVMOperand FRigVMCompilerWorkData::AddProperty(
+	ERigVMMemoryType InMemoryType,
+	const FName& InName,
+	const FString& InCPPType,
+	UObject* InCPPTypeObject,
+	const FString& InDefaultValue)
+{
+	check(bSetupMemory);
+	
+	URigVMMemoryStorage::FPropertyDescription Description(InName, InCPPType, InCPPTypeObject, InDefaultValue);
+
+	TArray<URigVMMemoryStorage::FPropertyDescription>& PropertyArray = PropertyDescriptions.FindOrAdd(InMemoryType);
+	const int32 PropertyIndex = PropertyArray.Add(Description);
+
+	return FRigVMOperand(InMemoryType, PropertyIndex);
+}
+
+FRigVMOperand FRigVMCompilerWorkData::FindProperty(ERigVMMemoryType InMemoryType, const FName& InName)
 {
-	ensure(InIncrement > 0);
-	if (int32* RefCountPtr = RegisterRefCount.Find(InRegister))
+	TArray<URigVMMemoryStorage::FPropertyDescription>* PropertyArray = PropertyDescriptions.Find(InMemoryType);
+	if(PropertyArray)
 	{
-		*RefCountPtr += InIncrement;
-		RecordRefCountStep(InRegister, *RefCountPtr);
-		return *RefCountPtr;
+		for(int32 Index=0;Index<PropertyArray->Num();Index++)
+		{
+			if(PropertyArray->operator[](Index).Name == InName)
+			{
+				return FRigVMOperand(InMemoryType, Index);
+			}
+		}
 	}
-	RegisterRefCount.Add(InRegister, InIncrement);
-	RecordRefCountStep(InRegister, InIncrement);
-	return InIncrement;
+	return FRigVMOperand();
 }
 
-int32 FRigVMCompilerWorkData::DecRefRegister(int32 InRegister, int32 InDecrement)
+URigVMMemoryStorage::FPropertyDescription FRigVMCompilerWorkData::GetProperty(const FRigVMOperand& InOperand)
 {
-	ensure(InDecrement > 0);
-	if (int32* RefCountPtr = RegisterRefCount.Find(InRegister))
+	TArray<URigVMMemoryStorage::FPropertyDescription>* PropertyArray = PropertyDescriptions.Find(InOperand.GetMemoryType());
+	if(PropertyArray)
 	{
-		int32& RefCount = *RefCountPtr;
-		RefCount = FMath::Max<int32>(0, RefCount - InDecrement);
-		RecordRefCountStep(InRegister, RefCount);
-		return RefCount;
+		if(PropertyArray->IsValidIndex(InOperand.GetRegisterIndex()))
+		{
+			return PropertyArray->operator[](InOperand.GetRegisterIndex());
+		}
 	}
-	return 0;
+	return URigVMMemoryStorage::FPropertyDescription();
 }
 
-void FRigVMCompilerWorkData::RecordRefCountStep(int32  InRegister, int32 InRefCount)
+int32 FRigVMCompilerWorkData::FindPropertyPath(const FString& InSegmentPath) const
 {
-	const FName& Name = VM->GetWorkMemory().Registers[InRegister].Name;
-	RefCountSteps.Add(TPair<FName, int32>(Name, InRefCount));
+	for(int32 Index = 0; Index < PropertyPathDescriptions.Num(); Index++)
+	{
+		const FRigVMCompilerWorkData::FPropertyPathDescription& Description = PropertyPathDescriptions[Index]; 
+		if(Description.SegmentPath == InSegmentPath)
+		{
+			return Index;
+		}
+	}
+
+	return INDEX_NONE;
 }
 
+#endif
+
 URigVMCompiler::URigVMCompiler()
 {
 }
@@ -221,7 +251,13 @@ bool URigVMCompiler::Compile(URigVMGraph* InGraph, URigVMController* InControlle
 							FRigVMASTProxy PinProxy = FRigVMASTProxy::MakeFromUObject(ValuePin);
 							FRigVMVarExprAST* TempVarExpr = AST->MakeExpr<FRigVMVarExprAST>(FRigVMExprAST::EType::Literal, PinProxy);
 							FRigVMOperand Operand = FindOrAddRegister(TempVarExpr, WorkData, false);
+
+#if UE_RIGVM_UCLASS_BASED_STORAGE_DISABLED
 							WorkData.VM->GetLiteralMemory().SetRegisterValueFromString(Operand, LocalVariable.CPPType, LocalVariable.CPPTypeObject, {LocalVariable.DefaultValue});
+#else
+							check(Operand.GetMemoryType() == ERigVMMemoryType::Literal);
+							WorkData.VM->GetLiteralMemory()->SetDataFromString(Operand.GetRegisterIndex(), LocalVariable.DefaultValue);
+#endif
 							break;
 						}
 					}
@@ -281,6 +317,66 @@ bool URigVMCompiler::Compile(URigVMGraph* InGraph, URigVMController* InControlle
 		}
 	}
 
+#if !UE_RIGVM_UCLASS_BASED_STORAGE_DISABLED
+
+	// now that we have determined the needed memory, let's
+	// setup properties as needed as well as property paths
+	TArray<ERigVMMemoryType> MemoryTypes;
+	MemoryTypes.Add(ERigVMMemoryType::Work);
+	MemoryTypes.Add(ERigVMMemoryType::Literal);
+	MemoryTypes.Add(ERigVMMemoryType::Debug);
+
+	for(ERigVMMemoryType MemoryType : MemoryTypes)
+	{
+		URigVMMemoryStorage** MemoryStorageObject = nullptr;
+		switch(MemoryType)
+		{
+			case ERigVMMemoryType::Literal:
+			{
+				MemoryStorageObject = &WorkData.VM->LiteralMemoryStorageObject;
+				break;
+			}
+			case ERigVMMemoryType::Work:
+			{
+				MemoryStorageObject = &WorkData.VM->WorkMemoryStorageObject;
+				break;
+			}
+			case ERigVMMemoryType::Debug:
+			{
+				MemoryStorageObject = &WorkData.VM->DebugMemoryStorageObject;
+				break;
+			}
+			default:
+			{
+				checkNoEntry();
+				break;
+			}
+		}
+
+		TArray<URigVMMemoryStorage::FPropertyDescription>* Properties = WorkData.PropertyDescriptions.Find(MemoryType);
+		if(Properties == nullptr)
+		{
+			*MemoryStorageObject = NewObject<URigVMMemoryStorage>(WorkData.VM, NAME_None, RF_Public | RF_Transactional);
+			continue;
+		}
+
+		UPackage* Package = InGraph->GetOutermost();
+		UClass* MemoryClass = URigVMMemoryStorage::CreateStorageClass(Package, MemoryType, *Properties);
+		*MemoryStorageObject = NewObject<URigVMMemoryStorage>(WorkData.VM, MemoryClass, NAME_None, RF_Public | RF_Transactional); 
+	}
+
+	WorkData.VM->PropertyPaths.Reset();
+	for(FRigVMCompilerWorkData::FPropertyPathDescription Description : WorkData.PropertyPathDescriptions)
+	{
+		URigVMMemoryStorage* MemoryStorageObject = WorkData.VM->GetMemoryByType(Description.OriginalOperand.GetMemoryType());
+		const FProperty* Property = MemoryStorageObject->GetProperties()[Description.OriginalOperand.GetRegisterIndex()];
+
+		FRigVMPropertyPath PropertyPath(Property, Description.SegmentPath);
+		WorkData.VM->PropertyPaths.Add(PropertyPath);
+	}
+	
+#endif
+
 	WorkData.bSetupMemory = false;
 	WorkData.ExprComplete.Reset();
 	for (FRigVMExprAST* RootExpr : *AST)
@@ -505,28 +601,6 @@ int32 URigVMCompiler::TraverseCallExtern(const FRigVMCallExternExprAST* InExpr,
 		WorkData.DefaultStructs.Add(DefaultStruct);
 		TraverseChildren(InExpr, WorkData);
 		WorkData.DefaultStructs.Pop();
-
-		if (Settings.ConsolidateWorkRegisters)
-		{
-			for (const FRigVMExprAST* Child : *InExpr)
-			{
-				const FRigVMExprAST* VarExpr = Child;
-				if (VarExpr->IsA(FRigVMExprAST::EType::CachedValue))
-				{
-					VarExpr = VarExpr->To<FRigVMCachedValueExprAST>()->GetVarExpr();
-				}
-				if (VarExpr->IsA(FRigVMExprAST::EType::Var))
-				{
-					if (const FRigVMOperand* Operand = WorkData.ExprToOperand.Find(VarExpr->To<FRigVMVarExprAST>()))
-					{
-						if (Operand->GetMemoryType() != ERigVMMemoryType::Literal)
-						{
-							WorkData.DecRefRegister(Operand->GetRegisterIndex());
-						}
-					}
-				}
-			}
-		}
 	}
 	else
 	{
@@ -724,17 +798,7 @@ void URigVMCompiler::TraverseAssign(const FRigVMAssignExprAST* InExpr, FRigVMCom
 
 	FRigVMOperand Source = WorkData.ExprToOperand.FindChecked(SourceExpr);
 
-	if (WorkData.bSetupMemory)
-	{
-		if (Settings.ConsolidateWorkRegisters)
-		{
-			if (Source.GetMemoryType() != ERigVMMemoryType::Literal)
-			{
-				WorkData.DecRefRegister(Source.GetRegisterIndex());
-			}
-		}
-	}
-	else
+	if (!WorkData.bSetupMemory)
 	{
 		const FRigVMVarExprAST* TargetExpr = InExpr->GetFirstParentOfType(FRigVMVarExprAST::EType::Var)->To<FRigVMVarExprAST>();
 		FRigVMOperand Target = WorkData.ExprToOperand.FindChecked(TargetExpr);
@@ -744,7 +808,7 @@ void URigVMCompiler::TraverseAssign(const FRigVMAssignExprAST* InExpr, FRigVMCom
 		{
 			struct Local
 			{
-				static void SetupRegisterOffset(URigVM* VM, URigVMPin* Pin, FRigVMOperand& Operand, const FRigVMVarExprAST* VarExpr, bool bSource)
+				static void SetupRegisterOffset(URigVM* VM, URigVMPin* Pin, FRigVMOperand& Operand, const FRigVMVarExprAST* VarExpr, bool bSource, FRigVMCompilerWorkData& WorkData)
 				{
 					URigVMPin* RootPin = Pin->GetRootPin();
 					if (Pin == RootPin)
@@ -785,6 +849,8 @@ void URigVMCompiler::TraverseAssign(const FRigVMAssignExprAST* InExpr, FRigVMCom
 						ArrayIndex = FCString::Atoi(*SegmentArrayIndex);
 					}
 
+#if UE_RIGVM_UCLASS_BASED_STORAGE_DISABLED
+					
 					if (Operand.GetMemoryType() == ERigVMMemoryType::External)
 					{
 						const FRigVMExternalVariable& ExternalVariable = VM->GetExternalVariables()[Operand.GetRegisterIndex()];
@@ -798,11 +864,28 @@ void URigVMCompiler::TraverseAssign(const FRigVMAssignExprAST* InExpr, FRigVMCom
 							(Operand.GetMemoryType() == ERigVMMemoryType::Work ? VM->GetWorkMemory() : VM->GetDebugMemory());
 						Operand = Memory.GetOperand(Operand.GetRegisterIndex(), SegmentPath, ArrayIndex);
 					}
+
+#else
+
+					int32 PropertyPathIndex = WorkData.FindPropertyPath(SegmentPath);
+					if(PropertyPathIndex == INDEX_NONE)
+					{
+						check(WorkData.bSetupMemory);
+						
+						FRigVMCompilerWorkData::FPropertyPathDescription PropertyPathDescription;
+						PropertyPathDescription.OriginalOperand = Operand;
+						PropertyPathDescription.SegmentPath = SegmentPath;
+						PropertyPathIndex = WorkData.PropertyPathDescriptions.Add(PropertyPathDescription);
+					}
+					
+					Operand = FRigVMOperand(Operand.GetMemoryType(), Operand.GetRegisterIndex(), PropertyPathIndex);
+					
+#endif
 				}
 			};
 
-			Local::SetupRegisterOffset(WorkData.VM, InExpr->GetSourcePin(), Source, SourceExpr, true);
-			Local::SetupRegisterOffset(WorkData.VM, InExpr->GetTargetPin(), Target, TargetExpr, false);
+			Local::SetupRegisterOffset(WorkData.VM, InExpr->GetSourcePin(), Source, SourceExpr, true, WorkData);
+			Local::SetupRegisterOffset(WorkData.VM, InExpr->GetTargetPin(), Target, TargetExpr, false, WorkData);
 		}
 
 		WorkData.VM->GetByteCode().AddCopyOp(WorkData.VM->GetCopyOpForOperands(Source, Target));
@@ -1015,6 +1098,9 @@ void URigVMCompiler::TraverseSelect(const FRigVMSelectExprAST* InExpr, FRigVMCom
 			if (!WorkData.IntegerLiterals.Contains(CaseIndex))
 			{
 				FName LiteralName = *FString::FromInt(CaseIndex);
+
+#if UE_RIGVM_UCLASS_BASED_STORAGE_DISABLED
+				
 				int32 Register = WorkData.VM->GetLiteralMemory().Add<int32>(LiteralName, CaseIndex);
  
 #if WITH_EDITORONLY_DATA
@@ -1023,12 +1109,26 @@ void URigVMCompiler::TraverseSelect(const FRigVMSelectExprAST* InExpr, FRigVMCom
 #endif
 
 				FRigVMOperand Operand = WorkData.VM->GetLiteralMemory().GetOperand(Register);
+				
+#else
+
+				const FString DefaultValue = FString::FromInt(CaseIndex);
+				FRigVMOperand Operand = WorkData.AddProperty(
+					ERigVMMemoryType::Literal,
+					LiteralName,
+					TEXT("int32"),
+					nullptr,
+					DefaultValue);
+				
+#endif
 				WorkData.IntegerLiterals.Add(CaseIndex, Operand);
 			}
 		}
 
 		if (!WorkData.ComparisonOperand.IsValid())
 		{
+#if UE_RIGVM_UCLASS_BASED_STORAGE_DISABLED
+			
 			int32 Register = WorkData.VM->GetWorkMemory().Add<bool>(FName(TEXT("IntEquals")), false);
 
 #if WITH_EDITORONLY_DATA
@@ -1037,6 +1137,17 @@ void URigVMCompiler::TraverseSelect(const FRigVMSelectExprAST* InExpr, FRigVMCom
 #endif
 
 			WorkData.ComparisonOperand = WorkData.VM->GetWorkMemory().GetOperand(Register);
+
+#else
+
+			WorkData.ComparisonOperand = WorkData.AddProperty(
+				ERigVMMemoryType::Work,
+				FName(TEXT("IntEquals")),
+				TEXT("bool"),
+				nullptr,
+				TEXT("false"));
+
+#endif
 		}
 		return;
 	}
@@ -1637,7 +1748,8 @@ FRigVMOperand URigVMCompiler::FindOrAddRegister(const FRigVMVarExprAST* InVarExp
 	URigVMPin::FPinOverride PinOverride(InVarExpr->GetProxy(), PinOverrides);
 
 	URigVMPin* Pin = InVarExpr->GetPin();
-	FString BaseCPPType = Pin->IsArray() ? Pin->GetArrayElementCppType() : Pin->GetCPPType();
+	FString CPPType = Pin->GetCPPType();
+	FString BaseCPPType = Pin->IsArray() ? Pin->GetArrayElementCppType() : CPPType;
 	FString Hash = GetPinHash(Pin, InVarExpr, bIsDebugValue);
 	FRigVMOperand Operand;
 	FString RegisterKey = Hash;
@@ -1665,14 +1777,26 @@ FRigVMOperand URigVMCompiler::FindOrAddRegister(const FRigVMVarExprAST* InVarExp
 		}
 	}
 
+#if UE_RIGVM_UCLASS_BASED_STORAGE_DISABLED
+	
 	FRigVMMemoryContainer& Memory = 
 	    bIsLiteral ? WorkData.VM->GetLiteralMemory() :
 			(bIsDebugValue ? WorkData.VM->GetDebugMemory() : WorkData.VM->GetWorkMemory());
 
+	const ERigVMMemoryType MemoryType = Memory.GetMemoryType();
+
+#else
+
+	const ERigVMMemoryType MemoryType =
+		bIsLiteral ? ERigVMMemoryType::Literal:
+		(bIsDebugValue ? ERigVMMemoryType::Debug : ERigVMMemoryType::Work);
+	
+#endif
+
 	FRigVMOperand const* ExistingOperand = WorkData.PinPathToOperand->Find(Hash);
 	if (ExistingOperand)
 	{
-		if(ExistingOperand->GetMemoryType() == Memory.GetMemoryType())
+		if(ExistingOperand->GetMemoryType() == MemoryType)
 		{
 			if (!bIsDebugValue)
 			{
@@ -1708,9 +1832,28 @@ FRigVMOperand URigVMCompiler::FindOrAddRegister(const FRigVMVarExprAST* InVarExp
 					Operand.RegisterOffset = INDEX_NONE;
 					if (!SegmentPath.IsEmpty())
 					{
+#if UE_RIGVM_UCLASS_BASED_STORAGE_DISABLED
+						
 						const FRigVMExternalVariable& ExternalVariable = WorkData.VM->GetExternalVariables()[Operand.GetRegisterIndex()];
 						UScriptStruct* ScriptStruct = CastChecked<UScriptStruct>(ExternalVariable.TypeObject);
 						Operand.RegisterOffset = WorkData.VM->GetWorkMemory().GetOrAddRegisterOffset(Operand.GetRegisterIndex(), ScriptStruct, SegmentPath, 0 /*ArrayIndex */);
+
+#else
+
+						int32 PropertyPathIndex = WorkData.FindPropertyPath(SegmentPath);
+						if(PropertyPathIndex == INDEX_NONE)
+						{
+							check(WorkData.bSetupMemory);
+
+							FRigVMCompilerWorkData::FPropertyPathDescription PropertyPathDescription;
+							PropertyPathDescription.OriginalOperand = Operand;
+							PropertyPathDescription.SegmentPath = SegmentPath;
+							PropertyPathIndex = WorkData.PropertyPathDescriptions.Add(PropertyPathDescription);
+						}
+
+						Operand.RegisterOffset = (uint16)PropertyPathIndex;
+						
+#endif
 					}
 				}
 			}
@@ -1722,6 +1865,7 @@ FRigVMOperand URigVMCompiler::FindOrAddRegister(const FRigVMVarExprAST* InVarExp
 	{
 		FName RegisterName = *RegisterKey;
 
+		FString JoinedDefaultValue;
 		TArray<FString> DefaultValues;
 		if (Pin->IsArray())
 		{
@@ -1733,26 +1877,38 @@ FRigVMOperand URigVMCompiler::FindOrAddRegister(const FRigVMVarExprAST* InVarExp
 
 				int32 DesiredArraySize = VMStruct->GetArraySize(Pin->GetFName(), WorkData.RigVMUserData);
 
-				DefaultValues = URigVMPin::SplitDefaultValue(Pin->GetDefaultValue(PinOverride));
+				JoinedDefaultValue = Pin->GetDefaultValue(PinOverride);
+				DefaultValues = URigVMPin::SplitDefaultValue(JoinedDefaultValue);
 
 				if (DefaultValues.Num() != DesiredArraySize)
 				{
-					FString DefaultValue;
+					FString FirstDefaultValue;
 					if (Pin->GetArraySize() > 0)
 					{
-						DefaultValue = Pin->GetSubPins()[0]->GetDefaultValue(PinOverride);
+						FirstDefaultValue = Pin->GetSubPins()[0]->GetDefaultValue(PinOverride);
 					}
 
 					DefaultValues.Reset();
 					for (int32 Index = 0; Index < DesiredArraySize; Index++)
 					{
-						DefaultValues.Add(DefaultValue);
+						DefaultValues.Add(FirstDefaultValue);
 					}
 				}
 			}
 			else
 			{
-				DefaultValues = URigVMPin::SplitDefaultValue(Pin->GetDefaultValue(PinOverride));
+				JoinedDefaultValue = Pin->GetDefaultValue(PinOverride);
+				if(!JoinedDefaultValue.IsEmpty())
+				{
+					if(JoinedDefaultValue[0] == TCHAR('('))
+					{
+						DefaultValues = URigVMPin::SplitDefaultValue(JoinedDefaultValue);
+					}
+					else
+					{
+						DefaultValues.Add(JoinedDefaultValue);
+					}
+				}
 			}
 
 			while (DefaultValues.Num() < Pin->GetSubPins().Num())
@@ -1765,16 +1921,19 @@ FRigVMOperand URigVMCompiler::FindOrAddRegister(const FRigVMVarExprAST* InVarExp
 			FString EnumValueStr = EnumNode->GetDefaultValue(PinOverride);
 			if (UEnum* Enum = EnumNode->GetEnum())
 			{
-				DefaultValues.Add(FString::FromInt((int32)Enum->GetValueByNameString(EnumValueStr)));
+				JoinedDefaultValue = FString::FromInt((int32)Enum->GetValueByNameString(EnumValueStr));
+				DefaultValues.Add(JoinedDefaultValue);
 			}
 			else
 			{
-				DefaultValues.Add(FString::FromInt(0));
+				JoinedDefaultValue = FString::FromInt(0);
+				DefaultValues.Add(JoinedDefaultValue);
 			}
 		}
 		else
 		{
-			DefaultValues.Add(Pin->GetDefaultValue(PinOverride));
+			JoinedDefaultValue = Pin->GetDefaultValue(PinOverride);
+			DefaultValues.Add(JoinedDefaultValue);
 		}
 
 		UScriptStruct* ScriptStruct = Pin->GetScriptStruct();
@@ -1783,139 +1942,13 @@ FRigVMOperand URigVMCompiler::FindOrAddRegister(const FRigVMVarExprAST* InVarExp
 			ScriptStruct = GetScriptStructForCPPType(BaseCPPType);
 		}
 
-		bool bRegisterIsMultiUse = Settings.ConsolidateWorkRegisters;
-		if (bRegisterIsMultiUse)
-		{
-			bRegisterIsMultiUse = !bIsLiteral && !bIsDebugValue;
-		}
-		if (bRegisterIsMultiUse)
-		{
-			bRegisterIsMultiUse = Pin->GetDirection() == ERigVMPinDirection::Output || Pin->GetDirection() == ERigVMPinDirection::IO;
-		}
-		if (bRegisterIsMultiUse)
-		{
-			bRegisterIsMultiUse = !bIsExecutePin && !Pin->IsDynamicArray();
-		}
-		if (bRegisterIsMultiUse)
-		{
-			if (URigVMUnitNode* UnitNode = Cast<URigVMUnitNode>(Pin->GetNode()))
-			{
-				if (UnitNode->IsLoopNode())
-				{
-					bRegisterIsMultiUse = false;
-				}
-			}
-		}
-
-		// look for a register to reuse
-		if (bRegisterIsMultiUse)
-		{
-			for (TPair<int32, int32> Pair : WorkData.RegisterRefCount)
-			{
-				int32 ExistingRegisterIndex = Pair.Key;
-				if (Pair.Value > 0)
-				{
-					continue;
-				}
-
-				if (Pin->IsArray())
-				{
-					if (!Memory[ExistingRegisterIndex].IsArray())
-					{
-						continue;
-					}
-
-					if (Pin->GetSubPins().Num() != Memory[ExistingRegisterIndex].ElementCount)
-					{
-						continue;
-					}
-				}
-
-				if (ScriptStruct)
-				{
-					if (Memory[ExistingRegisterIndex].ScriptStructIndex == INDEX_NONE)
-					{
-						continue;
-					}
-					if (Memory.GetScriptStruct(ExistingRegisterIndex) != ScriptStruct)
-					{
-						continue;
-					}
-				}
-				else if (Pin->GetCPPType() == TEXT("FString"))
-				{
-					if (Memory[ExistingRegisterIndex].Type != ERigVMRegisterType::String)
-					{
-						continue;
-					}
-				}
-				else if (Pin->GetCPPType() == TEXT("FName"))
-				{
-					if (Memory[ExistingRegisterIndex].Type != ERigVMRegisterType::Name)
-					{
-						continue;
-					}
-				}
-				else
-				{
-					if (Memory[ExistingRegisterIndex].Type != ERigVMRegisterType::Plain)
-					{
-						continue;
-					}
-
-					if (Pin->GetCPPType() == TEXT("bool"))
-					{
-						if (Memory[ExistingRegisterIndex].ElementSize != sizeof(bool))
-						{
-							continue;
-						}
-					}
-					else if (Pin->GetCPPType() == TEXT("int32"))
-					{
-						if (Memory[ExistingRegisterIndex].ElementSize != sizeof(int32))
-						{
-							continue;
-						}
-					}
-					else if (Pin->GetCPPType() == TEXT("float"))
-					{
-						if (Memory[ExistingRegisterIndex].ElementSize != sizeof(float))
-						{
-							continue;
-						}
-					}
-					else if (Pin->GetCPPType() == TEXT("double"))
-					{
-						if (Memory[ExistingRegisterIndex].ElementSize != sizeof(double))
-						{
-							continue;
-						}
-					}
-					else if (UEnum* Enum = Pin->GetEnum())
-					{
-						if (Memory[ExistingRegisterIndex].ElementSize != sizeof(uint8))
-						{
-							continue;
-						}
-					}
-					else
-					{
-						ensure(false);
-					}
-				}
-
-				WorkData.IncRefRegister(ExistingRegisterIndex, 1 + Pin->GetLinkedTargetPins(true).Num());
-				Operand = Memory.GetOperand(ExistingRegisterIndex);
-				break;
-			}
-		}
-
 		if (!Operand.IsValid())
 		{
 			const int32 NumSlices = 1;
 			int32 Register = INDEX_NONE;
 
 			// debug watch register might already exists - look for them by name
+#if UE_RIGVM_UCLASS_BASED_STORAGE_DISABLED
 			if(bIsDebugValue && Memory.SupportsNames())
 			{
 				Register = Memory.GetIndex(RegisterName);
@@ -1955,7 +1988,27 @@ FRigVMOperand URigVMCompiler::FindOrAddRegister(const FRigVMVarExprAST* InVarExp
 					}
 				}
 			}
+#else
+			if(bIsDebugValue)
+			{
+				Operand = WorkData.FindProperty(MemoryType, RegisterName);
+				if(Operand.IsValid())
+				{
+					URigVMMemoryStorage::FPropertyDescription Property = WorkData.GetProperty(Operand);
+					if(Property.IsValid())
+					{
+						if(ExistingOperand == nullptr)
+						{
+							WorkData.PinPathToOperand->Add(Hash, Operand);
+						}
+						return Operand;
+					}
+				}
+			}
+#endif
 
+#if UE_RIGVM_UCLASS_BASED_STORAGE_DISABLED
+			
 			if (ScriptStruct)
 			{
 				TArray<uint8, TAlignedHeapAllocator<16>> Data;
@@ -2128,13 +2181,16 @@ FRigVMOperand URigVMCompiler::FindOrAddRegister(const FRigVMVarExprAST* InVarExp
 			}
 
 			Memory[Operand.GetRegisterIndex()].bIsArray = Pin->IsArray();
+		}
+
+#else
 
-			if (!bIsParameter && !bIsLiteral && !bIsDebugValue && Settings.ConsolidateWorkRegisters && !bIsExecutePin && Pin->GetDirection() != ERigVMPinDirection::Hidden)
-			{
-				WorkData.IncRefRegister(Operand.GetRegisterIndex(), 1 + Pin->GetLinkedTargetPins(true).Num());
-			}
 		}
 
+		Operand = WorkData.AddProperty(MemoryType, RegisterName, CPPType, Pin->GetCPPTypeObject(), JoinedDefaultValue);
+		
+#endif
+
 		if (bIsParameter && !bIsLiteral)
 		{
 			URigVMParameterNode* ParameterNode = Cast<URigVMParameterNode>(Pin->GetNode());
diff --git a/Engine/Source/Developer/RigVMDeveloper/Private/RigVMModel/RigVMController.cpp b/Engine/Source/Developer/RigVMDeveloper/Private/RigVMModel/RigVMController.cpp
--- a/Engine/Source/Developer/RigVMDeveloper/Private/RigVMModel/RigVMController.cpp
+++ b/Engine/Source/Developer/RigVMDeveloper/Private/RigVMModel/RigVMController.cpp
@@ -6097,6 +6097,7 @@ bool URigVMController::BindPinToVariable(URigVMPin* InPin, const FString& InNewB
 		FRigVMExternalVariable ExternalVariable = GetVariableByName(*VariableName);
 		if (ExternalVariable.IsValid(true))
 		{
+#if UE_RIGVM_UCLASS_BASED_STORAGE_DISABLED
 			FRigVMRegisterOffset RegisterOffset;
 			if (!SegmentPath.IsEmpty())
 			{
@@ -6106,6 +6107,12 @@ bool URigVMController::BindPinToVariable(URigVMPin* InPin, const FString& InNewB
 			{
 				return false;
 			}
+#else
+			if (!InPin->CanBeBoundToVariable(ExternalVariable, SegmentPath))
+			{
+				return false;
+			}
+#endif
 		}
 		else
 		{
diff --git a/Engine/Source/Developer/RigVMDeveloper/Private/RigVMModel/RigVMPin.cpp b/Engine/Source/Developer/RigVMDeveloper/Private/RigVMModel/RigVMPin.cpp
--- a/Engine/Source/Developer/RigVMDeveloper/Private/RigVMModel/RigVMPin.cpp
+++ b/Engine/Source/Developer/RigVMDeveloper/Private/RigVMModel/RigVMPin.cpp
@@ -749,7 +749,11 @@ bool URigVMPin::IsBoundToVariable(const URigVMPin::FPinOverride& InOverride) con
 	return !GetBoundVariablePath(InOverride).IsEmpty();
 }
 
+#if UE_RIGVM_UCLASS_BASED_STORAGE_DISABLED
 bool URigVMPin::CanBeBoundToVariable(const FRigVMExternalVariable& InExternalVariable, const FRigVMRegisterOffset& InOffset) const
+#else
+bool URigVMPin::CanBeBoundToVariable(const FRigVMExternalVariable& InExternalVariable, const FString& InSegmentPath) const
+#endif
 {
 	if (!InExternalVariable.IsValid(true))
 	{
@@ -769,7 +773,11 @@ bool URigVMPin::CanBeBoundToVariable(const FRigVMExternalVariable& InExternalVar
 
 	// check type validity
 	// in the future we need to allow arrays as well
+#if UE_RIGVM_UCLASS_BASED_STORAGE_DISABLED
 	if (IsArray() && InOffset.IsValid())
+#else
+	if (IsArray() && !InSegmentPath.IsEmpty())
+#endif
 	{
 		return false;
 	}
@@ -778,14 +786,49 @@ bool URigVMPin::CanBeBoundToVariable(const FRigVMExternalVariable& InExternalVar
 		return false;
 	}
 
-	FString ExternalCPPType = InExternalVariable.TypeName.ToString();
+	FName ExternalCPPType = InExternalVariable.TypeName;
 	UObject* ExternalCPPTypeObject = InExternalVariable.TypeObject;
 
+#if UE_RIGVM_UCLASS_BASED_STORAGE_DISABLED
 	if (InOffset.IsValid())
 	{
-		ExternalCPPType = InOffset.GetCPPType().ToString();
+		ExternalCPPType = InOffset.GetCPPType();
 		ExternalCPPTypeObject = InOffset.GetScriptStruct();
 	}
+#else
+	if(!InSegmentPath.IsEmpty())
+	{
+		const FProperty* Property = InExternalVariable.Property;
+		if(Property == nullptr)
+		{
+			return false;
+		}
+
+		TArray<FString> Segments;
+		SplitPinPath(InSegmentPath, Segments);
+
+		while(!Segments.IsEmpty())
+		{
+			if(const FStructProperty* StructProperty = CastField<FStructProperty>(Property))
+			{
+				const FString FirstSegment = Segments[0];
+				Segments.RemoveAt(0);
+
+				Property = StructProperty->Struct->FindPropertyByName(*FirstSegment);
+				if(Property == nullptr)
+				{
+					return false;
+				}
+			}
+			else
+			{
+				return false;
+			}
+		}
+
+		FRigVMExternalVariable::GetTypeFromProperty(Property, ExternalCPPType, ExternalCPPTypeObject);
+	}
+#endif
 
 	if (GetCPPTypeObject() != nullptr)
 	{
@@ -796,8 +839,8 @@ bool URigVMPin::CanBeBoundToVariable(const FRigVMExternalVariable& InExternalVar
 	}
 	else
 	{
-		FString CPPBaseType = IsArray() ? GetArrayElementCppType() : GetCPPType();
-		if (CPPBaseType != ExternalCPPType)
+		const FString CPPBaseType = IsArray() ? GetArrayElementCppType() : GetCPPType();
+		if (CPPBaseType != ExternalCPPType.ToString())
 		{
 			return false;
 		}
diff --git a/Engine/Source/Developer/RigVMDeveloper/Public/RigVMCompiler/RigVMCompiler.h b/Engine/Source/Developer/RigVMDeveloper/Public/RigVMCompiler/RigVMCompiler.h
--- a/Engine/Source/Developer/RigVMDeveloper/Public/RigVMCompiler/RigVMCompiler.h
+++ b/Engine/Source/Developer/RigVMDeveloper/Public/RigVMCompiler/RigVMCompiler.h
@@ -38,9 +38,6 @@ public:
 	UPROPERTY(EditAnywhere, BlueprintReadWrite, Category = FRigVMCompileSettings)
 	bool EnablePinWatches;
 
-	UPROPERTY(VisibleAnywhere, BlueprintReadOnly, Category = FRigVMCompileSettings)
-	bool ConsolidateWorkRegisters;
-
 	UPROPERTY(Transient)
 	bool IsPreprocessorPhase;
 
@@ -53,7 +50,6 @@ public:
 	static FRigVMCompileSettings Fast()
 	{
 		FRigVMCompileSettings Settings;
-		Settings.ConsolidateWorkRegisters = false;
 		Settings.EnablePinWatches = true;
 		Settings.IsPreprocessorPhase = false;
 		Settings.ASTSettings = FRigVMParserASTSettings::Fast();
@@ -63,7 +59,6 @@ public:
 	static FRigVMCompileSettings Optimized()
 	{
 		FRigVMCompileSettings Settings;
-		Settings.ConsolidateWorkRegisters = false;
 		Settings.EnablePinWatches = false;
 		Settings.IsPreprocessorPhase = false;
 		Settings.ASTSettings = FRigVMParserASTSettings::Optimized();
@@ -94,14 +89,25 @@ public:
 	TArray<const FRigVMExprAST*> ExprToSkip;
 	TMap<FString, int32> ProcessedLinks;
 	TArray<TSharedPtr<FStructOnScope>> DefaultStructs;
-	TMap<int32, int32> RegisterRefCount;
-	TArray<TPair<FName, int32>> RefCountSteps;
 	TMap<int32, FRigVMOperand> IntegerLiterals;
 	FRigVMOperand ComparisonOperand;
 
-	int32 IncRefRegister(int32 InRegister, int32 InIncrement = 1);
-	int32 DecRefRegister(int32 InRegister, int32 InDecrement = 1);
-	void RecordRefCountStep(int32  InRegister, int32 InRefCount);
+#if !UE_RIGVM_UCLASS_BASED_STORAGE_DISABLED
+	struct FPropertyPathDescription
+	{
+		FRigVMOperand OriginalOperand;
+		FString SegmentPath;
+	};
+	
+	TArray<FPropertyPathDescription> PropertyPathDescriptions;
+	TMap<ERigVMMemoryType, TArray<URigVMMemoryStorage::FPropertyDescription>> PropertyDescriptions;
+
+	FRigVMOperand AddProperty(ERigVMMemoryType InMemoryType, const FName& InName, const FString& InCPPType, UObject* InCPPTypeObject, const FString& InDefaultValue = FString());
+	FRigVMOperand FindProperty(ERigVMMemoryType InMemoryType, const FName& InName);
+	URigVMMemoryStorage::FPropertyDescription GetProperty(const FRigVMOperand& InOperand);
+	int32 FindPropertyPath(const FString& InSegmentPath) const;
+	
+#endif
 };
 
 UCLASS(BlueprintType)
diff --git a/Engine/Source/Developer/RigVMDeveloper/Public/RigVMModel/RigVMControllerActions.h b/Engine/Source/Developer/RigVMDeveloper/Public/RigVMModel/RigVMControllerActions.h
--- a/Engine/Source/Developer/RigVMDeveloper/Public/RigVMModel/RigVMControllerActions.h
+++ b/Engine/Source/Developer/RigVMDeveloper/Public/RigVMModel/RigVMControllerActions.h
@@ -71,7 +71,7 @@ private:
 	FRigVMActionWrapper& operator = (const FRigVMActionWrapper& Other) = delete;
 
 	UScriptStruct* ScriptStruct;
-	FRigVMByteArray Data;
+	TArray<uint8> Data;
 };
 
 /**
diff --git a/Engine/Source/Developer/RigVMDeveloper/Public/RigVMModel/RigVMPin.h b/Engine/Source/Developer/RigVMDeveloper/Public/RigVMModel/RigVMPin.h
--- a/Engine/Source/Developer/RigVMDeveloper/Public/RigVMModel/RigVMPin.h
+++ b/Engine/Source/Developer/RigVMDeveloper/Public/RigVMModel/RigVMPin.h
@@ -336,8 +336,13 @@ public:
 	// Returns true if this pin is bound to a variable
 	bool IsBoundToVariable(const FPinOverride& InOverride) const;
 
+#if UE_RIGVM_UCLASS_BASED_STORAGE_DISABLED
 	// Returns true if the pin can be bound to a given variable
 	bool CanBeBoundToVariable(const FRigVMExternalVariable& InExternalVariable, const FRigVMRegisterOffset& InOffset = FRigVMRegisterOffset()) const;
+#else
+	// Returns true if the pin can be bound to a given variable
+	bool CanBeBoundToVariable(const FRigVMExternalVariable& InExternalVariable, const FString& InSegmentPath = FString()) const;
+#endif
 
 	// helper function to retrieve an object from a path
 	static UObject* FindObjectFromCPPTypeObjectPath(const FString& InObjectPath);
diff --git a/Engine/Source/Programs/UnrealHeaderTool/Private/CodeGenerator.cpp b/Engine/Source/Programs/UnrealHeaderTool/Private/CodeGenerator.cpp
--- a/Engine/Source/Programs/UnrealHeaderTool/Private/CodeGenerator.cpp
+++ b/Engine/Source/Programs/UnrealHeaderTool/Private/CodeGenerator.cpp
@@ -3330,11 +3330,14 @@ void FNativeClassHeaderGenerator::ExportGeneratedStructBodyMacros(FOutputDevice&
 		//RigVMStubProlog.Add(FString::Printf(TEXT("ensure(RigVMOperandMemory.Num() == %d);"), StructRigVMInfo->Members.Num()));
 
 		int32 OperandIndex = 0;
+
 		for (int32 ParameterIndex = 0; ParameterIndex < StructRigVMInfo.Members.Num(); ParameterIndex++)
 		{
 			const FRigVMParameter& Parameter = StructRigVMInfo.Members[ParameterIndex];
+
 			if(Parameter.RequiresCast())
 			{
+#if UE_RIGVM_UCLASS_BASED_STORAGE_DISABLED
 				if (Parameter.IsArray() && !Parameter.IsConst() && !Parameter.ArraySize.IsEmpty())
 				{
 					RigVMVirtualFuncProlog.Add(FString::Printf(TEXT("%s.SetNum( %s );"), *Parameter.Name, *Parameter.ArraySize));
@@ -3351,11 +3354,16 @@ void FNativeClassHeaderGenerator::ExportGeneratedStructBodyMacros(FOutputDevice&
 				{
 					RigVMVirtualFuncProlog.Add(FString::Printf(TEXT("%s %s(%s);"), *Parameter.CastType, *Parameter.CastName, *Parameter.Name));
 				}
+#else
+				RigVMVirtualFuncProlog.Add(FString::Printf(TEXT("%s %s(%s);"), *Parameter.CastType, *Parameter.CastName, *Parameter.Name));
+#endif
 			}
 
 			const FString& ParamTypeOriginal = Parameter.TypeOriginal(true);
 			const FString& ParamNameOriginal = Parameter.NameOriginal(false);
 
+#if UE_RIGVM_UCLASS_BASED_STORAGE_DISABLED
+
 			if (ParamTypeOriginal.StartsWith(FHeaderParser::FFixedArrayText, ESearchCase::CaseSensitive))
 			{
 				FString VariableType = ParamTypeOriginal;
@@ -3412,6 +3420,23 @@ void FNativeClassHeaderGenerator::ExportGeneratedStructBodyMacros(FOutputDevice&
 
 				OperandIndex++;
 			}
+
+#else
+
+			if (Parameter.IsArray())
+			{
+				FString ExtendedType = Parameter.ExtendedType();
+				RigVMStubProlog.Add(FString::Printf(TEXT("TArray%s& %s = *(TArray%s*)RigVMMemoryHandles[%d].GetData();"),
+					*ExtendedType,
+					*ParamNameOriginal,
+					*ExtendedType,
+					OperandIndex));
+
+				OperandIndex++;
+			}
+			
+#endif
+			
 			else
 			{
 				FString VariableType = Parameter.TypeVariableRef(true);
@@ -3502,6 +3527,8 @@ void FNativeClassHeaderGenerator::ExportGeneratedStructBodyMacros(FOutputDevice&
 				RigVMMethodsDeclarations += FString::Printf(TEXT("\t}\r\n"));
 			}
 
+#if UE_RIGVM_UCLASS_BASED_STORAGE_DISABLED
+
 			for (const FRigVMParameter& StructMember : StructRigVMInfo.Members)
 			{
 				if (!StructMember.ArraySize.IsEmpty())
@@ -3510,6 +3537,8 @@ void FNativeClassHeaderGenerator::ExportGeneratedStructBodyMacros(FOutputDevice&
 					break;
 				}
 			}
+
+#endif
 		}
 
 		const FString SuperTypedef = BaseStructDef ? FString::Printf(TEXT("\ttypedef %s Super;\r\n"), *BaseStructDef->GetAlternateNameCPP()) : FString();
@@ -3799,6 +3828,9 @@ void FNativeClassHeaderGenerator::ExportGeneratedStructBodyMacros(FOutputDevice&
 
 			Out.Log(TEXT("\r\n"));
 
+					
+#if UE_RIGVM_UCLASS_BASED_STORAGE_DISABLED
+		
 		bool bHasGetArraySize = false;
 		for (const FRigVMParameter& StructMember : StructRigVMInfo.Members)
 			{
@@ -3826,6 +3858,8 @@ void FNativeClassHeaderGenerator::ExportGeneratedStructBodyMacros(FOutputDevice&
 			Out.Log(TEXT("\treturn INDEX_NONE;\r\n"));
 			Out.Log(TEXT("}\r\n\r\n"));
 		}
+
+#endif
 	}
 }
 
diff --git a/Engine/Source/Programs/UnrealHeaderTool/Private/HeaderParser.cpp b/Engine/Source/Programs/UnrealHeaderTool/Private/HeaderParser.cpp
--- a/Engine/Source/Programs/UnrealHeaderTool/Private/HeaderParser.cpp
+++ b/Engine/Source/Programs/UnrealHeaderTool/Private/HeaderParser.cpp
@@ -5642,17 +5642,35 @@ const FName FHeaderParser::NAME_InputText(TEXT("Input"));
 const FName FHeaderParser::NAME_OutputText(TEXT("Output"));
 const FName FHeaderParser::NAME_ConstantText(TEXT("Constant"));
 const FName FHeaderParser::NAME_VisibleText(TEXT("Visible"));
+
+#if UE_RIGVM_UCLASS_BASED_STORAGE_DISABLED
+
 const FName FHeaderParser::NAME_ArraySizeText(TEXT("ArraySize"));
+
+#endif
+
 const FName FHeaderParser::NAME_SingletonText(TEXT("Singleton"));
 
 const TCHAR* FHeaderParser::TArrayText = TEXT("TArray");
 const TCHAR* FHeaderParser::TEnumAsByteText = TEXT("TEnumAsByte");
+const TCHAR* FHeaderParser::GetRefText = TEXT("GetRef");
+
+#if UE_RIGVM_UCLASS_BASED_STORAGE_DISABLED
+
 const TCHAR* FHeaderParser::FFixedArrayText = TEXT("FRigVMFixedArray");
 const TCHAR* FHeaderParser::FDynamicArrayText = TEXT("FRigVMDynamicArray");
-const TCHAR* FHeaderParser::GetRefText = TEXT("GetRef");
 const TCHAR* FHeaderParser::GetFixedArrayText = TEXT("GetFixedArray");
 const TCHAR* FHeaderParser::GetDynamicArrayText = TEXT("GetDynamicArray");
 
+#else
+
+const TCHAR* FHeaderParser::FTArrayText = TEXT("TArray");
+const TCHAR* FHeaderParser::FTArrayViewText = TEXT("TArrayView");
+const TCHAR* FHeaderParser::GetArrayText = TEXT("GetArray");
+const TCHAR* FHeaderParser::GetArrayViewText = TEXT("GetArrayView");
+
+#endif
+
 void FHeaderParser::ParseRigVMMethodParameters(FUnrealStructDefinitionInfo& StructDef)
 {
 	FRigVMStructInfo& StructRigVMInfo = StructDef.GetRigVMInfo();
@@ -5679,7 +5697,9 @@ void FHeaderParser::ParseRigVMMethodParameters(FUnrealStructDefinitionInfo& Stru
 		Parameter.bConstant = PropertyDef->HasMetaData(NAME_ConstantText);
 		Parameter.bInput = PropertyDef->HasMetaData(NAME_InputText);
 		Parameter.bOutput = PropertyDef->HasMetaData(NAME_OutputText);
+#if UE_RIGVM_UCLASS_BASED_STORAGE_DISABLED
 		Parameter.ArraySize = PropertyDef->GetMetaData(NAME_ArraySizeText);
+#endif
 		Parameter.Getter = GetRefText;
 		Parameter.bEditorOnly = PropertyDef->IsEditorOnlyProperty();
 		Parameter.bSingleton = PropertyDef->HasMetaData(NAME_SingletonText);
@@ -5710,7 +5730,11 @@ void FHeaderParser::ParseRigVMMethodParameters(FUnrealStructDefinitionInfo& Stru
 		if (MemberCPPType.StartsWith(TArrayText, ESearchCase::CaseSensitive))
 		{
 			ExtendedCPPType = FString::Printf(TEXT("<%s>"), *ExtendedCPPType.LeftChop(1).RightChop(1));
+			
+#if UE_RIGVM_UCLASS_BASED_STORAGE_DISABLED
+
 			Parameter.CastName = FString::Printf(TEXT("%s_%d_Array"), *Parameter.Name, StructRigVMInfo.Members.Num());
+
 			if (Parameter.IsConst() || !Parameter.ArraySize.IsEmpty())
 			{
 				Parameter.CastType = FString::Printf(TEXT("%s%s"), FFixedArrayText, *ExtendedCPPType);
@@ -5720,7 +5744,18 @@ void FHeaderParser::ParseRigVMMethodParameters(FUnrealStructDefinitionInfo& Stru
 			{
 				Parameter.CastType = FString::Printf(TEXT("%s%s"), FDynamicArrayText, *ExtendedCPPType);
 				Parameter.Getter = GetDynamicArrayText;
+			}				
+
+#else
+
+			if(Parameter.IsConst())
+			{
+				ExtendedCPPType = FString::Printf(TEXT("<const %s>"), *ExtendedCPPType.LeftChop(1).RightChop(1));
+				Parameter.CastName = FString::Printf(TEXT("%s_%d_Array"), *Parameter.Name, StructRigVMInfo.Members.Num());
+				Parameter.CastType = FString::Printf(TEXT("%s%s"), FTArrayViewText, *ExtendedCPPType);
 			}
+
+#endif
 		}
 
 		StructRigVMInfo.Members.Add(MoveTemp(Parameter));
diff --git a/Engine/Source/Programs/UnrealHeaderTool/Private/HeaderParser.h b/Engine/Source/Programs/UnrealHeaderTool/Private/HeaderParser.h
--- a/Engine/Source/Programs/UnrealHeaderTool/Private/HeaderParser.h
+++ b/Engine/Source/Programs/UnrealHeaderTool/Private/HeaderParser.h
@@ -9,6 +9,7 @@
 #include "UnrealTypeDefinitionInfo.h"
 #include "GeneratedCodeVersion.h"
 #include "ClassMaps.h"
+#include "RigVMDefines.h"
 
 class UClass;
 enum class EGeneratedCodeVersion : uint8;
@@ -630,15 +631,24 @@ public:
 	static const FName NAME_OutputText;
 	static const FName NAME_ConstantText;
 	static const FName NAME_VisibleText;
+#if UE_RIGVM_UCLASS_BASED_STORAGE_DISABLED
 	static const FName NAME_ArraySizeText;
+#endif
 	static const FName NAME_SingletonText;
 	static const TCHAR* TArrayText;
 	static const TCHAR* TEnumAsByteText;
+	static const TCHAR* GetRefText;
+#if UE_RIGVM_UCLASS_BASED_STORAGE_DISABLED
 	static const TCHAR* FFixedArrayText;
 	static const TCHAR* FDynamicArrayText;
-	static const TCHAR* GetRefText;
 	static const TCHAR* GetFixedArrayText;
 	static const TCHAR* GetDynamicArrayText;
+#else
+	static const TCHAR* FTArrayText;
+	static const TCHAR* FTArrayViewText;
+	static const TCHAR* GetArrayText;
+	static const TCHAR* GetArrayViewText;
+#endif	
 };
 
 /////////////////////////////////////////////////////
diff --git a/Engine/Source/Programs/UnrealHeaderTool/Private/ParserHelper.h b/Engine/Source/Programs/UnrealHeaderTool/Private/ParserHelper.h
--- a/Engine/Source/Programs/UnrealHeaderTool/Private/ParserHelper.h
+++ b/Engine/Source/Programs/UnrealHeaderTool/Private/ParserHelper.h
@@ -12,6 +12,7 @@
 #include "UnrealHeaderToolGlobals.h"
 #include "Templates/UniqueObj.h"
 #include "Templates/UniquePtr.h"
+#include "RigVMDefines.h"
 
 class UEnum;
 class UScriptStruct;
@@ -482,7 +483,9 @@ struct FRigVMParameter
 		, bInput(false)
 		, bOutput(false)
 		, bSingleton(false)
+#if UE_RIGVM_UCLASS_BASED_STORAGE_DISABLED
 		, ArraySize()
+#endif
 		, Getter()
 		, CastName()
 		, CastType()
@@ -497,7 +500,9 @@ struct FRigVMParameter
 	bool bInput;
 	bool bOutput;
 	bool bSingleton;
+#if UE_RIGVM_UCLASS_BASED_STORAGE_DISABLED
 	FString ArraySize;
+#endif
 	FString Getter;
 	FString CastName;
 	FString CastType;
@@ -587,6 +592,8 @@ struct FRigVMParameter
 		return BaseType().Equals(TEXT("TArray"));
 	}
 
+#if UE_RIGVM_UCLASS_BASED_STORAGE_DISABLED
+
 	bool IsDynamic() const
 	{
 		return ArraySize.IsEmpty() && !bInput && !bOutput && !bSingleton;
@@ -597,6 +604,8 @@ struct FRigVMParameter
 		return IsArray() && IsDynamic();
 	}
 
+#endif
+
 	bool RequiresCast() const
 	{
 		return !CastType.IsEmpty() && !CastName.IsEmpty();
diff --git a/Engine/Source/Programs/UnrealHeaderTool/UnrealHeaderTool.Build.cs b/Engine/Source/Programs/UnrealHeaderTool/UnrealHeaderTool.Build.cs
--- a/Engine/Source/Programs/UnrealHeaderTool/UnrealHeaderTool.Build.cs
+++ b/Engine/Source/Programs/UnrealHeaderTool/UnrealHeaderTool.Build.cs
@@ -25,6 +25,8 @@ public class UnrealHeaderTool : ModuleRules
 				"Runtime/Launch/Private",
 				"Runtime/RHI/Public",
 				"Programs/UnrealHeaderTool/Private",
+				// For RigVM Defines
+				"Runtime/RigVM/Public",
 			});
 		
 		bEnableExceptions = true;
diff --git a/Engine/Source/Runtime/Core/Public/UObject/UE5MainStreamObjectVersion.h b/Engine/Source/Runtime/Core/Public/UObject/UE5MainStreamObjectVersion.h
--- a/Engine/Source/Runtime/Core/Public/UObject/UE5MainStreamObjectVersion.h
+++ b/Engine/Source/Runtime/Core/Public/UObject/UE5MainStreamObjectVersion.h
@@ -166,6 +166,9 @@ struct CORE_API FUE5MainStreamObjectVersion
 		// Change texture bulkdatas to have unique guids
 		VirtualizedBulkDataHaveUniqueGuids,
 
+		// Introduce RigVM Memory Class Object
+		RigVMMemoryStorageObject,
+
 		// -----<new versions can be added above this line>-------------------------------------------------
 		VersionPlusOne,
 		LatestVersion = VersionPlusOne - 1
diff --git a/Engine/Source/Runtime/RigVM/Private/RigVMCore/RigVM.cpp b/Engine/Source/Runtime/RigVM/Private/RigVMCore/RigVM.cpp
--- a/Engine/Source/Runtime/RigVM/Private/RigVMCore/RigVM.cpp
+++ b/Engine/Source/Runtime/RigVM/Private/RigVMCore/RigVM.cpp
@@ -63,10 +63,17 @@ UScriptStruct* FRigVMParameter::GetScriptStruct() const
 }
 
 URigVM::URigVM()
+#if UE_RIGVM_UCLASS_BASED_STORAGE_DISABLED
     : WorkMemoryPtr(&WorkMemoryStorage)
     , LiteralMemoryPtr(&LiteralMemoryStorage)
     , DebugMemoryPtr(&DebugMemoryStorage)
-    , ByteCodePtr(&ByteCodeStorage)
+	, WorkMemoryStorageObject(nullptr)
+#else
+	: WorkMemoryStorageObject(nullptr)
+#endif
+	, LiteralMemoryStorageObject(nullptr)
+	, DebugMemoryStorageObject(nullptr)
+	, ByteCodePtr(&ByteCodeStorage)
 #if WITH_EDITOR
 	, DebugInfo(nullptr)
 	, HaltedAtBreakpoint(nullptr)
@@ -80,9 +87,11 @@ URigVM::URigVM()
 	, ExecutingThreadId(INDEX_NONE)
 	, DeferredVMToCopy(nullptr)
 {
+#if UE_RIGVM_UCLASS_BASED_STORAGE_DISABLED
 	GetWorkMemory().SetMemoryType(ERigVMMemoryType::Work);
 	GetLiteralMemory().SetMemoryType(ERigVMMemoryType::Literal);
 	GetDebugMemory().SetMemoryType(ERigVMMemoryType::Debug);
+#endif
 }
 
 URigVM::~URigVM()
@@ -125,8 +134,21 @@ void URigVM::Save(FArchive& Ar)
 {
 	CopyDeferredVMIfRequired();
 
+	int32 RigVMUClassBasedStorageDefine = UE_RIGVM_UCLASS_BASED_STORAGE_DISABLED;
+	Ar << RigVMUClassBasedStorageDefine;
+
+#if UE_RIGVM_UCLASS_BASED_STORAGE_DISABLED
+	
 	Ar << WorkMemoryStorage;
 	Ar << LiteralMemoryStorage;
+
+#else
+
+	WorkMemoryStorageObject->Serialize(Ar);
+	LiteralMemoryStorageObject->Serialize(Ar);
+	
+#endif
+	
 	Ar << FunctionNamesStorage;
 	Ar << ByteCodeStorage;
 	Ar << Parameters;
@@ -136,21 +158,71 @@ void URigVM::Load(FArchive& Ar)
 {
 	Reset();
 
-	Ar << WorkMemoryStorage;
-	Ar << LiteralMemoryStorage;
-	Ar << FunctionNamesStorage;
-	Ar << ByteCodeStorage;
-	Ar << Parameters;
+	int32 RigVMUClassBasedStorageDefine = 1;
+	if (Ar.CustomVer(FUE5MainStreamObjectVersion::GUID) >= FUE5MainStreamObjectVersion::RigVMMemoryStorageObject)
+	{
+		Ar << RigVMUClassBasedStorageDefine;
+	}
+
+	if(RigVMUClassBasedStorageDefine == 1)
+	{
+#if !UE_RIGVM_UCLASS_BASED_STORAGE_DISABLED
+		FRigVMMemoryContainer WorkMemoryStorage;
+		FRigVMMemoryContainer LiteralMemoryStorage;
+#endif
+		
+		Ar << WorkMemoryStorage;
+		Ar << LiteralMemoryStorage;
+		Ar << FunctionNamesStorage;
+		Ar << ByteCodeStorage;
+		Ar << Parameters;
+		
+		if (Ar.CustomVer(FUE5MainStreamObjectVersion::GUID) < FUE5MainStreamObjectVersion::RigVMCopyOpStoreNumBytes)
+		{
+			Reset();
+			return;
+		}
+	}
 
-	if (Ar.CustomVer(FUE5MainStreamObjectVersion::GUID) < FUE5MainStreamObjectVersion::RigVMCopyOpStoreNumBytes)
+	if(RigVMUClassBasedStorageDefine != UE_RIGVM_UCLASS_BASED_STORAGE_DISABLED)
 	{
 		Reset();
 		return;
 	}
 
+#if !UE_RIGVM_UCLASS_BASED_STORAGE_DISABLED
+
+	if(WorkMemoryStorageObject == nullptr)
+	{
+		UClass* StorageClass = URigVMMemoryStorage::GetStorageClass(GetOutermost(), ERigVMMemoryType::Work);
+		WorkMemoryStorageObject = NewObject<URigVMMemoryStorage>(this, StorageClass, NAME_None, RF_Public | RF_Transactional);
+	}
+	if(LiteralMemoryStorageObject == nullptr)
+	{
+		UClass* StorageClass = URigVMMemoryStorage::GetStorageClass(GetOutermost(), ERigVMMemoryType::Literal);
+		LiteralMemoryStorageObject = NewObject<URigVMMemoryStorage>(this, StorageClass, NAME_None, RF_Public | RF_Transactional);
+	}
+	if(DebugMemoryStorageObject == nullptr)
+	{
+		DebugMemoryStorageObject = NewObject<URigVMMemoryStorage>(this, NAME_None, RF_Public | RF_Transactional);
+	}
+
+	WorkMemoryStorageObject->Serialize(Ar);
+	LiteralMemoryStorageObject->Serialize(Ar);
+
+	Ar << FunctionNamesStorage;
+	Ar << ByteCodeStorage;
+	Ar << Parameters;
+
+#endif
+
+#if UE_RIGVM_UCLASS_BASED_STORAGE_DISABLED
 	if (WorkMemoryStorage.bEncounteredErrorDuringLoad ||
 		LiteralMemoryStorage.bEncounteredErrorDuringLoad ||
 		!ValidateAllOperandsDuringLoad())
+#else
+	if (!ValidateAllOperandsDuringLoad())
+#endif
 	{
 		Reset();
 	}
@@ -177,8 +249,22 @@ bool URigVM::ValidateAllOperandsDuringLoad()
 	// check all operands on all ops for validity
 	bool bAllOperandsValid = true;
 
+#if UE_RIGVM_UCLASS_BASED_STORAGE_DISABLED
+
 	FRigVMMemoryContainer* LocalMemory[] = { &WorkMemoryStorage, &LiteralMemoryStorage, &DebugMemoryStorage };
+	
+#else
+
+	TArray<URigVMMemoryStorage*> LocalMemory = { GetWorkMemory(), GetLiteralMemory(), GetDebugMemory() };
+	TArray<FRigVMPropertyPath>& LocalPropertyPaths = PropertyPaths;
+	
+#endif
+	
+#if UE_RIGVM_UCLASS_BASED_STORAGE_DISABLED
 	auto CheckOperandValidity = [LocalMemory, &bAllOperandsValid](const FRigVMOperand& InOperand) -> bool
+#else
+	auto CheckOperandValidity = [LocalMemory, &bAllOperandsValid, &LocalPropertyPaths](const FRigVMOperand& InOperand) -> bool
+#endif
 	{
 		if(InOperand.GetContainerIndex() < 0 || InOperand.GetContainerIndex() >= (int32)ERigVMMemoryType::Invalid)
 		{
@@ -187,11 +273,19 @@ bool URigVM::ValidateAllOperandsDuringLoad()
 		}
 
 
+#if UE_RIGVM_UCLASS_BASED_STORAGE_DISABLED
 		const FRigVMMemoryContainer* MemoryForOperand = LocalMemory[InOperand.GetContainerIndex()];
+#else
+		const URigVMMemoryStorage* MemoryForOperand = LocalMemory[InOperand.GetContainerIndex()];
+#endif
 
 		if(InOperand.GetMemoryType() != ERigVMMemoryType::External)
 		{
+#if UE_RIGVM_UCLASS_BASED_STORAGE_DISABLED
 			if(!MemoryForOperand->Registers.IsValidIndex(InOperand.GetRegisterIndex()))
+#else
+			if(!MemoryForOperand->IsValidIndex(InOperand.GetRegisterIndex()))
+#endif
 			{
 				bAllOperandsValid = false;
 				return false;
@@ -200,7 +294,11 @@ bool URigVM::ValidateAllOperandsDuringLoad()
 
 		if(InOperand.GetRegisterOffset() != INDEX_NONE)
 		{
+#if UE_RIGVM_UCLASS_BASED_STORAGE_DISABLED
 			if(!MemoryForOperand->RegisterOffsets.IsValidIndex(InOperand.GetRegisterOffset()))
+#else
+			if(!LocalPropertyPaths.IsValidIndex(InOperand.GetRegisterOffset()))
+#endif
 			{
 				bAllOperandsValid = false;
 				return false;
@@ -347,9 +445,11 @@ bool URigVM::ValidateAllOperandsDuringLoad()
 
 void URigVM::Reset()
 {
+#if UE_RIGVM_UCLASS_BASED_STORAGE_DISABLED
 	WorkMemoryStorage.Reset();
 	LiteralMemoryStorage.Reset();
 	DebugMemoryStorage.Reset();
+#endif
 	FunctionNamesStorage.Reset();
 	FunctionsStorage.Reset();
 	ByteCodeStorage.Reset();
@@ -358,13 +458,19 @@ void URigVM::Reset()
 	ParametersNameMap.Reset();
 	DeferredVMToCopy = nullptr;
 
+#if UE_RIGVM_UCLASS_BASED_STORAGE_DISABLED
 	WorkMemoryPtr = &WorkMemoryStorage;
 	LiteralMemoryPtr = &LiteralMemoryStorage;
 	DebugMemoryPtr = &DebugMemoryStorage;
+#endif
 	FunctionNamesPtr = &FunctionNamesStorage;
 	FunctionsPtr = &FunctionsStorage;
 	ByteCodePtr = &ByteCodeStorage;
 
+#if !UE_RIGVM_UCLASS_BASED_STORAGE_DISABLED
+	PropertyPaths.Reset();
+#endif
+
 	InvalidateCachedMemory();
 	
 	OperandToDebugRegisters.Reset();
@@ -372,9 +478,11 @@ void URigVM::Reset()
 
 void URigVM::Empty()
 {
+#if UE_RIGVM_UCLASS_BASED_STORAGE_DISABLED
 	WorkMemoryStorage.Empty();
 	LiteralMemoryStorage.Empty();
 	DebugMemoryStorage.Empty();
+#endif
 	FunctionNamesStorage.Empty();
 	FunctionsStorage.Empty();
 	ByteCodeStorage.Empty();
@@ -384,6 +492,10 @@ void URigVM::Empty()
 	DeferredVMToCopy = nullptr;
 	ExternalVariables.Empty();
 
+#if !UE_RIGVM_UCLASS_BASED_STORAGE_DISABLED
+	PropertyPaths.Empty();
+#endif
+
 	InvalidateCachedMemory();
 
 	CachedMemory.Empty();
@@ -407,6 +519,8 @@ void URigVM::CopyFrom(URigVM* InVM, bool bDeferCopy, bool bReferenceLiteralMemor
 	
 	Reset();
 
+#if UE_RIGVM_UCLASS_BASED_STORAGE_DISABLED
+	
 	if(InVM->WorkMemoryPtr == &InVM->WorkMemoryStorage)
 	{
 		WorkMemoryStorage = InVM->WorkMemoryStorage;
@@ -440,6 +554,46 @@ void URigVM::CopyFrom(URigVM* InVM, bool bDeferCopy, bool bReferenceLiteralMemor
 	{
 		DebugMemoryPtr = InVM->DebugMemoryPtr;
 	}
+	
+#else
+
+	auto CopyMemoryStorage = [](URigVMMemoryStorage*& TargetMemory, URigVMMemoryStorage* SourceMemory, UObject* Outer)
+	{
+		if(SourceMemory != nullptr)
+		{
+			if(TargetMemory == nullptr)
+			{
+				TargetMemory = NewObject<URigVMMemoryStorage>(Outer, SourceMemory->GetClass());
+			}
+
+			check(TargetMemory->GetClass() == SourceMemory->GetClass());
+
+			for(int32 PropertyIndex = 0; PropertyIndex < TargetMemory->Num(); PropertyIndex++)
+			{
+				URigVMMemoryStorage::CopyProperty(
+					TargetMemory,
+					PropertyIndex,
+					FRigVMPropertyPath::Empty,
+					SourceMemory,
+					PropertyIndex,
+					FRigVMPropertyPath::Empty
+				);
+			}
+		}
+		else if(TargetMemory != nullptr)
+		{
+			TargetMemory->Rename(nullptr, GetTransientPackage(), REN_ForceNoResetLoaders | REN_DoNotDirty | REN_DontCreateRedirectors | REN_NonTransactional);
+			TargetMemory->MarkPendingKill();
+			TargetMemory = nullptr;
+		}
+	};
+
+	CopyMemoryStorage(WorkMemoryStorageObject, InVM->WorkMemoryStorageObject, this);
+	CopyMemoryStorage(LiteralMemoryStorageObject, InVM->LiteralMemoryStorageObject, this);
+	CopyMemoryStorage(DebugMemoryStorageObject, InVM->DebugMemoryStorageObject, this);
+	PropertyPaths = InVM->PropertyPaths;
+	
+#endif
 
 	if(InVM->FunctionNamesPtr == &InVM->FunctionNamesStorage && !bReferenceByteCode)
 	{
@@ -508,6 +662,34 @@ FString URigVM::GetRigVMFunctionName(int32 InFunctionIndex) const
 	return GetFunctionNames()[InFunctionIndex].ToString();
 }
 
+#if !UE_RIGVM_UCLASS_BASED_STORAGE_DISABLED
+
+URigVMMemoryStorage* URigVM::GetMemoryByType(ERigVMMemoryType InMemoryType) const
+{
+	switch(InMemoryType)
+	{
+		case ERigVMMemoryType::Literal:
+		{
+			return GetLiteralMemory();
+		}
+		case ERigVMMemoryType::Work:
+		{
+			return GetWorkMemory();
+		}
+		case ERigVMMemoryType::Debug:
+		{
+			return GetDebugMemory();
+		}
+		default:
+		{
+			break;
+		}
+	}
+	return nullptr;
+}
+
+#endif
+
 const FRigVMInstructionArray& URigVM::GetInstructions()
 {
 	RefreshInstructionsIfRequired();
@@ -534,6 +716,7 @@ TArray<FName> URigVM::GetEntryNames() const
 }
 
 #if WITH_EDITOR
+
 bool URigVM::ResumeExecution()
 {
 	HaltedAtBreakpoint = nullptr;
@@ -551,7 +734,11 @@ bool URigVM::ResumeExecution()
 	return false;
 }
 
+#if UE_RIGVM_UCLASS_BASED_STORAGE_DISABLED
 bool URigVM::ResumeExecution(FRigVMMemoryContainerPtrArray Memory, FRigVMFixedArray<void*> AdditionalArguments, const FName& InEntryName)
+#else
+bool URigVM::ResumeExecution(TArrayView<URigVMMemoryStorage*> Memory, TArrayView<void*> AdditionalArguments, const FName& InEntryName)
+#endif
 {
 	ResumeExecution();
 	return Execute(Memory, AdditionalArguments, InEntryName);
@@ -635,7 +822,11 @@ void URigVM::CopyDeferredVMIfRequired()
 	}
 }
 
+#if UE_RIGVM_UCLASS_BASED_STORAGE_DISABLED
 void URigVM::CacheMemoryHandlesIfRequired(FRigVMMemoryContainerPtrArray InMemory)
+#else
+void URigVM::CacheMemoryHandlesIfRequired(TArrayView<URigVMMemoryStorage*> InMemory)
+#endif
 {
 	ensureMsgf(ExecutingThreadId == FPlatformTLS::GetCurrentThreadId(), TEXT("RigVM::CacheMemoryHandlesIfRequired from multiple threads (%d and %d)"), ExecutingThreadId, (int32)FPlatformTLS::GetCurrentThreadId());
 
@@ -780,10 +971,12 @@ void URigVM::CacheMemoryHandlesIfRequired(FRigVMMemoryContainerPtrArray InMemory
 				CacheSingleMemoryHandle(Op.Source);
 				CacheSingleMemoryHandle(Op.Target);
 
+#if UE_RIGVM_UCLASS_BASED_STORAGE_DISABLED
 				if (UScriptStruct* ScriptStruct = GetScriptStructForCopyOp(Op))
 				{
 					CachedMemoryHandles.Add(FRigVMMemoryHandle((uint8*)ScriptStruct));
 				}
+#endif
 
 				InstructionIndex++;
 				break;
@@ -883,8 +1076,10 @@ void URigVM::RebuildByteCodeOnLoad()
 				NewCopyOp = GetCopyOpForOperands(OldCopyOp.Source, OldCopyOp.Target);
 				check(OldCopyOp.Source == NewCopyOp.Source);
 				check(OldCopyOp.Target == NewCopyOp.Target);
+#if UE_RIGVM_UCLASS_BASED_STORAGE_DISABLED
 				//check(OldCopyOp.NumBytes == NewCopyOp.NumBytes);
 				check(OldCopyOp.RegisterType == NewCopyOp.RegisterType);
+#endif
 				break;
 			}
 			default:
@@ -1063,7 +1258,11 @@ bool URigVM::ShouldHaltAtInstruction(const uint16 InstructionIndex)
 }
 #endif
 
+#if UE_RIGVM_UCLASS_BASED_STORAGE_DISABLED
 bool URigVM::Initialize(FRigVMMemoryContainerPtrArray Memory, FRigVMFixedArray<void*> AdditionalArguments)
+#else
+bool URigVM::Initialize(TArrayView<URigVMMemoryStorage*> Memory, TArrayView<void*> AdditionalArguments)
+#endif
 {
 	if (ExecutingThreadId != INDEX_NONE)
 	{
@@ -1080,11 +1279,20 @@ bool URigVM::Initialize(FRigVMMemoryContainerPtrArray Memory, FRigVMFixedArray<v
 		return true;
 	}
 
+#if UE_RIGVM_UCLASS_BASED_STORAGE_DISABLED
 	FRigVMMemoryContainer* LocalMemory[] = { WorkMemoryPtr, LiteralMemoryPtr, DebugMemoryPtr };
 	if (Memory.Num() == 0)
 	{
 		Memory = FRigVMMemoryContainerPtrArray(LocalMemory, 3);
 	}
+#else
+	TArray<URigVMMemoryStorage*> LocalMemory;
+	if (Memory.Num() == 0)
+	{
+		LocalMemory = GetLocalMemoryArray();
+		Memory = LocalMemory;
+	}
+#endif
 
 	CacheMemoryHandlesIfRequired(Memory);
 	FRigVMByteCode& ByteCode = GetByteCode();
@@ -1180,6 +1388,8 @@ bool URigVM::Initialize(FRigVMMemoryContainerPtrArray Memory, FRigVMFixedArray<v
 
 				// find out the largest slice count
 				int32 MaxSliceCount = 1;
+
+#if UE_RIGVM_UCLASS_BASED_STORAGE_DISABLED
 				for (const FRigVMMemoryHandle& OpHandle : OpHandles)
 				{
 					if (OpHandle.Type == FRigVMMemoryHandle::Dynamic)
@@ -1197,6 +1407,9 @@ bool URigVM::Initialize(FRigVMMemoryContainerPtrArray Memory, FRigVMFixedArray<v
 						}
 					}
 				}
+#else
+				// todo Deal with slice counts
+#endif
 
 				Context.BeginSlice(MaxSliceCount);
 				for (int32 SliceIndex = 0; SliceIndex < MaxSliceCount; SliceIndex++)
@@ -1220,6 +1433,8 @@ bool URigVM::Initialize(FRigVMMemoryContainerPtrArray Memory, FRigVMFixedArray<v
 
 				FRigVMMemoryHandle& SourceHandle = CachedMemoryHandles[FirstHandleForInstruction[Context.InstructionIndex]];
 				FRigVMMemoryHandle& TargetHandle = CachedMemoryHandles[FirstHandleForInstruction[Context.InstructionIndex] + 1];
+
+#if UE_RIGVM_UCLASS_BASED_STORAGE_DISABLED
 				void* SourcePtr = SourceHandle;
 				void* TargetPtr = TargetHandle;
 
@@ -1325,6 +1540,11 @@ bool URigVM::Initialize(FRigVMMemoryContainerPtrArray Memory, FRigVMFixedArray<v
 						break;
 					}
 				}
+#else
+
+				URigVMMemoryStorage::CopyProperty(TargetHandle, SourceHandle);
+
+#endif
 				break;
 			}
 			case ERigVMOpCode::Increment:
@@ -1356,7 +1576,11 @@ bool URigVM::Initialize(FRigVMMemoryContainerPtrArray Memory, FRigVMFixedArray<v
 	return true;
 }
 
+#if UE_RIGVM_UCLASS_BASED_STORAGE_DISABLED
 bool URigVM::Execute(FRigVMMemoryContainerPtrArray Memory, FRigVMFixedArray<void*> AdditionalArguments, const FName& InEntryName)
+#else
+bool URigVM::Execute(TArrayView<URigVMMemoryStorage*> Memory, TArrayView<void*> AdditionalArguments, const FName& InEntryName)
+#endif
 {
 	if (ExecutingThreadId != INDEX_NONE)
 	{
@@ -1374,12 +1598,21 @@ bool URigVM::Execute(FRigVMMemoryContainerPtrArray Memory, FRigVMFixedArray<void
 	}
 
 	// changes to the layout of memory array should be reflected in GetContainerIndex()
-	FRigVMMemoryContainer* LocalMemory[] = { WorkMemoryPtr, LiteralMemoryPtr, DebugMemoryPtr }; 
+#if UE_RIGVM_UCLASS_BASED_STORAGE_DISABLED
+	FRigVMMemoryContainer* LocalMemory[] = { WorkMemoryPtr, LiteralMemoryPtr, DebugMemoryPtr };
 	if (Memory.Num() == 0)
 	{
 		Memory = FRigVMMemoryContainerPtrArray(LocalMemory, 3);
 	}
-
+#else
+	TArray<URigVMMemoryStorage*> LocalMemory;
+	if (Memory.Num() == 0)
+	{
+		LocalMemory = GetLocalMemoryArray();
+		Memory = LocalMemory;
+	}
+#endif
+	
 	CacheMemoryHandlesIfRequired(Memory);
 	FRigVMByteCode& ByteCode = GetByteCode();
 	TArray<FRigVMFunctionPtr>& Functions = GetFunctions();
@@ -1510,6 +1743,8 @@ bool URigVM::Execute(FRigVMMemoryContainerPtrArray Memory, FRigVMFixedArray<void
 				(*Functions[Op.FunctionIndex])(Context, Handles);
 
 #if WITH_EDITOR
+
+#if UE_RIGVM_UCLASS_BASED_STORAGE_DISABLED
 				if(DebugMemoryPtr->Num() > 0)
 				{
 					const FRigVMOperandArray Operands = ByteCode.GetOperandsForExecuteOp(Instruction);
@@ -1523,6 +1758,17 @@ bool URigVM::Execute(FRigVMMemoryContainerPtrArray Memory, FRigVMFixedArray<void
 						CopyOperandForDebuggingIfNeeded(Operands[OperandIndex++], Handles[HandleIndex]);
 					}
 				}
+#else
+				if(DebugMemoryStorageObject->Num() > 0)
+				{
+					const FRigVMOperandArray Operands = ByteCode.GetOperandsForExecuteOp(Instruction);
+					for(int32 OperandIndex = 0, HandleIndex = 0; OperandIndex < Operands.Num() && HandleIndex < Handles.Num(); HandleIndex++)
+					{
+						CopyOperandForDebuggingIfNeeded(Operands[OperandIndex++], Handles[HandleIndex]);
+					}
+				}
+#endif
+
 #endif
 
 				Context.InstructionIndex++;
@@ -1532,7 +1778,11 @@ bool URigVM::Execute(FRigVMMemoryContainerPtrArray Memory, FRigVMFixedArray<void
 			{
 				*((int32*)CachedMemoryHandles[FirstHandleForInstruction[Context.InstructionIndex]].GetData()) = 0;
 #if WITH_EDITOR
+#if UE_RIGVM_UCLASS_BASED_STORAGE_DISABLED
 				if(DebugMemoryPtr->Num() > 0)
+#else
+				if(DebugMemoryStorageObject->Num() > 0)
+#endif
 				{
 					const FRigVMUnaryOp& Op = ByteCode.GetOpAt<FRigVMUnaryOp>(Instruction);
 					CopyOperandForDebuggingIfNeeded(Op.Arg, CachedMemoryHandles[FirstHandleForInstruction[Context.InstructionIndex]]);
@@ -1560,6 +1810,8 @@ bool URigVM::Execute(FRigVMMemoryContainerPtrArray Memory, FRigVMFixedArray<void
 
 				FRigVMMemoryHandle& SourceHandle = CachedMemoryHandles[FirstHandleForInstruction[Context.InstructionIndex]];
 				FRigVMMemoryHandle& TargetHandle = CachedMemoryHandles[FirstHandleForInstruction[Context.InstructionIndex] + 1];
+
+#if UE_RIGVM_UCLASS_BASED_STORAGE_DISABLED
 				void* SourcePtr = SourceHandle;
 				void* TargetPtr = TargetHandle;
 
@@ -1665,14 +1917,24 @@ bool URigVM::Execute(FRigVMMemoryContainerPtrArray Memory, FRigVMFixedArray<void
 						break;
 					}
 				}
-
+					
+#else
+					
+				URigVMMemoryStorage::CopyProperty(TargetHandle, SourceHandle);
+					
+#endif
+					
 #if WITH_EDITOR
+#if UE_RIGVM_UCLASS_BASED_STORAGE_DISABLED
 				if(DebugMemoryPtr->Num() > 0)
+#else
+				if(DebugMemoryStorageObject->Num() > 0)
+#endif
 				{
 					CopyOperandForDebuggingIfNeeded(Op.Source, SourceHandle);
 				}
 #endif
-
+					
 				Context.InstructionIndex++;
 				break;
 			}
@@ -1680,7 +1942,11 @@ bool URigVM::Execute(FRigVMMemoryContainerPtrArray Memory, FRigVMFixedArray<void
 			{
 				(*((int32*)CachedMemoryHandles[FirstHandleForInstruction[Context.InstructionIndex]].GetData()))++;
 #if WITH_EDITOR
+#if UE_RIGVM_UCLASS_BASED_STORAGE_DISABLED
 				if(DebugMemoryPtr->Num() > 0)
+#else
+				if(DebugMemoryStorageObject->Num() > 0)
+#endif
 				{
 					const FRigVMUnaryOp& Op = ByteCode.GetOpAt<FRigVMUnaryOp>(Instruction);
 					CopyOperandForDebuggingIfNeeded(Op.Arg, CachedMemoryHandles[FirstHandleForInstruction[Context.InstructionIndex]]);
@@ -1693,7 +1959,11 @@ bool URigVM::Execute(FRigVMMemoryContainerPtrArray Memory, FRigVMFixedArray<void
 			{
 				(*((int32*)CachedMemoryHandles[FirstHandleForInstruction[Context.InstructionIndex]].GetData()))--;
 #if WITH_EDITOR
+#if UE_RIGVM_UCLASS_BASED_STORAGE_DISABLED
 				if(DebugMemoryPtr->Num() > 0)
+#else
+				if(DebugMemoryStorageObject->Num() > 0)
+#endif
 				{
 					const FRigVMUnaryOp& Op = ByteCode.GetOpAt<FRigVMUnaryOp>(Instruction);
 					CopyOperandForDebuggingIfNeeded(Op.Arg, CachedMemoryHandles[FirstHandleForInstruction[Context.InstructionIndex]]);
@@ -1706,6 +1976,9 @@ bool URigVM::Execute(FRigVMMemoryContainerPtrArray Memory, FRigVMFixedArray<void
 			case ERigVMOpCode::NotEquals:
 			{
 				const FRigVMComparisonOp& Op = ByteCode.GetOpAt<FRigVMComparisonOp>(Instruction);
+
+#if UE_RIGVM_UCLASS_BASED_STORAGE_DISABLED
+
 				const FRigVMRegister& RegisterA = (*Memory[Op.A.GetContainerIndex()])[Op.A.GetRegisterIndex()];
 				const FRigVMRegister& RegisterB = (*Memory[Op.B.GetContainerIndex()])[Op.B.GetRegisterIndex()];
 				const uint16 BytesA = RegisterA.GetNumBytesPerSlice();
@@ -1772,6 +2045,14 @@ bool URigVM::Execute(FRigVMMemoryContainerPtrArray Memory, FRigVMFixedArray<void
 					Result = !Result;
 				}
 
+#else
+
+				FRigVMMemoryHandle& HandleA = CachedMemoryHandles[FirstHandleForInstruction[Context.InstructionIndex]];
+				FRigVMMemoryHandle& HandleB = CachedMemoryHandles[FirstHandleForInstruction[Context.InstructionIndex] + 1];
+				const bool Result = HandleA.GetProperty()->Identical(HandleA.GetData(true), HandleB.GetData(true));
+					
+#endif				
+
 				*((bool*)CachedMemoryHandles[FirstHandleForInstruction[Context.InstructionIndex]+2].GetData()) = Result;
 				Context.InstructionIndex++;
 				break;
@@ -1890,7 +2171,11 @@ bool URigVM::Execute(FRigVMMemoryContainerPtrArray Memory, FRigVMFixedArray<void
 
 bool URigVM::Execute(const FName& InEntryName)
 {
+#if UE_RIGVM_UCLASS_BASED_STORAGE_DISABLED
 	return Execute(FRigVMMemoryContainerPtrArray(), FRigVMFixedArray<void*>(), InEntryName);
+#else
+	return Execute(TArray<URigVMMemoryStorage*>(), TArrayView<void*>(), InEntryName);
+#endif
 }
 
 FRigVMExternalVariable URigVM::GetExternalVariableByName(const FName& InExternalVariableName)
@@ -1905,6 +2190,8 @@ FRigVMExternalVariable URigVM::GetExternalVariableByName(const FName& InExternal
 	return FRigVMExternalVariable();
 }
 
+#if UE_RIGVM_UCLASS_BASED_STORAGE_DISABLED
+
 void URigVM::SetRegisterValueFromString(const FRigVMOperand& InOperand, const FString& InCPPType, const UObject* InCPPTypeObject, const TArray<FString>& InDefaultValues)
 {
 	if (InOperand.GetMemoryType() == ERigVMMemoryType::Literal)
@@ -1921,6 +2208,20 @@ void URigVM::SetRegisterValueFromString(const FRigVMOperand& InOperand, const FS
 	}
 }
 
+#else
+
+void URigVM::SetPropertyValueFromString(const FRigVMOperand& InOperand, const FString& InDefaultValue)
+{
+	URigVMMemoryStorage* Memory = GetMemoryByType(InOperand.GetMemoryType());
+	if(Memory == nullptr)
+	{
+		return;
+	}
+
+	Memory->SetDataFromString(InOperand.GetRegisterIndex(), InDefaultValue);
+}
+
+#endif
 
 #if WITH_EDITOR
 
@@ -2185,6 +2486,8 @@ FString URigVM::DumpByteCodeAsText(const TArray<int32>& InInstructionOrder, bool
 
 FString URigVM::GetOperandLabel(const FRigVMOperand& InOperand, TFunction<FString(const FString& RegisterName, const FString& RegisterOffsetName)> FormatFunction) const
 {
+#if UE_RIGVM_UCLASS_BASED_STORAGE_DISABLED
+	
 	const FRigVMMemoryContainer* MemoryPtr = nullptr;
 
 	if (InOperand.GetMemoryType() == ERigVMMemoryType::Literal)
@@ -2225,6 +2528,26 @@ FString URigVM::GetOperandLabel(const FRigVMOperand& InOperand, TFunction<FStrin
 		OperandLabel = FString::Printf(TEXT("%s.%s"), *OperandLabel, *RegisterOffsetName);
 	}
 
+#else
+
+	URigVMMemoryStorage* Memory = GetMemoryByType(InOperand.GetMemoryType());
+	if(Memory == nullptr)
+	{
+		return FString();
+	}
+
+	check(Memory->IsValidIndex(InOperand.GetRegisterIndex()));
+	
+	const FString RegisterName = Memory->GetProperties()[InOperand.GetRegisterIndex()]->GetName();
+	const FString RegisterOffsetName =
+		InOperand.GetRegisterOffset() != INDEX_NONE ?
+		PropertyPaths[InOperand.GetRegisterOffset()].ToString() :
+		FString();
+	
+	FString OperandLabel = RegisterName;
+	
+#endif
+	
 	// caller can provide an alternative format to override the default format(optional)
 	if (FormatFunction)
 	{
@@ -2239,6 +2562,7 @@ FString URigVM::GetOperandLabel(const FRigVMOperand& InOperand, TFunction<FStrin
 void URigVM::ClearDebugMemory()
 {
 #if WITH_EDITOR
+#if UE_RIGVM_UCLASS_BASED_STORAGE_DISABLED
 	FRigVMMemoryContainer* DebugMemory = CachedMemory[GetContainerIndex(ERigVMMemoryType::Debug)];
 	if (DebugMemory)
 	{ 
@@ -2247,12 +2571,26 @@ void URigVM::ClearDebugMemory()
 			ensure(DebugMemory->GetRegister(RegisterIndex).IsDynamic());
 			DebugMemory->Destroy(RegisterIndex);
 		}	
-	} 
+	}
+#else
+
+	for(int32 PropertyIndex = 0; PropertyIndex < GetDebugMemory()->Num(); PropertyIndex++)
+	{
+		if(const FArrayProperty* ArrayProperty = CastField<FArrayProperty>(GetDebugMemory()->GetProperties()[PropertyIndex]))
+		{
+			FScriptArrayHelper ArrayHelper(ArrayProperty, GetDebugMemory()->GetData<uint8>(PropertyIndex));
+			ArrayHelper.EmptyValues();
+		}
+	}
+	
+#endif
 #endif
 }
 
 void URigVM::CacheSingleMemoryHandle(const FRigVMOperand& InArg, bool bForExecute)
 {
+#if UE_RIGVM_UCLASS_BASED_STORAGE_DISABLED
+	
 	if (InArg.GetMemoryType() == ERigVMMemoryType::External)
 	{
 		ensure(ExternalVariables.IsValidIndex(InArg.GetRegisterIndex()));
@@ -2282,11 +2620,43 @@ void URigVM::CacheSingleMemoryHandle(const FRigVMOperand& InArg, bool bForExecut
 			CachedMemoryHandles.Add(FRigVMMemoryHandle((uint8*)ElementsForArray, sizeof(uint16), FRigVMMemoryHandle::FType::ArraySize));
 		}
 	}
+	
+#else
+
+	const FRigVMPropertyPath* PropertyPath = nullptr;
+	if(InArg.GetRegisterOffset())
+	{
+		check(PropertyPaths.IsValidIndex(InArg.GetRegisterOffset()));
+		PropertyPath = &PropertyPaths[InArg.GetRegisterOffset()];
+	}
+
+	if (InArg.GetMemoryType() == ERigVMMemoryType::External)
+	{
+		check(ExternalVariables.IsValidIndex(InArg.GetRegisterIndex()));
+
+		FRigVMExternalVariable& ExternalVariable = ExternalVariables[InArg.GetRegisterIndex()];
+		check(ExternalVariable.IsValid(false));
+
+		const FRigVMMemoryHandle Handle(ExternalVariable.Memory, ExternalVariable.Property, PropertyPath);
+		CachedMemoryHandles.Add(Handle);
+		return;
+	}
+
+	URigVMMemoryStorage* Memory = GetMemoryByType(InArg.GetMemoryType());
+	uint8* Data = Memory->GetData<uint8>(InArg.GetRegisterIndex());
+	const FProperty* Property = Memory->GetProperties()[InArg.GetRegisterIndex()];
+	const FRigVMMemoryHandle Handle(Data, Property, PropertyPath);
+	CachedMemoryHandles.Add(Handle);
+
+#endif
 }
 
 void URigVM::CopyOperandForDebuggingImpl(const FRigVMOperand& InArg, const FRigVMMemoryHandle& InHandle, const FRigVMOperand& InDebugOperand)
 {
 #if WITH_EDITOR
+
+#if UE_RIGVM_UCLASS_BASED_STORAGE_DISABLED
+
 	check(InArg.IsValid());
 	check(InArg.GetRegisterOffset() == INDEX_NONE);
 	check(InDebugOperand.IsValid());
@@ -2395,10 +2765,46 @@ void URigVM::CopyOperandForDebuggingImpl(const FRigVMOperand& InArg, const FRigV
 		}
 	}
 
+#else
+
+	URigVMMemoryStorage* TargetMemory = GetDebugMemory();
+	if(TargetMemory == nullptr)
+	{
+		return;
+	}
+	const FProperty* TargetProperty = TargetMemory->GetProperties()[InDebugOperand.GetRegisterIndex()];
+	uint8* TargetPtr = TargetMemory->GetData<uint8>(InDebugOperand.GetRegisterIndex());
+
+	// since debug properties are always arrays, we need to divert to the last array element's memory
+	const FArrayProperty* TargetArrayProperty = CastField<FArrayProperty>(TargetProperty);
+	if(TargetArrayProperty == nullptr)
+	{
+		return;
+	}
+
+	// add an element to the end for debug watching
+	FScriptArrayHelper ArrayHelper(TargetArrayProperty, TargetPtr);
+	const int32 AddedIndex = ArrayHelper.AddValue();
+	TargetPtr = ArrayHelper.GetRawPtr(AddedIndex);
+	TargetProperty = TargetArrayProperty->Inner;
+
+	URigVMMemoryStorage* SourceMemory = GetMemoryByType(InArg.GetMemoryType());
+	if(SourceMemory == nullptr)
+	{
+		return;
+	}
+	const FProperty* SourceProperty = SourceMemory->GetProperties()[InArg.GetRegisterIndex()];
+	const uint8* SourcePtr = SourceMemory->GetData<uint8>(InArg.GetRegisterIndex());
+
+	URigVMMemoryStorage::CopyProperty(TargetProperty, TargetPtr, SourceProperty, SourcePtr);
+	
+#endif
 	
 #endif
 }
 
+#if UE_RIGVM_UCLASS_BASED_STORAGE_DISABLED
+
 FRigVMCopyOp URigVM::GetCopyOpForOperands(const FRigVMOperand& InSource, const FRigVMOperand& InTarget)
 {
 	FCopyInfoForOperand SourceCopyInfo = GetCopyInfoForOperand(InSource);
@@ -2428,6 +2834,17 @@ FRigVMCopyOp URigVM::GetCopyOpForOperands(const FRigVMOperand& InSource, const F
 	return FRigVMCopyOp(InSource, InTarget, TargetCopyInfo.NumBytesToCopy, TargetCopyInfo.RegisterType, CopyType);
 }
 
+#else
+
+FRigVMCopyOp URigVM::GetCopyOpForOperands(const FRigVMOperand& InSource, const FRigVMOperand& InTarget)
+{
+	return FRigVMCopyOp(InSource, InTarget);
+}
+
+#endif
+
+#if UE_RIGVM_UCLASS_BASED_STORAGE_DISABLED
+
 URigVM::FCopyInfoForOperand URigVM::GetCopyInfoForOperand(const FRigVMOperand& InOperand)
 {
 	if(CachedMemory.IsEmpty())
@@ -2512,3 +2929,4 @@ UScriptStruct* URigVM::GetScripStructForOperand(const FRigVMOperand& InOperand)
 	return CachedMemory[InOperand.GetContainerIndex()]->GetScriptStruct(InOperand.GetRegisterIndex());
 }
 
+#endif
diff --git a/Engine/Source/Runtime/RigVM/Private/RigVMCore/RigVMByteCode.cpp b/Engine/Source/Runtime/RigVM/Private/RigVMCore/RigVMByteCode.cpp
--- a/Engine/Source/Runtime/RigVM/Private/RigVMCore/RigVMByteCode.cpp
+++ b/Engine/Source/Runtime/RigVM/Private/RigVMCore/RigVMByteCode.cpp
@@ -32,6 +32,8 @@ void FRigVMCopyOp::Serialize(FArchive& Ar)
 	Ar << Source;
 	Ar << Target;
 
+#if UE_RIGVM_UCLASS_BASED_STORAGE_DISABLED
+
 	if(Ar.IsLoading())
 	{
 		if (Ar.CustomVer(FUE5MainStreamObjectVersion::GUID) < FUE5MainStreamObjectVersion::RigVMCopyOpStoreNumBytes)
@@ -50,6 +52,8 @@ void FRigVMCopyOp::Serialize(FArchive& Ar)
 		Ar << NumBytes;
 		Ar << RegisterType;
 	}
+	
+#endif
 }
 
 void FRigVMComparisonOp::Serialize(FArchive& Ar)
@@ -733,7 +737,12 @@ uint64 FRigVMByteCode::GetOpNumBytesAt(uint64 InByteCodeIndex, bool bIncludeOper
 		}
 		case ERigVMOpCode::ChangeType:
 		{
+#if UE_RIGVM_UCLASS_BASED_STORAGE_DISABLED
 			return (uint64)sizeof(FRigVMChangeTypeOp);
+#else
+			checkNoEntry();
+			return 0;
+#endif
 		}
 		case ERigVMOpCode::Exit:
 		{
@@ -774,6 +783,8 @@ uint64 FRigVMByteCode::AddTrueOp(const FRigVMOperand& InArg)
 	return AddOp(Op);
 }
 
+#if UE_RIGVM_UCLASS_BASED_STORAGE_DISABLED
+
 uint64 FRigVMByteCode::AddCopyOp(const FRigVMOperand& InSource, const FRigVMOperand& InTarget, uint16 InNumBytes, ERigVMRegisterType InTargetType, ERigVMCopyType InCopyType)
 {
 	check(InTarget.GetMemoryType() != ERigVMMemoryType::Literal);
@@ -784,9 +795,25 @@ uint64 FRigVMByteCode::AddCopyOp(const FRigVMOperand& InSource, const FRigVMOper
 	return AddOp(Op);
 }
 
+#else
+
+uint64 FRigVMByteCode::AddCopyOp(const FRigVMOperand& InSource, const FRigVMOperand& InTarget)
+{
+	check(InTarget.GetMemoryType() != ERigVMMemoryType::Literal);
+
+	FRigVMCopyOp Op(InSource, InTarget);
+	return AddOp(Op);
+}
+
+#endif
+
 uint64 FRigVMByteCode::AddCopyOp(const FRigVMCopyOp& InCopyOp)
 {
+#if UE_RIGVM_UCLASS_BASED_STORAGE_DISABLED
 	return AddCopyOp(InCopyOp.Source, InCopyOp.Target, InCopyOp.NumBytes, InCopyOp.RegisterType, InCopyOp.CopyType);
+#else
+	return AddCopyOp(InCopyOp.Source, InCopyOp.Target);
+#endif
 }
 
 uint64 FRigVMByteCode::AddIncrementOp(const FRigVMOperand& InArg)
@@ -827,12 +854,16 @@ uint64 FRigVMByteCode::AddJumpIfOp(ERigVMOpCode InOpCode, uint16 InInstructionIn
 	return AddOp(Op);
 }
 
+#if UE_RIGVM_UCLASS_BASED_STORAGE_DISABLED
+
 uint64 FRigVMByteCode::AddChangeTypeOp(FRigVMOperand InArg, ERigVMRegisterType InType, uint16 InElementSize, uint16 InElementCount, uint16 InSliceCount)
 {
 	FRigVMChangeTypeOp Op(InArg, InType, InElementSize, InElementCount, InSliceCount);
 	return AddOp(Op);
 }
 
+#endif
+
 uint64 FRigVMByteCode::AddExecuteOp(uint16 InFunctionIndex, const FRigVMOperandArray& InOperands)
 {
 	FRigVMExecuteOp Op(InFunctionIndex, (uint8)InOperands.Num());
@@ -956,7 +987,9 @@ FString FRigVMByteCode::DumpToText() const
 				FString TargetContent;
 				FRigVMOperand::StaticStruct()->ExportText(TargetContent, &Op.Source, nullptr, nullptr, PPF_None, nullptr);
 				Line += FString::Printf(TEXT(", Target %s"), *TargetContent);
+#if UE_RIGVM_UCLASS_BASED_STORAGE_DISABLED
 				Line += FString::Printf(TEXT(", NumBytes %d"), (int32)Op.NumBytes);
+#endif
 				break;
 			}
 			case ERigVMOpCode::Zero:
@@ -1149,8 +1182,13 @@ uint64 FRigVMByteCode::GetOpAlignment(ERigVMOpCode InOpCode) const
 		}
 		case ERigVMOpCode::ChangeType:
 		{
+#if UE_RIGVM_UCLASS_BASED_STORAGE_DISABLED
 			static const uint64 Alignment = FRigVMChangeTypeOp::StaticStruct()->GetCppStructOps()->GetAlignment();
 			return Alignment;
+#else
+			checkNoEntry();
+			return 0;
+#endif
 		}
 		case ERigVMOpCode::Exit:
 		{
diff --git a/Engine/Source/Runtime/RigVM/Public/RigVMCore/RigVM.h b/Engine/Source/Runtime/RigVM/Public/RigVMCore/RigVM.h
--- a/Engine/Source/Runtime/RigVM/Public/RigVMCore/RigVM.h
+++ b/Engine/Source/Runtime/RigVM/Public/RigVMCore/RigVM.h
@@ -4,9 +4,11 @@
 
 #include "CoreMinimal.h"
 #include "RigVMMemory.h"
+#include "RigVMMemoryStorage.h"
 #include "RigVMExecuteContext.h"
 #include "RigVMRegistry.h"
 #include "RigVMByteCode.h"
+#include "RigVMMemoryDeprecated.h"
 #include "RigVMStatistics.h"
 #if WITH_EDITOR
 #include "RigVMDebugInfo.h"
@@ -158,12 +160,20 @@ public:
 	void CopyFrom(URigVM* InVM, bool bDeferCopy = false, bool bReferenceLiteralMemory = false, bool bReferenceByteCode = false, bool bCopyExternalVariables = false, bool bCopyDynamicRegisters = false);
 
 	// Initializes all execute ops and their memory.
+#if UE_RIGVM_UCLASS_BASED_STORAGE_DISABLED
 	bool Initialize(FRigVMMemoryContainerPtrArray Memory, FRigVMFixedArray<void*> AdditionalArguments);
+#else
+	bool Initialize(TArrayView<URigVMMemoryStorage*> Memory, TArrayView<void*> AdditionalArguments);
+#endif
 
 	// Executes the VM.
 	// You can optionally provide external memory to the execution
 	// and provide optional additional operands.
+#if UE_RIGVM_UCLASS_BASED_STORAGE_DISABLED
 	bool Execute(FRigVMMemoryContainerPtrArray Memory, FRigVMFixedArray<void*> AdditionalArguments, const FName& InEntryName = NAME_None);
+#else
+	bool Execute(TArrayView<URigVMMemoryStorage*> Memory, TArrayView<void*> AdditionalArguments, const FName& InEntryName = NAME_None);
+#endif
 
 	// Executes the VM.
 	// You can optionally provide external memory to the execution
@@ -180,27 +190,68 @@ public:
 	UFUNCTION()
 	FString GetRigVMFunctionName(int32 InFunctionIndex) const;
 
+#if UE_RIGVM_UCLASS_BASED_STORAGE_DISABLED
+	
 	// The default mutable work memory
-	UPROPERTY()
 	FRigVMMemoryContainer WorkMemoryStorage;
 	FRigVMMemoryContainer* WorkMemoryPtr;
 	FORCEINLINE FRigVMMemoryContainer& GetWorkMemory() { return *WorkMemoryPtr; }
 	FORCEINLINE const FRigVMMemoryContainer& GetWorkMemory() const { return *WorkMemoryPtr; }
 
 	// The default const literal memory
-	UPROPERTY()
 	FRigVMMemoryContainer LiteralMemoryStorage;
 	FRigVMMemoryContainer* LiteralMemoryPtr;
 	FORCEINLINE FRigVMMemoryContainer& GetLiteralMemory() { return *LiteralMemoryPtr; }
 	FORCEINLINE const FRigVMMemoryContainer& GetLiteralMemory() const { return *LiteralMemoryPtr; }
 
 	// The default debug watch memory
-	UPROPERTY()
 	FRigVMMemoryContainer DebugMemoryStorage;
 	FRigVMMemoryContainer* DebugMemoryPtr;
 	FORCEINLINE FRigVMMemoryContainer& GetDebugMemory() { return *DebugMemoryPtr; }
 	FORCEINLINE const FRigVMMemoryContainer& GetDebugMemory() const { return *DebugMemoryPtr; }
 
+#else
+
+	// Returns a memory storage by type
+	URigVMMemoryStorage* GetMemoryByType(ERigVMMemoryType InMemoryType) const;
+	
+	// The default mutable work memory
+	FORCEINLINE URigVMMemoryStorage* GetWorkMemory() const { return WorkMemoryStorageObject; }
+
+	// The default const literal memory
+	FORCEINLINE URigVMMemoryStorage* GetLiteralMemory() const { return LiteralMemoryStorageObject; }
+
+	// The default debug watch memory
+	FORCEINLINE URigVMMemoryStorage* GetDebugMemory() const { return DebugMemoryStorageObject; }
+
+	// returns all memory storages as an array
+	FORCEINLINE TArray<URigVMMemoryStorage*> GetLocalMemoryArray() const
+	{
+		check(GetWorkMemory());
+		check(GetLiteralMemory());
+		check(GetDebugMemory());
+		
+		TArray<URigVMMemoryStorage*> LocalMemory;
+		LocalMemory.Add(GetWorkMemory());
+		LocalMemory.Add(GetLiteralMemory());
+		LocalMemory.Add(GetDebugMemory());
+		
+		return LocalMemory;
+	}
+
+#endif
+
+	UPROPERTY()
+	URigVMMemoryStorage* WorkMemoryStorageObject;
+
+	UPROPERTY()
+	URigVMMemoryStorage* LiteralMemoryStorageObject;
+
+	UPROPERTY()
+	URigVMMemoryStorage* DebugMemoryStorageObject;
+
+	TArray<FRigVMPropertyPath> PropertyPaths;
+
 	// The byte code of the VM
 	UPROPERTY()
 	FRigVMByteCode ByteCodeStorage;
@@ -240,7 +291,11 @@ public:
 
 	FORCEINLINE const void SetFirstEntryEventInEventQueue(const FName& InFirstEventName) { FirstEntryEventInQueue = InFirstEventName; }
 
+#if UE_RIGVM_UCLASS_BASED_STORAGE_DISABLED
 	bool ResumeExecution(FRigVMMemoryContainerPtrArray Memory, FRigVMFixedArray<void*> AdditionalArguments, const FName& InEntryName = NAME_None);
+#else
+	bool ResumeExecution(TArrayView<URigVMMemoryStorage*> Memory, TArrayView<void*> AdditionalArguments, const FName& InEntryName = NAME_None);
+#endif
 
 	bool ResumeExecution();
 #endif
@@ -251,6 +306,8 @@ public:
 	// Returns a parameter given it's name
 	FRigVMParameter GetParameterByName(const FName& InParameterName);
 
+#if UE_RIGVM_UCLASS_BASED_STORAGE_DISABLED
+
 	// Adds a new input / output to the VM
 	template<class T>
 	FORCEINLINE_DEBUGGABLE FRigVMParameter AddParameter(ERigVMParameterType InParameterType, const FName& InName, const FString& InCPPType, const TArray<T>& DefaultValues)
@@ -265,21 +322,20 @@ public:
 			RegisterIndex = WorkMemoryPtr->Add<T>(WorkMemoryPtr->SupportsNames() ? InName : NAME_None, DefaultValues[0], 1);
 		}
 		else
-	{
+		{
 			RegisterIndex = WorkMemoryPtr->AddFixedArray<T>(WorkMemoryPtr->SupportsNames() ? InName : NAME_None, FRigVMFixedArray<T>(DefaultValues), 1);
-	}
+		}
 
 		if (RegisterIndex == INDEX_NONE)
-	{
+		{
 			return FRigVMParameter();
-	}
+		}
 
 		FName Name = WorkMemoryPtr->SupportsNames() ? GetWorkMemory()[RegisterIndex].Name : InName;
 
 		FRigVMParameter Parameter(InParameterType, Name, RegisterIndex, InCPPType, GetWorkMemory().GetScriptStruct(RegisterIndex));
 		ParametersNameMap.Add(Parameter.Name, Parameters.Add(Parameter));
 		return Parameter;
-
 	}
 
 	// Adds a new input / output to the VM
@@ -290,7 +346,7 @@ public:
 		DefaultValues.Add(DefaultValue);
 		return AddParameter(InParameterType, InName, InCPPType, DefaultValues);
 	}
-
+	
 	// Retrieve the array size of the parameter
 	int32 GetParameterArraySize(const FRigVMParameter& InParameter) const
 	{
@@ -304,20 +360,38 @@ public:
 	}
 
 	// Retrieve the array size of the parameter
-	UFUNCTION(BlueprintCallable, Category = RigVM)
 	int32 GetParameterArraySize(const FName& InParameterName) const
 	{
 		int32 ParameterIndex = ParametersNameMap.FindChecked(InParameterName);
 		return GetParameterArraySize(ParameterIndex);
 	}
 
+#endif
+	
 	// Retrieve the value of a parameter
 	template<class T>
 	T GetParameterValue(const FRigVMParameter& InParameter, int32 InArrayIndex = 0, T DefaultValue = T{})
 	{
 		if (InParameter.GetRegisterIndex() != INDEX_NONE)
 		{
+#if UE_RIGVM_UCLASS_BASED_STORAGE_DISABLED
 			return WorkMemoryPtr->GetFixedArray<T>(InParameter.GetRegisterIndex())[InArrayIndex];
+#else
+			if(GetWorkMemory()->IsArray(InParameter.GetRegisterIndex()))
+			{
+				TArray<T>& Storage = *GetWorkMemory()->GetData<TArray<T>>(InParameter.GetRegisterIndex());
+				if(Storage.IsValidIndex(InArrayIndex))
+				{
+					return Storage[InArrayIndex];
+				}
+			}
+			else
+			{
+				return *GetWorkMemory()->GetData<T>(InParameter.GetRegisterIndex());
+			}
+			
+			return *GetWorkMemory()->GetData<T>(InParameter.GetRegisterIndex());
+#endif
 		}
 		return DefaultValue;
 	}
@@ -343,7 +417,23 @@ public:
 	{
 		if (InParameter.GetRegisterIndex() != INDEX_NONE)
 		{
+#if UE_RIGVM_UCLASS_BASED_STORAGE_DISABLED
 			WorkMemoryPtr->GetFixedArray<T>(InParameter.GetRegisterIndex())[InArrayIndex] = InNewValue;
+#else
+			if(GetWorkMemory()->IsArray(InParameter.GetRegisterIndex()))
+			{
+				TArray<T>& Storage = *GetWorkMemory()->GetData<TArray<T>>(InParameter.GetRegisterIndex());
+				if(Storage.IsValidIndex(InArrayIndex))
+				{
+					Storage[InArrayIndex] = InNewValue;
+				}
+			}
+			else
+			{
+				T& Storage = *GetWorkMemory()->GetData<T>(InParameter.GetRegisterIndex());
+				Storage = InNewValue;
+			}
+#endif
 		}
 	}
 
@@ -512,15 +602,21 @@ public:
 		return AddExternalVariable(FRigVMExternalVariable::Make(InExternalVariableName, InValue));
 	}
 
+#if UE_RIGVM_UCLASS_BASED_STORAGE_DISABLED
 	void SetRegisterValueFromString(const FRigVMOperand& InOperand, const FString& InCPPType, const UObject* InCPPTypeObject, const TArray<FString>& InDefaultValues);
+#else
+	void SetPropertyValueFromString(const FRigVMOperand& InOperand, const FString& InDefaultValue);
+#endif
 
 	// returns the statistics information
 	FRigVMStatistics GetStatistics() const
 	{
 		FRigVMStatistics Statistics;
+#if UE_RIGVM_UCLASS_BASED_STORAGE_DISABLED
 		Statistics.LiteralMemory = LiteralMemoryPtr->GetStatistics();
 		Statistics.WorkMemory = WorkMemoryPtr->GetStatistics();
 		Statistics.DebugMemory = DebugMemoryPtr->GetStatistics();
+#endif
 		Statistics.ByteCode = ByteCodePtr->GetStatistics();
 		Statistics.BytesForCaching = FirstHandleForInstruction.GetAllocatedSize() + CachedMemoryHandles.GetAllocatedSize();
 		Statistics.BytesForCDO =
@@ -563,7 +659,11 @@ private:
 	void ResolveFunctionsIfRequired();
 	void RefreshInstructionsIfRequired();
 	void InvalidateCachedMemory();
+#if UE_RIGVM_UCLASS_BASED_STORAGE_DISABLED
 	void CacheMemoryHandlesIfRequired(FRigVMMemoryContainerPtrArray InMemory);
+#else
+	void CacheMemoryHandlesIfRequired(TArrayView<URigVMMemoryStorage*> InMemory);
+#endif
 	void RebuildByteCodeOnLoad();
 
 	UPROPERTY(transient)
@@ -601,7 +701,11 @@ private:
 	TArray<uint32> FirstHandleForInstruction;
 	TArray<FRigVMMemoryHandle> CachedMemoryHandles;
 	// changes to the layout of cached memory array should be reflected in GetContainerIndex()
+#if UE_RIGVM_UCLASS_BASED_STORAGE_DISABLED
 	TArray<FRigVMMemoryContainer*> CachedMemory;
+#else
+	TArray<URigVMMemoryStorage*> CachedMemory;
+#endif
 	TArray<FRigVMExternalVariable> ExternalVariables;
 
 	// this function should be kept in sync with FRigVMOperand::GetContainerIndex()
@@ -653,6 +757,7 @@ private:
 
 	void CopyOperandForDebuggingImpl(const FRigVMOperand& InArg, const FRigVMMemoryHandle& InHandle, const FRigVMOperand& InDebugOperand);
 
+#if UE_RIGVM_UCLASS_BASED_STORAGE_DISABLED
 	struct FCopyInfoForOperand
 	{
 		FCopyInfoForOperand()
@@ -671,11 +776,14 @@ private:
 		uint16 NumBytesToCopy;
 		uint16 ElementSize;
 	};
+#endif
 
 	FRigVMCopyOp GetCopyOpForOperands(const FRigVMOperand& InSource, const FRigVMOperand& InTarget);
+#if UE_RIGVM_UCLASS_BASED_STORAGE_DISABLED	
 	FCopyInfoForOperand GetCopyInfoForOperand(const FRigVMOperand& InOperand);
 	UScriptStruct* GetScriptStructForCopyOp(const FRigVMCopyOp& InCopyOp) const;
 	UScriptStruct* GetScripStructForOperand(const FRigVMOperand& InOperand) const;
+#endif
 	
 	TMap<FRigVMOperand, TArray<FRigVMOperand>> OperandToDebugRegisters;
 
diff --git a/Engine/Source/Runtime/RigVM/Public/RigVMCore/RigVMByteCode.h b/Engine/Source/Runtime/RigVM/Public/RigVMCore/RigVMByteCode.h
--- a/Engine/Source/Runtime/RigVM/Public/RigVMCore/RigVMByteCode.h
+++ b/Engine/Source/Runtime/RigVM/Public/RigVMCore/RigVMByteCode.h
@@ -94,7 +94,7 @@ enum class ERigVMOpCode : uint8
 	JumpAbsoluteIf, // jump to an absolute instruction index based on a condition register
 	JumpForwardIf, // jump forwards given a relative instruction index offset based on a condition register
 	JumpBackwardIf, // jump backwards given a relative instruction index offset based on a condition register
-	ChangeType, // change the type of a register
+	ChangeType, // change the type of a register (deprecated)
 	Exit, // exit the execution loop
 	BeginBlock, // begins a new memory slice / block
 	EndBlock, // ends the last memory slice / block
@@ -238,33 +238,41 @@ struct RIGVM_API FRigVMCopyOp : public FRigVMBaseOp
 	: FRigVMBaseOp(ERigVMOpCode::Copy)
 	, Source()
 	, Target()
+#if UE_RIGVM_UCLASS_BASED_STORAGE_DISABLED
 	, NumBytes(0)
 	, RegisterType(ERigVMRegisterType::Invalid)
 	, CopyType(ERigVMCopyType::Default)
+#endif
 	{
 	}
 
 	FRigVMCopyOp(
 		FRigVMOperand InSource,
-		FRigVMOperand InTarget,
-		uint16 InNumBytes,
+		FRigVMOperand InTarget
+#if UE_RIGVM_UCLASS_BASED_STORAGE_DISABLED
+		, uint16 InNumBytes,
 		ERigVMRegisterType InRegisterType,
 		ERigVMCopyType InCopyType
+#endif
 	)
 		: FRigVMBaseOp(ERigVMOpCode::Copy)
 		, Source(InSource)
 		, Target(InTarget)
+#if UE_RIGVM_UCLASS_BASED_STORAGE_DISABLED
 		, NumBytes(InNumBytes)
 		, RegisterType(InRegisterType)
 		, CopyType(InCopyType)
+#endif
 	{
 	}
 
 	FRigVMOperand Source;
 	FRigVMOperand Target;
+#if UE_RIGVM_UCLASS_BASED_STORAGE_DISABLED
 	uint16 NumBytes;
 	ERigVMRegisterType RegisterType;
 	ERigVMCopyType CopyType;
+#endif
 
 	void Serialize(FArchive& Ar);
 	FORCEINLINE friend FArchive& operator<<(FArchive& Ar, FRigVMCopyOp& P)
@@ -390,6 +398,8 @@ struct RIGVM_API FRigVMChangeTypeOp : public FRigVMUnaryOp
 {
 	GENERATED_USTRUCT_BODY()
 
+#if UE_RIGVM_UCLASS_BASED_STORAGE_DISABLED
+
 	FRigVMChangeTypeOp()
 		: FRigVMUnaryOp()
 		, Type(ERigVMRegisterType::Invalid)
@@ -413,6 +423,8 @@ struct RIGVM_API FRigVMChangeTypeOp : public FRigVMUnaryOp
 	uint16 ElementCount;
 	uint16 SliceCount;
 
+#endif
+
 	void Serialize(FArchive& Ar);
 	FORCEINLINE friend FArchive& operator<<(FArchive& Ar, FRigVMChangeTypeOp& P)
 	{
@@ -565,7 +577,11 @@ public:
 	uint64 AddTrueOp(const FRigVMOperand& InArg);
 
 	// adds a copy operator to copy the content of a source argument to a target argument
+#if UE_RIGVM_UCLASS_BASED_STORAGE_DISABLED
 	uint64 AddCopyOp(const FRigVMOperand& InSource, const FRigVMOperand& InTarget, uint16 InNumBytes, ERigVMRegisterType InTargetType, ERigVMCopyType InCopyType);
+#else
+	uint64 AddCopyOp(const FRigVMOperand& InSource, const FRigVMOperand& InTarget);
+#endif
 
 	// adds a copy operator to copy the content of a source argument to a target argument
 	uint64 AddCopyOp(const FRigVMCopyOp& InCopyOp);
@@ -588,9 +604,13 @@ public:
 	// adds an absolute, forward or backward jump operator based on a condition argument
 	uint64 AddJumpIfOp(ERigVMOpCode InOpCode, uint16 InInstructionIndex, const FRigVMOperand& InConditionArg, bool bJumpWhenConditionIs = false);
 
+#if UE_RIGVM_UCLASS_BASED_STORAGE_DISABLED
+
 	// adds a change-type operator to reuse a register for a smaller or same size type
 	uint64 AddChangeTypeOp(FRigVMOperand InArg, ERigVMRegisterType InType, uint16 InElementSize, uint16 InElementCount, uint16 InSliceCount = 1);
 
+#endif
+
 	// adds an exit operator to exit the execution loop
 	uint64 AddExitOp();
 
diff --git a/Engine/Source/Runtime/RigVM/Public/RigVMCore/RigVMExternalVariable.h b/Engine/Source/Runtime/RigVM/Public/RigVMCore/RigVMExternalVariable.h
--- a/Engine/Source/Runtime/RigVM/Public/RigVMCore/RigVMExternalVariable.h
+++ b/Engine/Source/Runtime/RigVM/Public/RigVMCore/RigVMExternalVariable.h
@@ -14,6 +14,9 @@ struct RIGVM_API FRigVMExternalVariable
 {
 	FORCEINLINE FRigVMExternalVariable()
 		: Name(NAME_None)
+#if !UE_RIGVM_UCLASS_BASED_STORAGE_DISABLED
+		, Property(nullptr)
+#endif
 		, TypeName(NAME_None)
 		, TypeObject(nullptr)
 		, bIsArray(false)
@@ -21,10 +24,60 @@ struct RIGVM_API FRigVMExternalVariable
 		, bIsReadOnly(false)
 		, Size(0)
 		, Memory(nullptr)
-		
 	{
 	}
 
+	FORCEINLINE static void GetTypeFromProperty(const FProperty* InProperty, FName& OutTypeName, UObject*& OutTypeObject)
+	{
+		if (CastField<FBoolProperty>(InProperty))
+		{
+			OutTypeName = TEXT("bool");
+		}
+		else if (CastField<FIntProperty>(InProperty))
+		{
+			OutTypeName = TEXT("int32");
+		}
+		else if (CastField<FFloatProperty>(InProperty))
+		{
+			OutTypeName = TEXT("float");
+		}
+		else if (CastField<FDoubleProperty>(InProperty))
+		{
+			OutTypeName = TEXT("double");
+		}
+		else if (CastField<FStrProperty>(InProperty))
+		{
+			OutTypeName = TEXT("FString");
+		}
+		else if (CastField<FNameProperty>(InProperty))
+		{
+			OutTypeName = TEXT("FName");
+		}
+		else if (const FEnumProperty* EnumProperty = CastField<FEnumProperty>(InProperty))
+		{
+			OutTypeName = EnumProperty->GetEnum()->GetFName();
+			OutTypeObject = EnumProperty->GetEnum();
+		}
+		else if (const FByteProperty* ByteProperty = CastField<FByteProperty>(InProperty))
+		{
+			if (UEnum* BytePropertyEnum = ByteProperty->Enum)
+			{
+				OutTypeName = BytePropertyEnum->GetFName();
+				OutTypeObject = BytePropertyEnum;
+			}
+		}
+		else if (const FStructProperty* StructProperty = CastField<FStructProperty>(InProperty))
+		{
+			OutTypeName = *StructProperty->Struct->GetStructCPPName();
+			OutTypeObject = StructProperty->Struct;
+		}
+		else if (const FObjectProperty* ObjectProperty = CastField<FObjectProperty>(InProperty))
+		{
+			OutTypeName = *ObjectProperty->PropertyClass->GetName();
+			OutTypeObject = ObjectProperty->PropertyClass;
+		}
+	}
+
 	FORCEINLINE static FRigVMExternalVariable Make(FProperty* InProperty, UObject* InContainer)
 	{
 		check(InProperty);
@@ -33,6 +86,9 @@ struct RIGVM_API FRigVMExternalVariable
 
 		FRigVMExternalVariable ExternalVariable;
 		ExternalVariable.Name = InProperty->GetFName();
+#if !UE_RIGVM_UCLASS_BASED_STORAGE_DISABLED
+		ExternalVariable.Property = Property;
+#endif
 		ExternalVariable.bIsPublic = !InProperty->HasAllPropertyFlags(CPF_DisableEditOnInstance);
 		ExternalVariable.bIsReadOnly = InProperty->HasAllPropertyFlags(CPF_BlueprintReadOnly);
 
@@ -51,54 +107,7 @@ struct RIGVM_API FRigVMExternalVariable
 		}
 
 		ExternalVariable.Size = Property->GetSize();
-
-		if (CastField<FBoolProperty>(Property))
-		{
-			ExternalVariable.TypeName = TEXT("bool");
-		}
-		else if (CastField<FIntProperty>(Property))
-		{
-			ExternalVariable.TypeName = TEXT("int32");
-		}
-		else if (CastField<FFloatProperty>(Property))
-		{
-			ExternalVariable.TypeName = TEXT("float");
-		}
-		else if (CastField<FDoubleProperty>(Property))
-		{
-			ExternalVariable.TypeName = TEXT("double");
-		}
-		else if (CastField<FStrProperty>(Property))
-		{
-			ExternalVariable.TypeName = TEXT("FString");
-		}
-		else if (CastField<FNameProperty>(Property))
-		{
-			ExternalVariable.TypeName = TEXT("FName");
-		}
-		else if (FEnumProperty* EnumProperty = CastField<FEnumProperty>(Property))
-		{
-			ExternalVariable.TypeName = EnumProperty->GetEnum()->GetFName();
-			ExternalVariable.TypeObject = EnumProperty->GetEnum();
-		}
-		else if (FByteProperty* ByteProperty = CastField<FByteProperty>(Property))
-		{
-			if (UEnum* BytePropertyEnum = ByteProperty->Enum)
-			{
-				ExternalVariable.TypeName = BytePropertyEnum->GetFName();
-				ExternalVariable.TypeObject = BytePropertyEnum;
-			}
-		}
-		else if (FStructProperty* StructProperty = CastField<FStructProperty>(Property))
-		{
-			ExternalVariable.TypeName = *StructProperty->Struct->GetStructCPPName();
-			ExternalVariable.TypeObject = StructProperty->Struct;
-		}
-		else if (FObjectProperty* ObjectProperty = CastField<FObjectProperty>(Property))
-		{
-			ExternalVariable.TypeName = *ObjectProperty->PropertyClass->GetName();
-			ExternalVariable.TypeObject = ObjectProperty->PropertyClass;
-		}
+		GetTypeFromProperty(Property, ExternalVariable.TypeName, ExternalVariable.TypeObject);
 
 		return ExternalVariable;
 	}
@@ -484,6 +493,9 @@ struct RIGVM_API FRigVMExternalVariable
 	}
 
 	FName Name;
+#if !UE_RIGVM_UCLASS_BASED_STORAGE_DISABLED
+	const FProperty* Property;
+#endif
 	FName TypeName;
 	UObject* TypeObject;
 	bool bIsArray;
