Control Rig: RigVM: Implement Compiler support for RigVMMemoryStorage

diff --git a/Engine/Source/Developer/RigVMDeveloper/Private/RigVMCompiler/RigVMCompiler.cpp b/Engine/Source/Developer/RigVMDeveloper/Private/RigVMCompiler/RigVMCompiler.cpp
--- a/Engine/Source/Developer/RigVMDeveloper/Private/RigVMCompiler/RigVMCompiler.cpp
+++ b/Engine/Source/Developer/RigVMDeveloper/Private/RigVMCompiler/RigVMCompiler.cpp
@@ -68,9 +68,9 @@ FRigVMOperand FRigVMCompilerWorkData::AddProperty(
 {
 	check(bSetupMemory);
 	
-	URigVMMemoryStorage::FPropertyDescription Description(InName, InCPPType, InCPPTypeObject, InDefaultValue);
+	FRigVMPropertyDescription Description(InName, InCPPType, InCPPTypeObject, InDefaultValue);
 
-	TArray<URigVMMemoryStorage::FPropertyDescription>& PropertyArray = PropertyDescriptions.FindOrAdd(InMemoryType);
+	TArray<FRigVMPropertyDescription>& PropertyArray = PropertyDescriptions.FindOrAdd(InMemoryType);
 	const int32 PropertyIndex = PropertyArray.Add(Description);
 
 	return FRigVMOperand(InMemoryType, PropertyIndex);
@@ -78,7 +78,7 @@ FRigVMOperand FRigVMCompilerWorkData::AddProperty(
 
 FRigVMOperand FRigVMCompilerWorkData::FindProperty(ERigVMMemoryType InMemoryType, const FName& InName)
 {
-	TArray<URigVMMemoryStorage::FPropertyDescription>* PropertyArray = PropertyDescriptions.Find(InMemoryType);
+	TArray<FRigVMPropertyDescription>* PropertyArray = PropertyDescriptions.Find(InMemoryType);
 	if(PropertyArray)
 	{
 		for(int32 Index=0;Index<PropertyArray->Num();Index++)
@@ -92,9 +92,9 @@ FRigVMOperand FRigVMCompilerWorkData::FindProperty(ERigVMMemoryType InMemoryType
 	return FRigVMOperand();
 }
 
-URigVMMemoryStorage::FPropertyDescription FRigVMCompilerWorkData::GetProperty(const FRigVMOperand& InOperand)
+FRigVMPropertyDescription FRigVMCompilerWorkData::GetProperty(const FRigVMOperand& InOperand)
 {
-	TArray<URigVMMemoryStorage::FPropertyDescription>* PropertyArray = PropertyDescriptions.Find(InOperand.GetMemoryType());
+	TArray<FRigVMPropertyDescription>* PropertyArray = PropertyDescriptions.Find(InOperand.GetMemoryType());
 	if(PropertyArray)
 	{
 		if(PropertyArray->IsValidIndex(InOperand.GetRegisterIndex()))
@@ -102,23 +102,43 @@ URigVMMemoryStorage::FPropertyDescription FRigVMCompilerWorkData::GetProperty(co
 			return PropertyArray->operator[](InOperand.GetRegisterIndex());
 		}
 	}
-	return URigVMMemoryStorage::FPropertyDescription();
+	return FRigVMPropertyDescription();
 }
 
-int32 FRigVMCompilerWorkData::FindPropertyPath(const FString& InSegmentPath) const
+int32 FRigVMCompilerWorkData::FindPropertyPath(ERigVMMemoryType InMemoryType, const FString& InSegmentPath) const
 {
-	for(int32 Index = 0; Index < PropertyPathDescriptions.Num(); Index++)
+	const TArray<FRigVMPropertyPathDescription>* Descriptions = PropertyPathDescriptions.Find(InMemoryType);
+	if(Descriptions)
 	{
-		const FRigVMCompilerWorkData::FPropertyPathDescription& Description = PropertyPathDescriptions[Index]; 
-		if(Description.SegmentPath == InSegmentPath)
+		for(int32 Index = 0; Index < Descriptions->Num(); Index++)
 		{
-			return Index;
+			const FRigVMPropertyPathDescription& Description = Descriptions->operator[](Index); 
+			if(Description.SegmentPath == InSegmentPath)
+			{
+				return Index;
+			}
 		}
 	}
 
 	return INDEX_NONE;
 }
 
+int32 FRigVMCompilerWorkData::FindOrAddPropertyPath(const FRigVMOperand& InOperand, const FString& InSegmentPath)
+{
+	check(bSetupMemory);
+	
+	TArray<FRigVMPropertyPathDescription>& Descriptions = PropertyPathDescriptions.FindOrAdd(InOperand.GetMemoryType());
+	for(int32 Index = 0; Index < Descriptions.Num(); Index++)
+	{
+		const FRigVMPropertyPathDescription& Description = Descriptions[Index]; 
+		if(Description.SegmentPath == InSegmentPath)
+		{
+			return Index;
+		}
+	}
+	return Descriptions.Add(FRigVMPropertyPathDescription(InOperand.GetRegisterIndex(), InSegmentPath));
+}
+
 #endif
 
 URigVMCompiler::URigVMCompiler()
@@ -353,28 +373,18 @@ bool URigVMCompiler::Compile(URigVMGraph* InGraph, URigVMController* InControlle
 			}
 		}
 
-		TArray<URigVMMemoryStorage::FPropertyDescription>* Properties = WorkData.PropertyDescriptions.Find(MemoryType);
+		TArray<FRigVMPropertyDescription>* Properties = WorkData.PropertyDescriptions.Find(MemoryType);
 		if(Properties == nullptr)
 		{
-			*MemoryStorageObject = NewObject<URigVMMemoryStorage>(WorkData.VM, NAME_None, RF_Public | RF_Transactional);
+			*MemoryStorageObject = NewObject<URigVMMemoryStorage>(WorkData.VM, NAME_None, RF_Public);
 			continue;
 		}
 
 		UPackage* Package = InGraph->GetOutermost();
-		UClass* MemoryClass = URigVMMemoryStorage::CreateStorageClass(Package, MemoryType, *Properties);
-		*MemoryStorageObject = NewObject<URigVMMemoryStorage>(WorkData.VM, MemoryClass, NAME_None, RF_Public | RF_Transactional); 
+		UClass* MemoryClass = URigVMMemoryStorageGeneratorClass::CreateStorageClass(Package, MemoryType, *Properties);
+		*MemoryStorageObject = NewObject<URigVMMemoryStorage>(WorkData.VM, MemoryClass, NAME_None, RF_Public); 
 	}
 
-	WorkData.VM->PropertyPaths.Reset();
-	for(FRigVMCompilerWorkData::FPropertyPathDescription Description : WorkData.PropertyPathDescriptions)
-	{
-		URigVMMemoryStorage* MemoryStorageObject = WorkData.VM->GetMemoryByType(Description.OriginalOperand.GetMemoryType());
-		const FProperty* Property = MemoryStorageObject->GetProperties()[Description.OriginalOperand.GetRegisterIndex()];
-
-		FRigVMPropertyPath PropertyPath(Property, Description.SegmentPath);
-		WorkData.VM->PropertyPaths.Add(PropertyPath);
-	}
-	
 #endif
 
 	WorkData.bSetupMemory = false;
@@ -867,17 +877,7 @@ void URigVMCompiler::TraverseAssign(const FRigVMAssignExprAST* InExpr, FRigVMCom
 
 #else
 
-					int32 PropertyPathIndex = WorkData.FindPropertyPath(SegmentPath);
-					if(PropertyPathIndex == INDEX_NONE)
-					{
-						check(WorkData.bSetupMemory);
-						
-						FRigVMCompilerWorkData::FPropertyPathDescription PropertyPathDescription;
-						PropertyPathDescription.OriginalOperand = Operand;
-						PropertyPathDescription.SegmentPath = SegmentPath;
-						PropertyPathIndex = WorkData.PropertyPathDescriptions.Add(PropertyPathDescription);
-					}
-					
+					const int32 PropertyPathIndex = WorkData.FindOrAddPropertyPath(Operand, SegmentPath);
 					Operand = FRigVMOperand(Operand.GetMemoryType(), Operand.GetRegisterIndex(), PropertyPathIndex);
 					
 #endif
@@ -1840,17 +1840,7 @@ FRigVMOperand URigVMCompiler::FindOrAddRegister(const FRigVMVarExprAST* InVarExp
 
 #else
 
-						int32 PropertyPathIndex = WorkData.FindPropertyPath(SegmentPath);
-						if(PropertyPathIndex == INDEX_NONE)
-						{
-							check(WorkData.bSetupMemory);
-
-							FRigVMCompilerWorkData::FPropertyPathDescription PropertyPathDescription;
-							PropertyPathDescription.OriginalOperand = Operand;
-							PropertyPathDescription.SegmentPath = SegmentPath;
-							PropertyPathIndex = WorkData.PropertyPathDescriptions.Add(PropertyPathDescription);
-						}
-
+						const int32 PropertyPathIndex = WorkData.FindOrAddPropertyPath(Operand, SegmentPath);
 						Operand.RegisterOffset = (uint16)PropertyPathIndex;
 						
 #endif
@@ -1994,7 +1984,7 @@ FRigVMOperand URigVMCompiler::FindOrAddRegister(const FRigVMVarExprAST* InVarExp
 				Operand = WorkData.FindProperty(MemoryType, RegisterName);
 				if(Operand.IsValid())
 				{
-					URigVMMemoryStorage::FPropertyDescription Property = WorkData.GetProperty(Operand);
+					FRigVMPropertyDescription Property = WorkData.GetProperty(Operand);
 					if(Property.IsValid())
 					{
 						if(ExistingOperand == nullptr)
diff --git a/Engine/Source/Developer/RigVMDeveloper/Public/RigVMCompiler/RigVMCompiler.h b/Engine/Source/Developer/RigVMDeveloper/Public/RigVMCompiler/RigVMCompiler.h
--- a/Engine/Source/Developer/RigVMDeveloper/Public/RigVMCompiler/RigVMCompiler.h
+++ b/Engine/Source/Developer/RigVMDeveloper/Public/RigVMCompiler/RigVMCompiler.h
@@ -93,19 +93,15 @@ public:
 	FRigVMOperand ComparisonOperand;
 
 #if !UE_RIGVM_UCLASS_BASED_STORAGE_DISABLED
-	struct FPropertyPathDescription
-	{
-		FRigVMOperand OriginalOperand;
-		FString SegmentPath;
-	};
 	
-	TArray<FPropertyPathDescription> PropertyPathDescriptions;
-	TMap<ERigVMMemoryType, TArray<URigVMMemoryStorage::FPropertyDescription>> PropertyDescriptions;
+	TMap<ERigVMMemoryType, TArray<FRigVMPropertyPathDescription>> PropertyPathDescriptions;
+	TMap<ERigVMMemoryType, TArray<FRigVMPropertyDescription>> PropertyDescriptions;
 
 	FRigVMOperand AddProperty(ERigVMMemoryType InMemoryType, const FName& InName, const FString& InCPPType, UObject* InCPPTypeObject, const FString& InDefaultValue = FString());
 	FRigVMOperand FindProperty(ERigVMMemoryType InMemoryType, const FName& InName);
-	URigVMMemoryStorage::FPropertyDescription GetProperty(const FRigVMOperand& InOperand);
-	int32 FindPropertyPath(const FString& InSegmentPath) const;
+	FRigVMPropertyDescription GetProperty(const FRigVMOperand& InOperand);
+	int32 FindPropertyPath(ERigVMMemoryType InMemoryType, const FString& InSegmentPath) const;
+	int32 FindOrAddPropertyPath(const FRigVMOperand& InOperand, const FString& InSegmentPath);
 	
 #endif
 };
diff --git a/Engine/Source/Runtime/RigVM/Private/RigVMCore/RigVM.cpp b/Engine/Source/Runtime/RigVM/Private/RigVMCore/RigVM.cpp
--- a/Engine/Source/Runtime/RigVM/Private/RigVMCore/RigVM.cpp
+++ b/Engine/Source/Runtime/RigVM/Private/RigVMCore/RigVM.cpp
@@ -144,8 +144,18 @@ void URigVM::Save(FArchive& Ar)
 
 #else
 
-	WorkMemoryStorageObject->Serialize(Ar);
-	LiteralMemoryStorageObject->Serialize(Ar);
+	bool bWorkMemoryStorageObjectValid = WorkMemoryStorageObject != nullptr;
+	Ar << bWorkMemoryStorageObjectValid;
+	if(bWorkMemoryStorageObjectValid)
+	{
+		WorkMemoryStorageObject->Serialize(Ar);
+	}
+	bool bLiteralMemoryStorageObjectValid = LiteralMemoryStorageObject != nullptr;
+	Ar << bLiteralMemoryStorageObjectValid;
+	if(bLiteralMemoryStorageObjectValid)
+	{
+		LiteralMemoryStorageObject->Serialize(Ar);
+	}
 	
 #endif
 	
@@ -194,21 +204,32 @@ void URigVM::Load(FArchive& Ar)
 
 	if(WorkMemoryStorageObject == nullptr)
 	{
-		UClass* StorageClass = URigVMMemoryStorage::GetStorageClass(GetOutermost(), ERigVMMemoryType::Work);
-		WorkMemoryStorageObject = NewObject<URigVMMemoryStorage>(this, StorageClass, NAME_None, RF_Public | RF_Transactional);
+		UClass* StorageClass = URigVMMemoryStorageGeneratorClass::GetStorageClass(GetOutermost(), ERigVMMemoryType::Work);
+		WorkMemoryStorageObject = NewObject<URigVMMemoryStorage>(this, StorageClass, NAME_None, RF_Public);
 	}
 	if(LiteralMemoryStorageObject == nullptr)
 	{
-		UClass* StorageClass = URigVMMemoryStorage::GetStorageClass(GetOutermost(), ERigVMMemoryType::Literal);
-		LiteralMemoryStorageObject = NewObject<URigVMMemoryStorage>(this, StorageClass, NAME_None, RF_Public | RF_Transactional);
+		UClass* StorageClass = URigVMMemoryStorageGeneratorClass::GetStorageClass(GetOutermost(), ERigVMMemoryType::Literal);
+		LiteralMemoryStorageObject = NewObject<URigVMMemoryStorage>(this, StorageClass, NAME_None, RF_Public);
 	}
 	if(DebugMemoryStorageObject == nullptr)
 	{
-		DebugMemoryStorageObject = NewObject<URigVMMemoryStorage>(this, NAME_None, RF_Public | RF_Transactional);
+		DebugMemoryStorageObject = NewObject<URigVMMemoryStorage>(this, NAME_None, RF_Public);
 	}
 
-	WorkMemoryStorageObject->Serialize(Ar);
-	LiteralMemoryStorageObject->Serialize(Ar);
+	bool bWorkMemoryStorageObjectValid = false;
+	Ar << bWorkMemoryStorageObjectValid;
+	if(bWorkMemoryStorageObjectValid)
+	{
+		WorkMemoryStorageObject->Serialize(Ar);
+	}
+	
+	bool bLiteralMemoryStorageObjectValid = false;
+	Ar << bLiteralMemoryStorageObjectValid;
+	if(bLiteralMemoryStorageObjectValid)
+	{
+		LiteralMemoryStorageObject->Serialize(Ar);
+	}
 
 	Ar << FunctionNamesStorage;
 	Ar << ByteCodeStorage;
@@ -256,14 +277,13 @@ bool URigVM::ValidateAllOperandsDuringLoad()
 #else
 
 	TArray<URigVMMemoryStorage*> LocalMemory = { GetWorkMemory(), GetLiteralMemory(), GetDebugMemory() };
-	TArray<FRigVMPropertyPath>& LocalPropertyPaths = PropertyPaths;
 	
 #endif
 	
 #if UE_RIGVM_UCLASS_BASED_STORAGE_DISABLED
 	auto CheckOperandValidity = [LocalMemory, &bAllOperandsValid](const FRigVMOperand& InOperand) -> bool
 #else
-	auto CheckOperandValidity = [LocalMemory, &bAllOperandsValid, &LocalPropertyPaths](const FRigVMOperand& InOperand) -> bool
+	auto CheckOperandValidity = [LocalMemory, &bAllOperandsValid](const FRigVMOperand& InOperand) -> bool
 #endif
 	{
 		if(InOperand.GetContainerIndex() < 0 || InOperand.GetContainerIndex() >= (int32)ERigVMMemoryType::Invalid)
@@ -297,7 +317,7 @@ bool URigVM::ValidateAllOperandsDuringLoad()
 #if UE_RIGVM_UCLASS_BASED_STORAGE_DISABLED
 			if(!MemoryForOperand->RegisterOffsets.IsValidIndex(InOperand.GetRegisterOffset()))
 #else
-			if(!LocalPropertyPaths.IsValidIndex(InOperand.GetRegisterOffset()))
+			if(!MemoryForOperand->GetPropertyPaths().IsValidIndex(InOperand.GetRegisterOffset()))
 #endif
 			{
 				bAllOperandsValid = false;
@@ -467,10 +487,6 @@ void URigVM::Reset()
 	FunctionsPtr = &FunctionsStorage;
 	ByteCodePtr = &ByteCodeStorage;
 
-#if !UE_RIGVM_UCLASS_BASED_STORAGE_DISABLED
-	PropertyPaths.Reset();
-#endif
-
 	InvalidateCachedMemory();
 	
 	OperandToDebugRegisters.Reset();
@@ -492,10 +508,6 @@ void URigVM::Empty()
 	DeferredVMToCopy = nullptr;
 	ExternalVariables.Empty();
 
-#if !UE_RIGVM_UCLASS_BASED_STORAGE_DISABLED
-	PropertyPaths.Empty();
-#endif
-
 	InvalidateCachedMemory();
 
 	CachedMemory.Empty();
@@ -591,7 +603,6 @@ void URigVM::CopyFrom(URigVM* InVM, bool bDeferCopy, bool bReferenceLiteralMemor
 	CopyMemoryStorage(WorkMemoryStorageObject, InVM->WorkMemoryStorageObject, this);
 	CopyMemoryStorage(LiteralMemoryStorageObject, InVM->LiteralMemoryStorageObject, this);
 	CopyMemoryStorage(DebugMemoryStorageObject, InVM->DebugMemoryStorageObject, this);
-	PropertyPaths = InVM->PropertyPaths;
 	
 #endif
 
@@ -673,6 +684,7 @@ URigVMMemoryStorage* URigVM::GetMemoryByType(ERigVMMemoryType InMemoryType) cons
 			return GetLiteralMemory();
 		}
 		case ERigVMMemoryType::Work:
+		case ERigVMMemoryType::External:
 		{
 			return GetWorkMemory();
 		}
@@ -2541,7 +2553,7 @@ FString URigVM::GetOperandLabel(const FRigVMOperand& InOperand, TFunction<FStrin
 	const FString RegisterName = Memory->GetProperties()[InOperand.GetRegisterIndex()]->GetName();
 	const FString RegisterOffsetName =
 		InOperand.GetRegisterOffset() != INDEX_NONE ?
-		PropertyPaths[InOperand.GetRegisterOffset()].ToString() :
+		Memory->GetPropertyPaths()[InOperand.GetRegisterOffset()].ToString() :
 		FString();
 	
 	FString OperandLabel = RegisterName;
@@ -2623,11 +2635,13 @@ void URigVM::CacheSingleMemoryHandle(const FRigVMOperand& InArg, bool bForExecut
 	
 #else
 
+	URigVMMemoryStorage* Memory = GetMemoryByType(InArg.GetMemoryType());
+
 	const FRigVMPropertyPath* PropertyPath = nullptr;
 	if(InArg.GetRegisterOffset())
 	{
-		check(PropertyPaths.IsValidIndex(InArg.GetRegisterOffset()));
-		PropertyPath = &PropertyPaths[InArg.GetRegisterOffset()];
+		check(Memory->GetPropertyPaths().IsValidIndex(InArg.GetRegisterOffset()));
+		PropertyPath = &Memory->GetPropertyPaths()[InArg.GetRegisterOffset()];
 	}
 
 	if (InArg.GetMemoryType() == ERigVMMemoryType::External)
@@ -2642,7 +2656,6 @@ void URigVM::CacheSingleMemoryHandle(const FRigVMOperand& InArg, bool bForExecut
 		return;
 	}
 
-	URigVMMemoryStorage* Memory = GetMemoryByType(InArg.GetMemoryType());
 	uint8* Data = Memory->GetData<uint8>(InArg.GetRegisterIndex());
 	const FProperty* Property = Memory->GetProperties()[InArg.GetRegisterIndex()];
 	const FRigVMMemoryHandle Handle(Data, Property, PropertyPath);
diff --git a/Engine/Source/Runtime/RigVM/Private/RigVMCore/RigVMMemoryDeprecated.cpp b/Engine/Source/Runtime/RigVM/Private/RigVMCore/RigVMMemoryDeprecated.cpp
--- a/Engine/Source/Runtime/RigVM/Private/RigVMCore/RigVMMemoryDeprecated.cpp
+++ b/Engine/Source/Runtime/RigVM/Private/RigVMCore/RigVMMemoryDeprecated.cpp
@@ -537,8 +537,6 @@ FRigVMMemoryContainer& FRigVMMemoryContainer::operator= (const FRigVMMemoryConta
 
 #endif
 
-#if UE_RIGVM_UCLASS_BASED_STORAGE_DISABLED
-
 void FRigVMMemoryContainer::Serialize(FArchive& Ar)
 {
 	Ar.UsingCustomVersion(FAnimObjectVersion::GUID);
@@ -562,19 +560,10 @@ void FRigVMMemoryContainer::Serialize(FArchive& Ar)
 	}
 }
 
-#else
-
-void FRigVMMemoryContainer::Serialize(FArchive& Ar)
+void FRigVMMemoryContainer::Save(FArchive& Ar)
 {
-	return;
-}
-
-#endif
-
 #if UE_RIGVM_UCLASS_BASED_STORAGE_DISABLED
 
-void FRigVMMemoryContainer::Save(FArchive& Ar)
-{
 	for (FRigVMRegister& Register : Registers)
 	{
 		if (Register.IsNestedDynamic())
@@ -761,21 +750,23 @@ void FRigVMMemoryContainer::Save(FArchive& Ar)
 			}
 		}
 	}
+#endif
 }
 
 void FRigVMMemoryContainer::Load(FArchive& Ar)
 {
+#if UE_RIGVM_UCLASS_BASED_STORAGE_DISABLED
 	for (int32 RegisterIndex = 0; RegisterIndex < Registers.Num(); RegisterIndex++)
 	{
 		Destroy(RegisterIndex);
 	}
+#endif
 	
 	Ar << bUseNameMap;
 	Ar << MemoryType;
 	Ar << Registers;
 	Ar << RegisterOffsets;
 
-
 #if DEBUG_RIGVMMEMORY
 	UE_LOG_RIGVMMEMORY(TEXT("%d Memory - Begin Loading..."), (int32)GetMemoryType());
 #endif
@@ -786,6 +777,8 @@ void FRigVMMemoryContainer::Load(FArchive& Ar)
 	TArray<FString> ScriptStructPaths;
 	Ar << ScriptStructPaths;
 
+#if UE_RIGVM_UCLASS_BASED_STORAGE_DISABLED
+
 	for (const FString& ScriptStructPath : ScriptStructPaths)
 	{
 		UScriptStruct* ScriptStruct = FindObject<UScriptStruct>(nullptr, *ScriptStructPath);
@@ -802,13 +795,18 @@ void FRigVMMemoryContainer::Load(FArchive& Ar)
 		ScriptStructs.Add(ScriptStruct);
 	}
 
+#endif
+
 	uint64 TotalBytes = 0;
 	Ar << TotalBytes;
 
+#if UE_RIGVM_UCLASS_BASED_STORAGE_DISABLED
+
 	Data.Empty();
 
 	if (!bEncounteredErrorDuringLoad)
 	{
+
 		// during load we'll recreate the memory for all registers.
 		// the size for structs might have changed, so we need to reallocate.
 		for (int32 RegisterIndex = 0; RegisterIndex < Registers.Num(); RegisterIndex++)
@@ -842,6 +840,7 @@ void FRigVMMemoryContainer::Load(FArchive& Ar)
 			}
 		}
 		UpdateRegisters();
+		
 
 		for (int32 RegisterOffsetIndex = 0; RegisterOffsetIndex < RegisterOffsets.Num(); RegisterOffsetIndex++)
 		{
@@ -873,10 +872,14 @@ void FRigVMMemoryContainer::Load(FArchive& Ar)
 		}
 	}
 
+#endif
+
 	for (int32 RegisterIndex = 0; RegisterIndex < Registers.Num(); RegisterIndex++)
 	{
 		FRigVMRegister& Register = Registers[RegisterIndex];
 
+#if UE_RIGVM_UCLASS_BASED_STORAGE_DISABLED
+
 		if (Register.ElementCount == 0 && !Register.IsDynamic())
 		{
 			continue;
@@ -889,6 +892,8 @@ void FRigVMMemoryContainer::Load(FArchive& Ar)
 		{
 			check(!Register.IsDynamic());
 		}
+
+#endif
 		
 		if (!Register.IsDynamic())
 		{
@@ -898,39 +903,57 @@ void FRigVMMemoryContainer::Load(FArchive& Ar)
 				{
 					FRigVMByteArray View;
 					Ar << View;
+
+#if UE_RIGVM_UCLASS_BASED_STORAGE_DISABLED
+
 					if (!bEncounteredErrorDuringLoad)
 					{
 						ensure(View.Num() <= Register.GetAllocatedBytes());
 						FMemory::Memcpy(&Data[Register.GetWorkByteIndex()], View.GetData(), View.Num());
 					}
+
+#endif
 					break;
 				}
 				case ERigVMRegisterType::Name:
 				{
 					TArray<FName> View;
 					Ar << View;
+
+#if UE_RIGVM_UCLASS_BASED_STORAGE_DISABLED
+
 					if (!bEncounteredErrorDuringLoad)
 					{
 						ensure(View.Num() == Register.GetTotalElementCount());
 						RigVMCopy<FName>(&Data[Register.GetWorkByteIndex()], View.GetData(), View.Num());
 					}
+
+#endif
 					break;
 				}
 				case ERigVMRegisterType::String:
 				{
 					TArray<FString> View;
 					Ar << View;
+
+#if UE_RIGVM_UCLASS_BASED_STORAGE_DISABLED
+						
 					if (!bEncounteredErrorDuringLoad)
 					{
 						ensure(View.Num() == Register.GetTotalElementCount());
 						RigVMCopy<FString>(&Data[Register.GetWorkByteIndex()], View.GetData(), View.Num());
 					}
+						
+#endif
 					break;
 				}
 				case ERigVMRegisterType::Struct:
 				{
 					TArray<FString> View;
 					Ar << View;
+
+#if UE_RIGVM_UCLASS_BASED_STORAGE_DISABLED
+
 					if (!bEncounteredErrorDuringLoad)
 					{
 						ensure(View.Num() == Register.GetTotalElementCount());
@@ -954,12 +977,17 @@ void FRigVMMemoryContainer::Load(FArchive& Ar)
 							}
 						}
 					}
+
+#endif
 					break;
 				}
 			}
 		}
 		else if (!Register.IsNestedDynamic())
 		{
+
+#if UE_RIGVM_UCLASS_BASED_STORAGE_DISABLED
+			
 			FRigVMByteArray* ArrayStorage = nullptr;
 
 			if (!bEncounteredErrorDuringLoad)
@@ -967,6 +995,8 @@ void FRigVMMemoryContainer::Load(FArchive& Ar)
 				ArrayStorage = (FRigVMByteArray*)&Data[Register.GetWorkByteIndex()];
 			}
 			
+#endif
+			
 			// notice that all dynamic registers' memory is not really serialized
 			// serialize empty arrays here just to avoid having to increment serialization version
 			// the loaded data is not getting used at all, it will be recomputed at execution time
@@ -1000,6 +1030,9 @@ void FRigVMMemoryContainer::Load(FArchive& Ar)
 		}
 		else
 		{
+
+#if UE_RIGVM_UCLASS_BASED_STORAGE_DISABLED
+			
 			FRigVMNestedByteArray* NestedArrayStorage = nullptr;
 			if (!bEncounteredErrorDuringLoad)
 			{
@@ -1007,15 +1040,21 @@ void FRigVMMemoryContainer::Load(FArchive& Ar)
 				NestedArrayStorage->Reset();
 				NestedArrayStorage->SetNumZeroed(Register.SliceCount);
 			}
+			
+#endif
 
 			for (int32 SliceIndex = 0; SliceIndex < Register.SliceCount; SliceIndex++)
 			{
+#if UE_RIGVM_UCLASS_BASED_STORAGE_DISABLED
+				
 				FRigVMByteArray* ArrayStorage = nullptr;
 				if (NestedArrayStorage)
 				{
 					ArrayStorage = &(*NestedArrayStorage)[SliceIndex];
 				}
 
+#endif
+
 				// notice that all dynamic registers' memory is not really serialized
 				// serialize empty arrays here just to avoid having to increment serialization version
 				// the loaded data is not getting used at all, it will be recomputed at execution time
@@ -1050,6 +1089,9 @@ void FRigVMMemoryContainer::Load(FArchive& Ar)
 		}
 	}
 
+
+#if UE_RIGVM_UCLASS_BASED_STORAGE_DISABLED
+
 	if (bEncounteredErrorDuringLoad)
 	{
 #if DEBUG_RIGVMMEMORY
@@ -1065,9 +1107,12 @@ void FRigVMMemoryContainer::Load(FArchive& Ar)
 #if DEBUG_RIGVMMEMORY
 	UE_LOG_RIGVMMEMORY(TEXT("%d Memory - Finished Loading."), (int32)GetMemoryType());
 #endif
-	
+
+#endif
 }
 
+#if UE_RIGVM_UCLASS_BASED_STORAGE_DISABLED
+
 void FRigVMMemoryContainer::Reset()
 {
 	if (Data.Num() > 0)
diff --git a/Engine/Source/Runtime/RigVM/Private/RigVMCore/RigVMMemoryStorage.cpp b/Engine/Source/Runtime/RigVM/Private/RigVMCore/RigVMMemoryStorage.cpp
--- a/Engine/Source/Runtime/RigVM/Private/RigVMCore/RigVMMemoryStorage.cpp
+++ b/Engine/Source/Runtime/RigVM/Private/RigVMCore/RigVMMemoryStorage.cpp
@@ -5,12 +5,14 @@
 
 ////////////////////////////////////////////////////////////////////////////////////////////////////////////////////////
 
-const FString URigVMMemoryStorage::FPropertyDescription::ArrayPrefix = TEXT("TArray<");
-const FString URigVMMemoryStorage::FPropertyDescription::MapPrefix = TEXT("TMap<");
-const FString URigVMMemoryStorage::FPropertyDescription::SetPrefix = TEXT("TSet<");
-const FString URigVMMemoryStorage::FPropertyDescription::ContainerSuffix = TEXT(">");
+const FString FRigVMPropertyDescription::ArrayPrefix = TEXT("TArray<");
+const FString FRigVMPropertyDescription::MapPrefix = TEXT("TMap<");
+const FString FRigVMPropertyDescription::SetPrefix = TEXT("TSet<");
+const FString FRigVMPropertyDescription::ContainerSuffix = TEXT(">");
 
-URigVMMemoryStorage::FPropertyDescription::FPropertyDescription(const FProperty* InProperty, const FString& InDefaultValue, const FName& InName)
+////////////////////////////////////////////////////////////////////////////////////////////////////////////////////////
+
+FRigVMPropertyDescription::FRigVMPropertyDescription(const FProperty* InProperty, const FString& InDefaultValue, const FName& InName)
 	: Name(InName)
 	, Property(InProperty)
 	, CPPType()
@@ -45,7 +47,7 @@
 	while (ChildProperty);
 }
 
-URigVMMemoryStorage::FPropertyDescription::FPropertyDescription(const FName& InName, const FString& InCPPType, UObject* InCPPTypeObject, const FString& InDefaultValue)
+FRigVMPropertyDescription::FRigVMPropertyDescription(const FName& InName, const FString& InCPPType, UObject* InCPPTypeObject, const FString& InDefaultValue)
 	: Name(InName)
 	, Property(nullptr)
 	, CPPType(InCPPType)
@@ -82,7 +84,7 @@
 	while(!BaseCPPType.IsEmpty());
 }
 
-FName URigVMMemoryStorage::FPropertyDescription::SanitizeName(const FName& InName)
+FName FRigVMPropertyDescription::SanitizeName(const FName& InName)
 {
 	FString NameString = InName.ToString();
 
@@ -109,12 +111,12 @@
 	return InName;
 }
 
-void URigVMMemoryStorage::FPropertyDescription::SanitizeName()
+void FRigVMPropertyDescription::SanitizeName()
 {
 	Name = SanitizeName(Name);
 }
 
-FString URigVMMemoryStorage::FPropertyDescription::GetBaseCPPType() const
+FString FRigVMPropertyDescription::GetBaseCPPType() const
 {
 	FString BaseCPPType = CPPType;
 
@@ -154,21 +156,99 @@
 ////////////////////////////////////////////////////////////////////////////////////////////////////////////////////////
 
 FCriticalSection RigVMMemoryStorageClassCriticalSection;
-TMap<FString, URigVMMemoryStorage::FClassInfo> URigVMMemoryStorage::PackageToInfo;
+TMap<FString, URigVMMemoryStorageGeneratorClass::FClassInfo> URigVMMemoryStorageGeneratorClass::PackageToInfo;
 
-UClass* URigVMMemoryStorage::GetStorageClass(UObject* InOuter, ERigVMMemoryType InMemoryType)
+void URigVMMemoryStorageGeneratorClass::PurgeClass(bool bRecompilingOnLoad)
+{
+	Super::PurgeClass(bRecompilingOnLoad);
+	LinkedProperties.Reset();
+	PropertyPaths.Reset();
+	PropertyPathDescriptions.Reset();
+}
+
+void URigVMMemoryStorageGeneratorClass::Link(FArchive& Ar, bool bRelinkExistingProperties)
+{
+	Super::Link(Ar, bRelinkExistingProperties);
+
+	// Force assembly of the reference token stream so that we can be properly handled by the
+	// garbage collector.
+	AssembleReferenceTokenStream(/*bForce=*/true);
+
+	// Setup the LinkedProperties
+	LinkedProperties.Reset();
+	const FProperty* Property = CastField<FProperty>(ChildProperties);
+	while(Property)
+	{
+		LinkedProperties.Add(Property);
+		Property = CastField<FProperty>(Property->Next);
+	}
+
+	// Update the property paths based on the descriptions
+	PropertyPaths.SetNumZeroed(PropertyPathDescriptions.Num());
+	for(int32 PropertyPathIndex = 0; PropertyPathIndex < PropertyPaths.Num(); PropertyPathIndex++)
+	{
+		PropertyPaths[PropertyPathIndex] = FRigVMPropertyPath();
+
+		const int32 PropertyIndex = PropertyPathDescriptions[PropertyPathIndex].PropertyIndex;
+		if(LinkedProperties.IsValidIndex(PropertyIndex))
+		{
+			PropertyPaths[PropertyPathIndex] = FRigVMPropertyPath(
+				LinkedProperties[PropertyIndex],
+				PropertyPathDescriptions[PropertyPathIndex].SegmentPath);
+		}
+	}
+}
+
+void URigVMMemoryStorageGeneratorClass::Serialize(FArchive& Ar)
+{
+	Super::Serialize(Ar);
+
+	if(Ar.IsLoading() || Ar.IsSaving())
+	{
+		Ar << PropertyPathDescriptions;
+		Ar << MemoryType;
+	}
+
+	if(Ar.IsLoading())
+	{
+		UPackage* Package = GetOutermost();
+		
+		FScopeLock Lock(&RigVMMemoryStorageClassCriticalSection);
+
+		FClassInfo& ClassInfo = PackageToInfo.FindOrAdd(Package->GetPathName());
+		URigVMMemoryStorageGeneratorClass** ClassPtr = ClassInfo.GetClassPtr(MemoryType);
+		if(ClassPtr)
+		{
+			if(*ClassPtr != this)
+			{
+				if(*ClassPtr)
+				{
+					(*ClassPtr)->RemoveFromRoot();
+					(*ClassPtr)->MarkPendingKill();
+				}
+				*ClassPtr = this;
+			}
+		}
+	}
+}
+
+URigVMMemoryStorageGeneratorClass* URigVMMemoryStorageGeneratorClass::GetStorageClass(UObject* InOuter, ERigVMMemoryType InMemoryType)
 {
 	check(InOuter);
 	UPackage* Package = InOuter->GetOutermost(); 
 
 	FScopeLock Lock(&RigVMMemoryStorageClassCriticalSection);
 
-	FClassInfo& ClassInfo = PackageToInfo.FindOrAdd(InOuter->GetPathName());
-	UClass** ClassPtr = ClassInfo.GetClassPtr(InMemoryType);
+	FClassInfo& ClassInfo = PackageToInfo.FindOrAdd(Package->GetPathName());
+	URigVMMemoryStorageGeneratorClass** ClassPtr = ClassInfo.GetClassPtr(InMemoryType);
 	return *ClassPtr;
 }
 
-UClass* URigVMMemoryStorage::CreateStorageClass(UObject* InOuter, ERigVMMemoryType InMemoryType, const TArray<FPropertyDescription>& InProperties)
+URigVMMemoryStorageGeneratorClass* URigVMMemoryStorageGeneratorClass::CreateStorageClass(
+	UObject* InOuter,
+	ERigVMMemoryType InMemoryType,
+	const TArray<FRigVMPropertyDescription>& InProperties,
+	const TArray<FRigVMPropertyPathDescription>& InPropertyPaths)
 {
 	check(InOuter);
 	UPackage* Package = InOuter->GetOutermost(); 
@@ -176,54 +256,51 @@
 
 	FScopeLock Lock(&RigVMMemoryStorageClassCriticalSection);
 
-	FClassInfo& ClassInfo = PackageToInfo.FindOrAdd(InOuter->GetPathName());
-	UClass** ClassPtr = ClassInfo.GetClassPtr(InMemoryType);
+	FClassInfo& ClassInfo = PackageToInfo.FindOrAdd(Package->GetPathName());
+	URigVMMemoryStorageGeneratorClass** ClassPtr = ClassInfo.GetClassPtr(InMemoryType);
 	check(ClassPtr);
 	
 	if(*ClassPtr == nullptr)
 	{
-		(*ClassPtr) = NewObject<UClass>(
+		(*ClassPtr) = NewObject<URigVMMemoryStorageGeneratorClass>(
 			InOuter->GetOutermost(),
 			NAME_None,
-			RF_Public | RF_Transactional
+			RF_Standalone | RF_Public
 			);
 
 		(*ClassPtr)->AddToRoot();
 	}
 
-	UClass* Class = *ClassPtr;
+	URigVMMemoryStorageGeneratorClass* Class = *ClassPtr;
 
 	Class->PurgeClass(false);
 	Class->PropertyLink = SuperClass->PropertyLink;
 	Class->SetSuperStruct(SuperClass);
 	Class->ClassWithin = UObject::StaticClass();
+	Class->MemoryType = InMemoryType;
 
 	// Generate properties
 	FField** LinkToProperty = &Class->ChildProperties;
 
-	TArray<const FProperty*> CachedProperties;
-	for(const FPropertyDescription& PropertyDescription : InProperties)
+	for(const FRigVMPropertyDescription& PropertyDescription : InProperties)
 	{
 		FProperty* CachedProperty = AddProperty(Class, PropertyDescription, false, false, LinkToProperty);
 		check(CachedProperty);
-		CachedProperties.Add(CachedProperty);
 	}
 
+	// Store the property path descriptions
+	Class->PropertyPathDescriptions = InPropertyPaths;
+
 	// Update the class
 	Class->Bind();
 	Class->StaticLink(true);
-
-	// Similar to FConfigPropertyHelperDetails::CustomizeDetails, this is required for GC to work properly
-	Class->AssembleReferenceTokenStream();
-
-	check(CachedProperties.Num() == InProperties.Num());
 	
 	// Create default object
 	URigVMMemoryStorage* CDO = Cast<URigVMMemoryStorage>(Class->GetDefaultObject(true));
-	CDO->CachedProperties = CachedProperties;
 
 	// and store default values.
-	for(int32 PropertyIndex = 0; PropertyIndex < CachedProperties.Num(); PropertyIndex++)
+	const TArray<const FProperty*>& LinkedProperties = CDO->GetProperties();
+	for(int32 PropertyIndex = 0; PropertyIndex < LinkedProperties.Num(); PropertyIndex++)
 	{
 		const FString& DefaultValue = InProperties[PropertyIndex].DefaultValue;
 		if(DefaultValue.IsEmpty())
@@ -231,7 +308,7 @@
 			continue;
 		}
 
-		const FProperty* Property = CachedProperties[PropertyIndex];
+		const FProperty* Property = LinkedProperties[PropertyIndex];
 		uint8* ValuePtr = Property->ContainerPtrToValuePtr<uint8>(CDO);
 
 		Property->ImportText(*DefaultValue, ValuePtr, EPropertyPortFlags::PPF_None, nullptr);
@@ -240,18 +317,174 @@
 	return Class;
 }
 
-URigVMMemoryStorage* URigVMMemoryStorage::CreateStorage(UObject* InOuter, ERigVMMemoryType InMemoryType)
+FProperty* URigVMMemoryStorageGeneratorClass::AddProperty(URigVMMemoryStorageGeneratorClass* InClass, const FRigVMPropertyDescription& InProperty, bool bPurge, bool bLink, FField** LinkToProperty)
 {
-	UClass* Class = GetStorageClass(InOuter, InMemoryType);
+	UClass *SuperClass = URigVMMemoryStorage::StaticClass();
+	
+	check(InClass);
+	check(InClass->GetSuperClass() == SuperClass);
 
-	if(InMemoryType == ERigVMMemoryType::Literal)
+	if(bPurge)
 	{
-		return Cast<URigVMMemoryStorage>(Class->GetDefaultObject(true));
+		InClass->PurgeClass(false);
+		InClass->PropertyLink = SuperClass->PropertyLink;
+		InClass->SetSuperStruct(SuperClass);
+		InClass->ClassWithin = UObject::StaticClass();
+	}
+	
+	if(LinkToProperty == nullptr)
+	{
+		LinkToProperty = &InClass->ChildProperties;
+	}
+	while (*LinkToProperty != nullptr)
+	{
+		LinkToProperty = &(*LinkToProperty)->Next;
 	}
 
-	return NewObject<URigVMMemoryStorage>(InOuter, Class, NAME_None, RF_Public | RF_Transactional);
+	FProperty* Result = nullptr;
+	if(InProperty.Property)
+	{
+		FProperty* NewProperty = CastFieldChecked<FProperty>(FField::Duplicate(InProperty.Property, InClass, InProperty.Name));
+		check(NewProperty);
+
+		Result = NewProperty;
+		*LinkToProperty = NewProperty;
+	}
+	else
+	{
+		FFieldVariant PropertyOwner = InClass;
+		// FProperty** KeyPropertyPtr = nullptr;
+		FProperty** ValuePropertyPtr = &Result;
+
+		for(EPinContainerType Container : InProperty.Containers)
+		{
+			switch(Container)
+			{
+				case EPinContainerType::Array:
+				{
+					FArrayProperty* ArrayProperty = new FArrayProperty(PropertyOwner, InProperty.Name, RF_Public);
+					*ValuePropertyPtr = ArrayProperty;
+					ValuePropertyPtr = &ArrayProperty->Inner;
+					PropertyOwner = ArrayProperty;
+					break;
+				}
+				case EPinContainerType::Map:
+				{
+					checkNoEntry(); // this is not implemented yet
+					FMapProperty* MapProperty = new FMapProperty(PropertyOwner, InProperty.Name, RF_Public);
+					*ValuePropertyPtr = MapProperty;
+					// KeyPropertyPtr = &MapProperty->KeyProp;
+					ValuePropertyPtr = &MapProperty->ValueProp;
+					PropertyOwner = MapProperty;
+					break;
+				}
+				case EPinContainerType::Set:
+				{
+					FSetProperty* SetProperty = new FSetProperty(PropertyOwner, InProperty.Name, RF_Public);
+					*ValuePropertyPtr = SetProperty;
+					ValuePropertyPtr = &SetProperty->ElementProp;
+					PropertyOwner = SetProperty;
+					break;
+				}
+				case EPinContainerType::None:
+				default:
+				{
+					break;
+				}
+			}
+		}
+
+		if(InProperty.CPPTypeObject != nullptr)
+		{
+			if(UEnum* Enum = Cast<UEnum>(InProperty.CPPTypeObject))
+			{
+				FByteProperty* EnumProperty = new FByteProperty(PropertyOwner, InProperty.Name, RF_Public);
+				EnumProperty->Enum = Enum;
+				(*ValuePropertyPtr) = EnumProperty;
+			}
+			else if(UScriptStruct* ScriptStruct = Cast<UScriptStruct>(InProperty.CPPTypeObject))
+			{
+				FStructProperty* StructProperty = new FStructProperty(PropertyOwner, InProperty.Name, RF_Public);
+				StructProperty->Struct = ScriptStruct;
+				(*ValuePropertyPtr) = StructProperty;
+			}
+			else if(UClass* PropertyClass = Cast<UClass>(InProperty.CPPTypeObject))
+			{
+				FObjectProperty* ObjectProperty = new FObjectProperty(PropertyOwner, InProperty.Name, RF_Public);
+				ObjectProperty->SetPropertyClass(PropertyClass);
+				(*ValuePropertyPtr) = ObjectProperty;
+			}
+			else
+			{
+				checkNoEntry();
+			}
+		}
+		else // take care of default types...
+		{
+			static FString BoolString = TEXT("bool");
+			static FString Int32String = TEXT("int32");
+			static FString IntString = TEXT("int");
+			static FString FloatString = TEXT("float");
+			static FString DoubleString = TEXT("double");
+			static FString StringString = TEXT("FString");
+			static FString NameString = TEXT("FName");
+
+			const FString BaseCPPType = InProperty.GetBaseCPPType();
+			if(BaseCPPType.Equals(BoolString, ESearchCase::IgnoreCase))
+			{
+				(*ValuePropertyPtr) = new FBoolProperty(PropertyOwner, InProperty.Name, RF_Public);;
+			}
+			else if(BaseCPPType.Equals(Int32String, ESearchCase::IgnoreCase) ||
+				BaseCPPType.Equals(IntString, ESearchCase::IgnoreCase))
+			{
+				(*ValuePropertyPtr) = new FIntProperty(PropertyOwner, InProperty.Name, RF_Public);;
+			}
+			else if(BaseCPPType.Equals(FloatString, ESearchCase::IgnoreCase))
+			{
+				(*ValuePropertyPtr) = new FFloatProperty(PropertyOwner, InProperty.Name, RF_Public);;
+			}
+			else if(BaseCPPType.Equals(DoubleString, ESearchCase::IgnoreCase))
+			{
+				(*ValuePropertyPtr) = new FDoubleProperty(PropertyOwner, InProperty.Name, RF_Public);;
+			}
+			else if(BaseCPPType.Equals(StringString, ESearchCase::IgnoreCase))
+			{
+				(*ValuePropertyPtr) = new FStrProperty(PropertyOwner, InProperty.Name, RF_Public);;
+			}
+			else if(BaseCPPType.Equals(NameString, ESearchCase::IgnoreCase))
+			{
+				(*ValuePropertyPtr) = new FNameProperty(PropertyOwner, InProperty.Name, RF_Public);;
+			}
+			else
+			{
+				checkNoEntry();
+			}
+		}
+
+		Result->SetPropertyFlags(CPF_Edit | CPF_NonTransactional);
+		(*LinkToProperty) = Result;
+	}
+
+	if(bLink)
+	{
+		// Update the class
+		InClass->Bind();
+		InClass->StaticLink(true);
+
+		URigVMMemoryStorage* CDO = Cast<URigVMMemoryStorage>(InClass->GetDefaultObject(true));
+		const FString& DefaultValue = InProperty.DefaultValue;
+		if(!DefaultValue.IsEmpty())
+		{
+			uint8* ValuePtr = Result->ContainerPtrToValuePtr<uint8>(CDO);
+			Result->ImportText(*DefaultValue, ValuePtr, EPropertyPortFlags::PPF_None, nullptr);
+		}		
+	};
+	
+	return Result;
 }
 
+////////////////////////////////////////////////////////////////////////////////////////////////////////////////////////
+
 FString URigVMMemoryStorage::GetDataAsString(int32 InPropertyIndex)
 {
 	check(IsValidIndex(InPropertyIndex));
@@ -420,187 +653,16 @@
 
 #endif
 
-FProperty* URigVMMemoryStorage::AddProperty(UClass* InClass, const FPropertyDescription& InProperty, bool bPurge, bool bLink, FField** LinkToProperty)
-{
-	UClass *SuperClass = URigVMMemoryStorage::StaticClass();
-	
-	check(InClass);
-	check(InClass->GetSuperClass() == SuperClass);
-
-	if(bPurge)
-	{
-		InClass->PurgeClass(false);
-		InClass->PropertyLink = SuperClass->PropertyLink;
-		InClass->SetSuperStruct(SuperClass);
-		InClass->ClassWithin = UObject::StaticClass();
-	}
-	
-	if(LinkToProperty == nullptr)
-	{
-		LinkToProperty = &InClass->ChildProperties;
-	}
-	while (*LinkToProperty != nullptr)
-	{
-		LinkToProperty = &(*LinkToProperty)->Next;
-	}
-
-	FProperty* Result = nullptr;
-	if(InProperty.Property)
-	{
-		FProperty* NewProperty = CastFieldChecked<FProperty>(FField::Duplicate(InProperty.Property, InClass, InProperty.Name));
-		check(NewProperty);
-
-		Result = NewProperty;
-		*LinkToProperty = NewProperty;
-	}
-	else
-	{
-		FFieldVariant PropertyOwner = InClass;
-		// FProperty** KeyPropertyPtr = nullptr;
-		FProperty** ValuePropertyPtr = &Result;
-
-		for(EPinContainerType Container : InProperty.Containers)
-		{
-			switch(Container)
-			{
-				case EPinContainerType::Array:
-				{
-					FArrayProperty* ArrayProperty = new FArrayProperty(PropertyOwner, InProperty.Name, RF_Public);
-					*ValuePropertyPtr = ArrayProperty;
-					ValuePropertyPtr = &ArrayProperty->Inner;
-					PropertyOwner = ArrayProperty;
-					break;
-				}
-				case EPinContainerType::Map:
-				{
-					checkNoEntry(); // this is not implemented yet
-					FMapProperty* MapProperty = new FMapProperty(PropertyOwner, InProperty.Name, RF_Public);
-					*ValuePropertyPtr = MapProperty;
-					// KeyPropertyPtr = &MapProperty->KeyProp;
-					ValuePropertyPtr = &MapProperty->ValueProp;
-					PropertyOwner = MapProperty;
-					break;
-				}
-				case EPinContainerType::Set:
-				{
-					FSetProperty* SetProperty = new FSetProperty(PropertyOwner, InProperty.Name, RF_Public);
-					*ValuePropertyPtr = SetProperty;
-					ValuePropertyPtr = &SetProperty->ElementProp;
-					PropertyOwner = SetProperty;
-					break;
-				}
-				case EPinContainerType::None:
-				default:
-				{
-					break;
-				}
-			}
-		}
-
-		if(InProperty.CPPTypeObject != nullptr)
-		{
-			if(UEnum* Enum = Cast<UEnum>(InProperty.CPPTypeObject))
-			{
-				FByteProperty* EnumProperty = new FByteProperty(PropertyOwner, InProperty.Name, RF_Public);
-				EnumProperty->Enum = Enum;
-				(*ValuePropertyPtr) = EnumProperty;
-			}
-			else if(UScriptStruct* ScriptStruct = Cast<UScriptStruct>(InProperty.CPPTypeObject))
-			{
-				FStructProperty* StructProperty = new FStructProperty(PropertyOwner, InProperty.Name, RF_Public);
-				StructProperty->Struct = ScriptStruct;
-				(*ValuePropertyPtr) = StructProperty;
-			}
-			else if(UClass* PropertyClass = Cast<UClass>(InProperty.CPPTypeObject))
-			{
-				FObjectProperty* ObjectProperty = new FObjectProperty(PropertyOwner, InProperty.Name, RF_Public);
-				ObjectProperty->SetPropertyClass(PropertyClass);
-				(*ValuePropertyPtr) = ObjectProperty;
-			}
-			else
-			{
-				checkNoEntry();
-			}
-		}
-		else // take care of default types...
-		{
-			static FString BoolString = TEXT("bool");
-			static FString Int32String = TEXT("int32");
-			static FString IntString = TEXT("int");
-			static FString FloatString = TEXT("float");
-			static FString DoubleString = TEXT("double");
-			static FString StringString = TEXT("FString");
-			static FString NameString = TEXT("FName");
-
-			const FString BaseCPPType = InProperty.GetBaseCPPType();
-			if(BaseCPPType.Equals(BoolString, ESearchCase::IgnoreCase))
-			{
-				(*ValuePropertyPtr) = new FBoolProperty(PropertyOwner, InProperty.Name, RF_Public);;
-			}
-			else if(BaseCPPType.Equals(Int32String, ESearchCase::IgnoreCase) ||
-				BaseCPPType.Equals(IntString, ESearchCase::IgnoreCase))
-			{
-				(*ValuePropertyPtr) = new FIntProperty(PropertyOwner, InProperty.Name, RF_Public);;
-			}
-			else if(BaseCPPType.Equals(FloatString, ESearchCase::IgnoreCase))
-			{
-				(*ValuePropertyPtr) = new FFloatProperty(PropertyOwner, InProperty.Name, RF_Public);;
-			}
-			else if(BaseCPPType.Equals(DoubleString, ESearchCase::IgnoreCase))
-			{
-				(*ValuePropertyPtr) = new FDoubleProperty(PropertyOwner, InProperty.Name, RF_Public);;
-			}
-			else if(BaseCPPType.Equals(StringString, ESearchCase::IgnoreCase))
-			{
-				(*ValuePropertyPtr) = new FStrProperty(PropertyOwner, InProperty.Name, RF_Public);;
-			}
-			else if(BaseCPPType.Equals(NameString, ESearchCase::IgnoreCase))
-			{
-				(*ValuePropertyPtr) = new FNameProperty(PropertyOwner, InProperty.Name, RF_Public);;
-			}
-			else
-			{
-				checkNoEntry();
-			}
-		}
-
-		Result->SetPropertyFlags(CPF_Edit);
-		(*LinkToProperty) = Result;
-	}
-
-	if(bLink)
-	{
-		// Update the class
-		InClass->Bind();
-		InClass->StaticLink(true);
-
-		// Similar to FConfigPropertyHelperDetails::CustomizeDetails, this is required for GC to work properly
-		InClass->AssembleReferenceTokenStream();
-
-		URigVMMemoryStorage* CDO = Cast<URigVMMemoryStorage>(InClass->GetDefaultObject(true));
-		CDO->CachedProperties.Add(Result);
-		
-		const FString& DefaultValue = InProperty.DefaultValue;
-		if(!DefaultValue.IsEmpty())
-		{
-			uint8* ValuePtr = Result->ContainerPtrToValuePtr<uint8>(CDO);
-			Result->ImportText(*DefaultValue, ValuePtr, EPropertyPortFlags::PPF_None, nullptr);
-		}		
-	};
-	
-	return Result;
-}
-
 const TArray<const FProperty*>& URigVMMemoryStorage::GetProperties() const
 {
-	if(!HasAnyFlags((RF_ClassDefaultObject)))
-	{
-		if(URigVMMemoryStorage* CDO = Cast<URigVMMemoryStorage>(GetClass()->GetDefaultObject()))
-		{
-			return CDO->GetProperties();
-		}
-	}
-	return CachedProperties;
+	const URigVMMemoryStorageGeneratorClass* Class = CastChecked<URigVMMemoryStorageGeneratorClass>(GetClass());
+	return Class->GetProperties();
 }
 
+const TArray<FRigVMPropertyPath>& URigVMMemoryStorage::GetPropertyPaths() const
+{
+	const URigVMMemoryStorageGeneratorClass* Class = CastChecked<URigVMMemoryStorageGeneratorClass>(GetClass());
+	return Class->GetPropertyPaths();
+}
+
 int32 URigVMMemoryStorage::GetPropertyIndex(const FProperty* InProperty) const
@@ -616,6 +678,6 @@
 
 const FProperty* URigVMMemoryStorage::FindPropertyByName(const FName& InName) const
 {
-	const FName SanitizedName = FPropertyDescription::SanitizeName(InName);
+	const FName SanitizedName = FRigVMPropertyDescription::SanitizeName(InName);
 	return GetClass()->FindPropertyByName(SanitizedName);
 }
diff --git a/Engine/Source/Runtime/RigVM/Public/RigVMCore/RigVM.h b/Engine/Source/Runtime/RigVM/Public/RigVMCore/RigVM.h
--- a/Engine/Source/Runtime/RigVM/Public/RigVMCore/RigVM.h
+++ b/Engine/Source/Runtime/RigVM/Public/RigVMCore/RigVM.h
@@ -250,8 +250,6 @@ public:
 	UPROPERTY()
 	URigVMMemoryStorage* DebugMemoryStorageObject;
 
-	TArray<FRigVMPropertyPath> PropertyPaths;
-
 	// The byte code of the VM
 	UPROPERTY()
 	FRigVMByteCode ByteCodeStorage;
diff --git a/Engine/Source/Runtime/RigVM/Public/RigVMCore/RigVMMemoryStorage.h b/Engine/Source/Runtime/RigVM/Public/RigVMCore/RigVMMemoryStorage.h
--- a/Engine/Source/Runtime/RigVM/Public/RigVMCore/RigVMMemoryStorage.h
+++ b/Engine/Source/Runtime/RigVM/Public/RigVMCore/RigVMMemoryStorage.h
@@ -82,57 +82,134 @@ private:
 
 #endif
 
+//////////////////////////////////////////////////////////////////////////////
+/// Property Management
+//////////////////////////////////////////////////////////////////////////////
+struct RIGVM_API FRigVMPropertyDescription
+{
+	FName Name;
+	const FProperty* Property;
+	FString CPPType;
+	UObject* CPPTypeObject;
+	TArray<EPinContainerType> Containers;
+	FString DefaultValue;
+
+	FRigVMPropertyDescription()
+		: Name(NAME_None)
+		, Property(nullptr)
+		, CPPType()
+		, CPPTypeObject(nullptr)
+		, Containers()
+		, DefaultValue()
+	{}
+	FRigVMPropertyDescription(const FProperty* InProperty, const FString& InDefaultValue, const FName& InName = NAME_None);
+	FRigVMPropertyDescription(const FName& InName, const FString& InCPPType, UObject* InCPPTypeObject, const FString& InDefaultValue);
+
+	static FName SanitizeName(const FName& InName);
+	void SanitizeName();
+
+	bool IsValid() const { return !Name.IsNone(); }
+	int32 NumContainers() const { return Containers.Num(); }
+	bool HasContainer() const { return NumContainers() > 0; }
+	bool IsArray(int32 InContainerIndex = 0) const { return Containers[InContainerIndex] == EPinContainerType::Array; }
+	bool IsMap(int32 InContainerIndex = 0) const { return Containers[InContainerIndex] == EPinContainerType::Map; }
+	bool IsSet(int32 InContainerIndex = 0) const { return Containers[InContainerIndex] == EPinContainerType::Set; }
+
+	FString GetBaseCPPType() const;
+
+	static const FString ArrayPrefix;
+	static const FString MapPrefix;
+	static const FString SetPrefix;
+	static const FString ContainerSuffix;
+};
+
 UCLASS()
-class RIGVM_API URigVMMemoryStorage : public UObject
+class RIGVM_API URigVMMemoryStorageGeneratorClass :
+	public UClass
 {
 	GENERATED_BODY()
 
 public:
 
-	//////////////////////////////////////////////////////////////////////////////
-	/// Property Management
-	//////////////////////////////////////////////////////////////////////////////
-	struct RIGVM_API FPropertyDescription
+	URigVMMemoryStorageGeneratorClass()
+		: Super()
+		, MemoryType(ERigVMMemoryType::Literal)
+	{}
+	
+	// UClass overrides
+	void PurgeClass(bool bRecompilingOnLoad) override;
+	void Link(FArchive& Ar, bool bRelinkExistingProperties) override;
+	void Serialize(FArchive& Ar) override;;
+
+	// URigVMMemoryStorageGeneratorClass specific
+	static URigVMMemoryStorageGeneratorClass* GetStorageClass(UObject* InOuter, ERigVMMemoryType InMemoryType);
+	static URigVMMemoryStorageGeneratorClass* CreateStorageClass(
+		UObject* InOuter,
+		ERigVMMemoryType InMemoryType,
+		const TArray<FRigVMPropertyDescription>& InProperties,
+		const TArray<FRigVMPropertyPathDescription>& InPropertyPaths = TArray<FRigVMPropertyPathDescription>());
+
+	ERigVMMemoryType GetMemoryType() const { return MemoryType; }
+	const TArray<const FProperty*>& GetProperties() const { return LinkedProperties; }
+	const TArray<FRigVMPropertyPath>& GetPropertyPaths() const { return PropertyPaths; }
+
+protected:
+
+	static FProperty* AddProperty(URigVMMemoryStorageGeneratorClass* InClass, const FRigVMPropertyDescription& InProperty, bool bPurge = false, bool bLink = true, FField** LinkToProperty = nullptr);
+
+private:
+
+	ERigVMMemoryType MemoryType;
+	TArray<const FProperty*> LinkedProperties;
+	TArray<FRigVMPropertyPath> PropertyPaths;
+	TArray<FRigVMPropertyPathDescription> PropertyPathDescriptions;
+
+	struct FClassInfo
 	{
-		FName Name;
-		const FProperty* Property;
-		FString CPPType;
-		UObject* CPPTypeObject;
-		TArray<EPinContainerType> Containers;
-		FString DefaultValue;
-
-		FPropertyDescription()
-			: Name(NAME_None)
-			, Property(nullptr)
-			, CPPType()
-			, CPPTypeObject(nullptr)
-			, Containers()
-			, DefaultValue()
-		{}
-		FPropertyDescription(const FProperty* InProperty, const FString& InDefaultValue, const FName& InName = NAME_None);
-		FPropertyDescription(const FName& InName, const FString& InCPPType, UObject* InCPPTypeObject, const FString& InDefaultValue);
-
-		static FName SanitizeName(const FName& InName);
-		void SanitizeName();
-
-		bool IsValid() const { return !Name.IsNone(); }
-		int32 NumContainers() const { return Containers.Num(); }
-		bool HasContainer() const { return NumContainers() > 0; }
-		bool IsArray(int32 InContainerIndex = 0) const { return Containers[InContainerIndex] == EPinContainerType::Array; }
-		bool IsMap(int32 InContainerIndex = 0) const { return Containers[InContainerIndex] == EPinContainerType::Map; }
-		bool IsSet(int32 InContainerIndex = 0) const { return Containers[InContainerIndex] == EPinContainerType::Set; }
-
-		FString GetBaseCPPType() const;
-
-		static const FString ArrayPrefix;
-		static const FString MapPrefix;
-		static const FString SetPrefix;
-		static const FString ContainerSuffix;
+		URigVMMemoryStorageGeneratorClass* LiteralStorageClass;
+		URigVMMemoryStorageGeneratorClass* WorkStorageClass;
+		URigVMMemoryStorageGeneratorClass* DebugStorageClass;
+		
+		FClassInfo()
+			: LiteralStorageClass(nullptr)
+			, WorkStorageClass(nullptr)
+			, DebugStorageClass(nullptr)
+		{
+		}
+
+		URigVMMemoryStorageGeneratorClass** GetClassPtr(ERigVMMemoryType InMemoryType)
+		{
+			switch(InMemoryType)
+			{
+			case ERigVMMemoryType::Literal:
+				{
+					return &LiteralStorageClass;
+				}
+			case ERigVMMemoryType::Debug:
+				{
+					return &DebugStorageClass;
+				}
+			default:
+				{
+					break;
+				}
+			}
+			return &WorkStorageClass;
+		}
 	};
+	
+	static TMap<FString, FClassInfo> PackageToInfo;
+
+	friend class URigVMMemoryStorage;
+	friend class URigVMCompiler;
+};
+
+UCLASS()
+class RIGVM_API URigVMMemoryStorage : public UObject
+{
+	GENERATED_BODY()
 
-	static UClass* GetStorageClass(UObject* InOuter, ERigVMMemoryType InMemoryType);
-	static UClass* CreateStorageClass(UObject* InOuter, ERigVMMemoryType InMemoryType, const TArray<FPropertyDescription>& InProperties);
-	static URigVMMemoryStorage* CreateStorage(UObject* InOuter, ERigVMMemoryType InMemoryType);
+public:
 
 	//////////////////////////////////////////////////////////////////////////////
 	/// Memory Access
@@ -149,6 +226,7 @@ public:
 	}
 	
 	const TArray<const FProperty*>& GetProperties() const;
+	const TArray<FRigVMPropertyPath>& GetPropertyPaths() const;
 
 	int32 GetPropertyIndex(const FProperty* InProperty) const;
 
@@ -286,48 +364,4 @@ public:
 		FRigVMMemoryHandle& SourceHandle);
 	
 #endif
-
-protected:
-
-	static FProperty* AddProperty(UClass* InClass, const FPropertyDescription& InProperty, bool bPurge = false, bool bLink = true, FField** LinkToProperty = nullptr);
-
-	TArray<const FProperty*> CachedProperties;
-
-private:
-	
-	struct FClassInfo
-	{
-		UClass* LiteralStorageClass;
-		UClass* WorkStorageClass;
-		UClass* DebugStorageClass;
-		
-		FClassInfo()
-			: LiteralStorageClass(nullptr)
-			, WorkStorageClass(nullptr)
-			, DebugStorageClass(nullptr)
-		{
-		}
-
-		UClass** GetClassPtr(ERigVMMemoryType InMemoryType)
-		{
-			switch(InMemoryType)
-			{
-				case ERigVMMemoryType::Literal:
-				{
-					return &LiteralStorageClass;
-				}
-				case ERigVMMemoryType::Debug:
-				{
-					return &DebugStorageClass;
-				}
-				default:
-				{
-					break;
-				}
-			}
-			return &WorkStorageClass;
-		}
-	};
-	
-	static TMap<FString, FClassInfo> PackageToInfo;
 };
diff --git a/Engine/Source/Runtime/RigVM/Public/RigVMCore/RigVMPropertyPath.h b/Engine/Source/Runtime/RigVM/Public/RigVMCore/RigVMPropertyPath.h
--- a/Engine/Source/Runtime/RigVM/Public/RigVMCore/RigVMPropertyPath.h
+++ b/Engine/Source/Runtime/RigVM/Public/RigVMCore/RigVMPropertyPath.h
@@ -5,6 +5,29 @@
 #include "CoreMinimal.h"
 #include "UObject/UnrealType.h"
 
+struct RIGVM_API FRigVMPropertyPathDescription
+{
+	int32 PropertyIndex;
+	FString SegmentPath;
+	
+	FRigVMPropertyPathDescription()
+		: PropertyIndex(INDEX_NONE)
+		, SegmentPath()
+	{}
+
+	FRigVMPropertyPathDescription(int32 InPropertyIndex, const FString& InSegmentPath)
+		: PropertyIndex(InPropertyIndex)
+		, SegmentPath(InSegmentPath)
+	{}
+
+	FORCEINLINE friend FArchive& operator<<(FArchive& Ar, FRigVMPropertyPathDescription& Path)
+	{
+		Ar << Path.PropertyIndex;
+		Ar << Path.SegmentPath;
+		return Ar;
+	}
+};
+
 enum RIGVM_API ERigVMPropertyPathSegmentType
 {
 	StructMember,
@@ -82,4 +105,5 @@ private:
 
 	FString Path;
 	TArray<FRigVMPropertyPathSegment> Segments;
-};
\ No newline at end of file
+};
+
