Control Rig: Re-enable all remaining Unit Tests for new memory storage class

diff --git a/Engine/Plugins/Animation/ControlRig/Source/ControlRig/Private/Units/Simulation/RigUnit_TimeOffset.cpp b/Engine/Plugins/Animation/ControlRig/Source/ControlRig/Private/Units/Simulation/RigUnit_TimeOffset.cpp
--- a/Engine/Plugins/Animation/ControlRig/Source/ControlRig/Private/Units/Simulation/RigUnit_TimeOffset.cpp
+++ b/Engine/Plugins/Animation/ControlRig/Source/ControlRig/Private/Units/Simulation/RigUnit_TimeOffset.cpp
@@ -31,6 +31,7 @@ FRigUnit_TimeOffsetFloat_Execute()
 	int32 MaxSize = FMath::Clamp<int32>(BufferSize, 2, 512);
 #if !UE_RIGVM_UCLASS_BASED_STORAGE_DISABLED
 	Buffer.SetNum(MaxSize);
+	DeltaTimes.SetNum(MaxSize);
 #endif
 	MaxSize = FMath::Min<int32>(MaxSize, Buffer.Num());
 
@@ -141,6 +142,7 @@ FRigUnit_TimeOffsetVector_Execute()
 	int32 MaxSize = FMath::Clamp<int32>(BufferSize, 2, 512);
 #if !UE_RIGVM_UCLASS_BASED_STORAGE_DISABLED
 	Buffer.SetNum(MaxSize);
+	DeltaTimes.SetNum(MaxSize);
 #endif
 	MaxSize = FMath::Min<int32>(MaxSize, Buffer.Num());
 
@@ -251,6 +253,7 @@ FRigUnit_TimeOffsetTransform_Execute()
 	int32 MaxSize = FMath::Clamp<int32>(BufferSize, 2, 512);
 #if !UE_RIGVM_UCLASS_BASED_STORAGE_DISABLED
 	Buffer.SetNum(MaxSize);
+	DeltaTimes.SetNum(MaxSize);
 #endif
 	MaxSize = FMath::Min<int32>(MaxSize, Buffer.Num());
 
diff --git a/Engine/Source/Runtime/RigVM/Private/RigVMCore/RigVM.cpp b/Engine/Source/Runtime/RigVM/Private/RigVMCore/RigVM.cpp
--- a/Engine/Source/Runtime/RigVM/Private/RigVMCore/RigVM.cpp
+++ b/Engine/Source/Runtime/RigVM/Private/RigVMCore/RigVM.cpp
@@ -599,8 +599,10 @@ void URigVM::CopyFrom(URigVM* InVM, bool bDeferCopy, bool bReferenceLiteralMemor
 			{
 				TargetMemory = NewObject<URigVMMemoryStorage>(Outer, SourceMemory->GetClass());
 			}
-
-			check(TargetMemory->GetClass() == SourceMemory->GetClass());
+			else if(TargetMemory->GetClass() != SourceMemory->GetClass())
+			{
+				TargetMemory = NewObject<URigVMMemoryStorage>(Outer, SourceMemory->GetClass());
+			}
 
 			for(int32 PropertyIndex = 0; PropertyIndex < TargetMemory->Num(); PropertyIndex++)
 			{
diff --git a/Engine/Source/Runtime/RigVM/Private/RigVMCore/RigVMByteCode.cpp b/Engine/Source/Runtime/RigVM/Private/RigVMCore/RigVMByteCode.cpp
--- a/Engine/Source/Runtime/RigVM/Private/RigVMCore/RigVMByteCode.cpp
+++ b/Engine/Source/Runtime/RigVM/Private/RigVMCore/RigVMByteCode.cpp
@@ -32,8 +32,6 @@ void FRigVMCopyOp::Serialize(FArchive& Ar)
 	Ar << Source;
 	Ar << Target;
 
-#if UE_RIGVM_UCLASS_BASED_STORAGE_DISABLED
-
 	if(Ar.IsLoading())
 	{
 		if (Ar.CustomVer(FUE5MainStreamObjectVersion::GUID) < FUE5MainStreamObjectVersion::RigVMCopyOpStoreNumBytes)
@@ -52,8 +50,6 @@ void FRigVMCopyOp::Serialize(FArchive& Ar)
 		Ar << NumBytes;
 		Ar << RegisterType;
 	}
-	
-#endif
 }
 
 void FRigVMComparisonOp::Serialize(FArchive& Ar)
diff --git a/Engine/Source/Runtime/RigVM/Private/RigVMCore/RigVMMemoryStorage.cpp b/Engine/Source/Runtime/RigVM/Private/RigVMCore/RigVMMemoryStorage.cpp
--- a/Engine/Source/Runtime/RigVM/Private/RigVMCore/RigVMMemoryStorage.cpp
+++ b/Engine/Source/Runtime/RigVM/Private/RigVMCore/RigVMMemoryStorage.cpp
@@ -224,7 +224,7 @@ URigVMMemoryStorageGeneratorClass* URigVMMemoryStorageGeneratorClass::CreateStor
 	const TArray<FRigVMPropertyPathDescription>& InPropertyPaths)
 {
 	check(InOuter);
-	UPackage* Package = InOuter->GetOutermost(); 
+	UPackage* Package = InOuter->GetOutermost();
 	UClass *SuperClass = URigVMMemoryStorage::StaticClass();
 
 	const FString ClassName = GetClassName(InMemoryType);
@@ -484,6 +484,13 @@ FString URigVMMemoryStorage::GetDataAsString(int32 InPropertyIndex)
 	return Value;
 }
 
+FString URigVMMemoryStorage::GetDataAsString(const FRigVMOperand& InOperand)
+{
+	const int32 PropertyIndex = InOperand.GetRegisterIndex();
+	check(IsValidIndex(PropertyIndex));
+	return GetDataAsString(PropertyIndex);
+}
+
 class FRigVMMemoryStorageImportErrorContext : public FOutputDevice
 {
 public:
@@ -542,8 +549,36 @@ bool URigVMMemoryStorage::CopyProperty(
 	check(InTargetPtr != nullptr);
 	check(InSourcePtr != nullptr);
 
-	if(!ensure(InTargetProperty->SameType(InSourceProperty)))
+	if(!InTargetProperty->SameType(InSourceProperty))
 	{
+		if(const FFloatProperty* TargetFloatProperty = CastField<FFloatProperty>(InTargetProperty))
+		{
+			if(const FDoubleProperty* SourceDoubleProperty = CastField<FDoubleProperty>(InSourceProperty))
+			{
+				float* TargetFloats = (float*)InTargetPtr;
+				double* SourceDoubles = (double*)InSourcePtr;
+				for(int32 Index=0;Index<TargetFloatProperty->ArrayDim;Index++)
+				{
+					TargetFloats[Index] = (float)SourceDoubles[Index];
+				}
+				return true;
+			}
+		}
+		else if(const FDoubleProperty* TargetDoubleProperty = CastField<FDoubleProperty>(InTargetProperty))
+		{
+			if(const FFloatProperty* SourceFloatProperty = CastField<FFloatProperty>(InSourceProperty))
+			{
+				double* TargetDoubles = (double*)InTargetPtr;
+				float* SourceFloats = (float*)InSourcePtr;
+				for(int32 Index=0;Index<TargetDoubleProperty->ArrayDim;Index++)
+				{
+					TargetDoubles[Index] = (double)SourceFloats[Index];
+				}
+				return true;
+			}
+		}
+		
+		ensure(InTargetProperty->SameType(InSourceProperty));
 		return false;
 	}
 
@@ -638,7 +673,15 @@ const TArray<FRigVMPropertyPath>& URigVMMemoryStorage::GetPropertyPaths() const
 
 int32 URigVMMemoryStorage::GetPropertyIndex(const FProperty* InProperty) const
 {
-	return GetProperties().Find(InProperty);
+	const TArray<const FProperty*>& Properties = GetProperties();
+	for(int32 PropertyIndex = 0; PropertyIndex < Properties.Num(); PropertyIndex++)
+	{
+		if(Properties[PropertyIndex] == InProperty)
+		{
+			return PropertyIndex;
+		}
+	}
+	return INDEX_NONE;
 }
 
 int32 URigVMMemoryStorage::GetPropertyIndexByName(const FName& InName) const
@@ -652,3 +695,23 @@ const FProperty* URigVMMemoryStorage::FindPropertyByName(const FName& InName) co
 	const FName SanitizedName = FRigVMPropertyDescription::SanitizeName(InName);
 	return GetClass()->FindPropertyByName(SanitizedName);
 }
+
+FRigVMOperand URigVMMemoryStorage::GetOperand(int32 InPropertyIndex, int32 InPropertyPathIndex) const
+{
+	if(GetProperties().IsValidIndex(InPropertyIndex))
+	{
+		if(InPropertyPathIndex != INDEX_NONE)
+		{
+			check(GetPropertyPaths().IsValidIndex(InPropertyPathIndex));
+			return FRigVMOperand(GetMemoryType(), InPropertyIndex, InPropertyPathIndex);
+		}
+		return FRigVMOperand(GetMemoryType(), InPropertyIndex);
+	}
+	return FRigVMOperand();
+}
+
+FRigVMOperand URigVMMemoryStorage::GetOperandByName(const FName& InName, int32 InPropertyPathIndex) const
+{
+	const int32 PropertyIndex = GetPropertyIndexByName(InName);
+	return GetOperand(PropertyIndex, InPropertyPathIndex);
+}
diff --git a/Engine/Source/Runtime/RigVM/Public/RigVMCore/RigVM.h b/Engine/Source/Runtime/RigVM/Public/RigVMCore/RigVM.h
--- a/Engine/Source/Runtime/RigVM/Public/RigVMCore/RigVM.h
+++ b/Engine/Source/Runtime/RigVM/Public/RigVMCore/RigVM.h
@@ -356,6 +356,29 @@ public:
 
 #else
 
+	FORCEINLINE_DEBUGGABLE FRigVMParameter AddParameter(ERigVMParameterType InType, const FName& InParameterName, const FName& InWorkMemoryPropertyName)
+	{
+		check(GetWorkMemory());
+
+		if(ParametersNameMap.Contains(InParameterName))
+		{
+			return FRigVMParameter();
+		}
+
+		const FProperty* Property = GetWorkMemory()->FindPropertyByName(InWorkMemoryPropertyName);
+		const int32 PropertyIndex = GetWorkMemory()->GetPropertyIndex(Property);
+
+		UScriptStruct* Struct = nullptr;
+		if(const FStructProperty* StructProperty = CastField<FStructProperty>(Property))
+		{
+			Struct = StructProperty->Struct;
+		}
+		
+		FRigVMParameter Parameter(InType, InParameterName, PropertyIndex, Property->GetCPPType(), Struct);
+		ParametersNameMap.Add(Parameter.Name, Parameters.Add(Parameter));
+		return Parameter;
+	}
+
 	// Retrieve the array size of the parameter
 	FORCEINLINE int32 GetParameterArraySize(const FRigVMParameter& InParameter) const
 	{
@@ -605,6 +628,7 @@ public:
 		return FRigVMOperand(ERigVMMemoryType::External, VariableIndex);
 	}
 
+#if UE_RIGVM_UCLASS_BASED_STORAGE_DISABLED
 	// Adds a new external / unowned variable to the VM
 	template<typename T>
 	FORCEINLINE_DEBUGGABLE FRigVMOperand AddExternalVariable(const FName& InExternalVariableName, T& InValue)
@@ -618,6 +642,7 @@ public:
 	{
 		return AddExternalVariable(FRigVMExternalVariable::Make(InExternalVariableName, InValue));
 	}
+#endif
 
 #if UE_RIGVM_UCLASS_BASED_STORAGE_DISABLED
 	void SetRegisterValueFromString(const FRigVMOperand& InOperand, const FString& InCPPType, const UObject* InCPPTypeObject, const TArray<FString>& InDefaultValues);
diff --git a/Engine/Source/Runtime/RigVM/Public/RigVMCore/RigVMByteCode.h b/Engine/Source/Runtime/RigVM/Public/RigVMCore/RigVMByteCode.h
--- a/Engine/Source/Runtime/RigVM/Public/RigVMCore/RigVMByteCode.h
+++ b/Engine/Source/Runtime/RigVM/Public/RigVMCore/RigVMByteCode.h
@@ -6,6 +6,7 @@
 #include "RigVMDefines.h"
 #include "RigVMRegistry.h"
 #include "RigVMStatistics.h"
+#include "RigVMMemoryDeprecated.h"
 
 #include "RigVMByteCode.generated.h"
 
@@ -234,15 +235,17 @@ struct RIGVM_API FRigVMCopyOp : public FRigVMBaseOp
 {
 	GENERATED_USTRUCT_BODY()
 
+#if !UE_RIGVM_UCLASS_BASED_STORAGE_DISABLED
+public:
+#endif
+
 	FRigVMCopyOp()
 	: FRigVMBaseOp(ERigVMOpCode::Copy)
 	, Source()
 	, Target()
-#if UE_RIGVM_UCLASS_BASED_STORAGE_DISABLED
 	, NumBytes(0)
 	, RegisterType(ERigVMRegisterType::Invalid)
 	, CopyType(ERigVMCopyType::Default)
-#endif
 	{
 	}
 
@@ -262,18 +265,26 @@ struct RIGVM_API FRigVMCopyOp : public FRigVMBaseOp
 		, NumBytes(InNumBytes)
 		, RegisterType(InRegisterType)
 		, CopyType(InCopyType)
+#else
+		, NumBytes(0)
+		, RegisterType(ERigVMRegisterType::Invalid)
+		, CopyType(ERigVMCopyType::Default)
 #endif
 	{
 	}
 
 	FRigVMOperand Source;
 	FRigVMOperand Target;
-#if UE_RIGVM_UCLASS_BASED_STORAGE_DISABLED
+#if !UE_RIGVM_UCLASS_BASED_STORAGE_DISABLED
+private:
+#endif
 	uint16 NumBytes;
 	ERigVMRegisterType RegisterType;
 	ERigVMCopyType CopyType;
-#endif
 
+#if !UE_RIGVM_UCLASS_BASED_STORAGE_DISABLED
+public:
+#endif
 	void Serialize(FArchive& Ar);
 	FORCEINLINE friend FArchive& operator<<(FArchive& Ar, FRigVMCopyOp& P)
 	{
diff --git a/Engine/Source/Runtime/RigVM/Public/RigVMCore/RigVMMemoryStorage.h b/Engine/Source/Runtime/RigVM/Public/RigVMCore/RigVMMemoryStorage.h
--- a/Engine/Source/Runtime/RigVM/Public/RigVMCore/RigVMMemoryStorage.h
+++ b/Engine/Source/Runtime/RigVM/Public/RigVMCore/RigVMMemoryStorage.h
@@ -185,6 +185,16 @@ public:
 	/// Memory Access
 	//////////////////////////////////////////////////////////////////////////////
 
+	FORCEINLINE ERigVMMemoryType GetMemoryType() const
+	{
+		if(URigVMMemoryStorageGeneratorClass* Class = Cast<URigVMMemoryStorageGeneratorClass>(GetClass()))
+		{
+			return Class->GetMemoryType();
+		}
+		// empty debug containers don't have a generator class
+		return ERigVMMemoryType::Debug;
+	}
+	
 	FORCEINLINE int32 Num() const
 	{
 		return GetProperties().Num();
@@ -194,7 +204,7 @@ public:
 	{
 		return GetProperties().IsValidIndex(InIndex);
 	}
-	
+
 	const TArray<const FProperty*>& GetProperties() const;
 	const TArray<FRigVMPropertyPath>& GetPropertyPaths() const;
 
@@ -202,11 +212,19 @@ public:
 
 	int32 GetPropertyIndexByName(const FName& InName) const;
 
+	FORCEINLINE const FProperty* GetProperty(int32 InPropertyIndex) const
+	{
+		return GetProperties()[InPropertyIndex];
+	}
+
 	const FProperty* FindPropertyByName(const FName& InName) const;
+
+	FRigVMOperand GetOperand(int32 InPropertyIndex, int32 InPropertyPathIndex = INDEX_NONE) const;
+	FRigVMOperand GetOperandByName(const FName& InName, int32 InPropertyPathIndex = INDEX_NONE) const;
 	
 	FORCEINLINE bool IsArray(int32 InPropertyIndex) const
 	{
-		return GetProperties()[InPropertyIndex]->IsA<FArrayProperty>();
+		return GetProperty(InPropertyIndex)->IsA<FArrayProperty>();
 	}
 	
 	FORCEINLINE bool IsArrayByName(const FName& InName) const
@@ -220,7 +238,7 @@ public:
 	
 	FORCEINLINE bool IsMap(int32 InPropertyIndex) const
 	{
-		return GetProperties()[InPropertyIndex]->IsA<FMapProperty>();
+		return GetProperty(InPropertyIndex)->IsA<FMapProperty>();
 	}
 	
 	FORCEINLINE bool IsMapByName(const FName& InName) const
@@ -234,7 +252,7 @@ public:
 	
 	FORCEINLINE bool IsSet(int32 InPropertyIndex) const
 	{
-		return GetProperties()[InPropertyIndex]->IsA<FSetProperty>();
+		return GetProperty(InPropertyIndex)->IsA<FSetProperty>();
 	}
 	
 	FORCEINLINE bool IsSetByName(const FName& InName) const
@@ -267,7 +285,7 @@ public:
 	template<typename T>
 	FORCEINLINE T* GetData(int32 InPropertyIndex, const FRigVMPropertyPath& InPropertyPath)
 	{
-		const FProperty* Property = GetProperties()[InPropertyIndex];
+		const FProperty* Property = GetProperty(InPropertyIndex);
 		return InPropertyPath.GetData<T>(GetData<uint8>(InPropertyIndex), Property);
 	}
 
@@ -278,7 +296,55 @@ public:
 		return GetData<T>(PropertyIndex, InPropertyPath);
 	}
 
+	template<typename T>
+	FORCEINLINE T* GetData(const FRigVMOperand& InOperand)
+	{
+		const int32 PropertyIndex = InOperand.GetRegisterIndex();
+		const int32 PropertyPathIndex = InOperand.GetRegisterOffset();
+
+		check(GetProperties().IsValidIndex(PropertyIndex));
+		
+		if(PropertyPathIndex == INDEX_NONE)
+		{
+			return GetData<T>(PropertyIndex);
+		}
+
+		check(GetPropertyPaths().IsValidIndex(PropertyPathIndex));
+		return GetData<T>(PropertyIndex, GetPropertyPaths()[PropertyPathIndex]); 
+	}
+
+	template<typename T>
+	FORCEINLINE T& GetRef(int32 InPropertyIndex)
+	{
+		return *GetData<T>(InPropertyIndex);
+	}
+
+	template<typename T>
+	FORCEINLINE T& GetRefByName(const FName& InName)
+	{
+		return *GetDataByName<T>(InName);
+	}
+
+	template<typename T>
+	FORCEINLINE T& GetRef(int32 InPropertyIndex, const FRigVMPropertyPath& InPropertyPath)
+	{
+		return *GetData<T>(InPropertyIndex, InPropertyPath);
+	}
+
+	template<typename T>
+	FORCEINLINE T& GetRefByName(const FName& InName, const FRigVMPropertyPath& InPropertyPath)
+	{
+		return *GetDataByName<T>(InName, InPropertyPath);
+	}
+
+	template<typename T>
+	FORCEINLINE T& GetRef(const FRigVMOperand& InOperand)
+	{
+		return *GetData<T>(InOperand);
+	}
+
 	FString GetDataAsString(int32 InPropertyIndex);
+	FString GetDataAsString(const FRigVMOperand& InOperand);
 	
 	FORCEINLINE FString GetDataAsStringByName(const FName& InName)
 	{
