Control Rig: Bracket Property Metadata access with EDITORDATA_ONLY

diff --git a/Engine/Source/Runtime/RigVM/Private/RigVMCore/RigVMMemoryStorage.cpp b/Engine/Source/Runtime/RigVM/Private/RigVMCore/RigVMMemoryStorage.cpp
--- a/Engine/Source/Runtime/RigVM/Private/RigVMCore/RigVMMemoryStorage.cpp
+++ b/Engine/Source/Runtime/RigVM/Private/RigVMCore/RigVMMemoryStorage.cpp
@@ -446,6 +446,8 @@ FProperty* URigVMMemoryStorageGeneratorClass::AddProperty(URigVMMemoryStorageGen
 
 		Result->SetPropertyFlags(CPF_Edit | CPF_EditConst | CPF_NonTransactional);
 
+#if WITH_EDITORONLY_DATA
+
 		static const FName NAME_DisplayName(TEXT("DisplayName"));
 		static const FName NAME_ToolTipName(TEXT("ToolTip"));
 
@@ -457,7 +459,8 @@ FProperty* URigVMMemoryStorageGeneratorClass::AddProperty(URigVMMemoryStorageGen
 		DisplayName = FString::Printf(TEXT("[%d] %s"), InClass->LinkedProperties.Num(), *DisplayName);
 		Result->SetMetaData(NAME_DisplayName, *DisplayName);
 		Result->SetMetaData(NAME_ToolTipName, *FString::Printf(TEXT("Name %s\nCPPType %s"), *Result->GetName(), *InProperty.CPPType));
-
+		
+#endif
 		InClass->LinkedProperties.Add(Result);
 		(*LinkToProperty) = Result;
 	}
