Control Rig: RigVM: Documented headers and body of new RigVMMemoryStorage classes

diff --git a/Engine/Source/Developer/RigVMDeveloper/Private/RigVMCompiler/RigVMCompiler.cpp b/Engine/Source/Developer/RigVMDeveloper/Private/RigVMCompiler/RigVMCompiler.cpp
--- a/Engine/Source/Developer/RigVMDeveloper/Private/RigVMCompiler/RigVMCompiler.cpp
+++ b/Engine/Source/Developer/RigVMDeveloper/Private/RigVMCompiler/RigVMCompiler.cpp
@@ -105,7 +105,7 @@ FRigVMPropertyDescription FRigVMCompilerWorkData::GetProperty(const FRigVMOperan
 	return FRigVMPropertyDescription();
 }
 
-int32 FRigVMCompilerWorkData::FindOrAddPropertyPath(const FRigVMOperand& InOperand, const FString& InRootCPPType, const FString& InSegmentPath)
+int32 FRigVMCompilerWorkData::FindOrAddPropertyPath(const FRigVMOperand& InOperand, const FString& InHeadCPPType, const FString& InSegmentPath)
 {
 	if(InSegmentPath.IsEmpty())
 	{
@@ -116,12 +116,12 @@ int32 FRigVMCompilerWorkData::FindOrAddPropertyPath(const FRigVMOperand& InOpera
 	for(int32 Index = 0; Index < Descriptions.Num(); Index++)
 	{
 		const FRigVMPropertyPathDescription& Description = Descriptions[Index]; 
-		if(Description.RootCPPType == InRootCPPType && Description.SegmentPath == InSegmentPath)
+		if(Description.HeadCPPType == InHeadCPPType && Description.SegmentPath == InSegmentPath)
 		{
 			return Index;
 		}
 	}
-	return Descriptions.Add(FRigVMPropertyPathDescription(InOperand.GetRegisterIndex(), InRootCPPType, InSegmentPath));
+	return Descriptions.Add(FRigVMPropertyPathDescription(InOperand.GetRegisterIndex(), InHeadCPPType, InSegmentPath));
 }
 
 #endif
diff --git a/Engine/Source/Developer/RigVMDeveloper/Private/RigVMModel/RigVMPin.cpp b/Engine/Source/Developer/RigVMDeveloper/Private/RigVMModel/RigVMPin.cpp
--- a/Engine/Source/Developer/RigVMDeveloper/Private/RigVMModel/RigVMPin.cpp
+++ b/Engine/Source/Developer/RigVMDeveloper/Private/RigVMModel/RigVMPin.cpp
@@ -814,7 +814,7 @@ bool URigVMPin::CanBeBoundToVariable(const FRigVMExternalVariable& InExternalVar
 
 		const FProperty* Property = InExternalVariable.Property;
 		const FRigVMPropertyPath PropertyPath(Property, InSegmentPath);
-		Property = PropertyPath.GetTargetProperty();
+		Property = PropertyPath.GetTailProperty();
 
 		FRigVMExternalVariable::GetTypeFromProperty(Property, ExternalCPPType, ExternalCPPTypeObject);
 	}
diff --git a/Engine/Source/Developer/RigVMDeveloper/Public/RigVMCompiler/RigVMCompiler.h b/Engine/Source/Developer/RigVMDeveloper/Public/RigVMCompiler/RigVMCompiler.h
--- a/Engine/Source/Developer/RigVMDeveloper/Public/RigVMCompiler/RigVMCompiler.h
+++ b/Engine/Source/Developer/RigVMDeveloper/Public/RigVMCompiler/RigVMCompiler.h
@@ -102,7 +102,7 @@ public:
 	FRigVMOperand AddProperty(ERigVMMemoryType InMemoryType, const FName& InName, const FString& InCPPType, UObject* InCPPTypeObject, const FString& InDefaultValue = FString());
 	FRigVMOperand FindProperty(ERigVMMemoryType InMemoryType, const FName& InName);
 	FRigVMPropertyDescription GetProperty(const FRigVMOperand& InOperand);
-	int32 FindOrAddPropertyPath(const FRigVMOperand& InOperand, const FString& InRootCPPType, const FString& InSegmentPath);
+	int32 FindOrAddPropertyPath(const FRigVMOperand& InOperand, const FString& InHeadCPPType, const FString& InSegmentPath);
 	
 #endif
 };
diff --git a/Engine/Source/Runtime/RigVM/Private/RigVMCore/RigVMMemoryStorage.cpp b/Engine/Source/Runtime/RigVM/Private/RigVMCore/RigVMMemoryStorage.cpp
--- a/Engine/Source/Runtime/RigVM/Private/RigVMCore/RigVMMemoryStorage.cpp
+++ b/Engine/Source/Runtime/RigVM/Private/RigVMCore/RigVMMemoryStorage.cpp
@@ -7,15 +7,12 @@
 
 #if !UE_RIGVM_UCLASS_BASED_STORAGE_DISABLED
 
-const FRigVMPropertyPath FRigVMMemoryHandle::EmptyPropertyPath;
-
 #endif
 
 ////////////////////////////////////////////////////////////////////////////////////////////////////////////////////////
 
 const FString FRigVMPropertyDescription::ArrayPrefix = TEXT("TArray<");
 const FString FRigVMPropertyDescription::MapPrefix = TEXT("TMap<");
-const FString FRigVMPropertyDescription::SetPrefix = TEXT("TSet<");
 const FString FRigVMPropertyDescription::ContainerSuffix = TEXT(">");
 
 ////////////////////////////////////////////////////////////////////////////////////////////////////////////////////////
@@ -28,8 +25,11 @@ FRigVMPropertyDescription::FRigVMPropertyDescription(const FProperty* InProperty
 	, Containers()
 	, DefaultValue(InDefaultValue)
 {
+	// make sure to use valid names only
 	SanitizeName();
 
+	// walk the property until you reach the tail
+	// and build up the container array along the way
 	const FProperty* ChildProperty = InProperty;
 	do
 	{
@@ -40,8 +40,7 @@ FRigVMPropertyDescription::FRigVMPropertyDescription(const FProperty* InProperty
 		}
 		else if(const FSetProperty* SetProperty = CastField<FSetProperty>(ChildProperty))
 		{
-			Containers.Add(EPinContainerType::Set);
-			ChildProperty = SetProperty->ElementProp;
+			checkNoEntry();
 		}
 		else if(const FMapProperty* MapProperty = CastField<FMapProperty>(ChildProperty))
 		{
@@ -64,33 +63,31 @@ FRigVMPropertyDescription::FRigVMPropertyDescription(const FName& InName, const
 	, Containers()
 	, DefaultValue(InDefaultValue)
 {
+	// only allow valid names
 	SanitizeName();
 
-	FString BaseCPPType = CPPType;
+	FString TailCPPType = CPPType;
 
+	// build the containers and validate the expected string
+	// for the tail CPP type
 	do
 	{
-		if(BaseCPPType.RemoveFromStart(ArrayPrefix))
+		if(TailCPPType.RemoveFromStart(ArrayPrefix))
 		{
 			Containers.Add(EPinContainerType::Array);
-			check(BaseCPPType.RemoveFromEnd(ContainerSuffix));
+			check(TailCPPType.RemoveFromEnd(ContainerSuffix));
 		}
-		else if(BaseCPPType.RemoveFromStart(MapPrefix))
+		else if(TailCPPType.RemoveFromStart(MapPrefix))
 		{
 			Containers.Add(EPinContainerType::Map);
-			check(BaseCPPType.RemoveFromEnd(ContainerSuffix));
-		}
-		else if(BaseCPPType.RemoveFromStart(SetPrefix))
-		{
-			Containers.Add(EPinContainerType::Set);
-			check(BaseCPPType.RemoveFromEnd(ContainerSuffix));
+			check(TailCPPType.RemoveFromEnd(ContainerSuffix));
 		}
 		else
 		{
 			break;
 		}
 	}
-	while(!BaseCPPType.IsEmpty());
+	while(!TailCPPType.IsEmpty());
 }
 
 FName FRigVMPropertyDescription::SanitizeName(const FName& InName)
@@ -125,30 +122,26 @@ void FRigVMPropertyDescription::SanitizeName()
 	Name = SanitizeName(Name);
 }
 
-FString FRigVMPropertyDescription::GetBaseCPPType() const
+FString FRigVMPropertyDescription::GetTailCPPType() const
 {
-	FString BaseCPPType = CPPType;
+	FString TailCPPType = CPPType;
 
+	// walk the containers and remove prefix and
+	// suffix from the path. Turns 'TArray<TArray<float>>' to 'float'
 	for(EPinContainerType Container : Containers)
 	{
 		switch(Container)
 		{
 			case EPinContainerType::Array:
 			{
-				check(BaseCPPType.RemoveFromStart(ArrayPrefix))
-				check(BaseCPPType.RemoveFromEnd(ContainerSuffix));
+				check(TailCPPType.RemoveFromStart(ArrayPrefix))
+				check(TailCPPType.RemoveFromEnd(ContainerSuffix));
 				break;
 			}		
 			case EPinContainerType::Map:
 			{
-				check(BaseCPPType.RemoveFromStart(MapPrefix))
-				check(BaseCPPType.RemoveFromEnd(ContainerSuffix));
-				break;
-			}		
-			case EPinContainerType::Set:
-			{
-				check(BaseCPPType.RemoveFromStart(SetPrefix))
-				check(BaseCPPType.RemoveFromEnd(ContainerSuffix));
+				check(TailCPPType.RemoveFromStart(MapPrefix))
+				check(TailCPPType.RemoveFromEnd(ContainerSuffix));
 				break;
 			}		
 			case EPinContainerType::None:
@@ -159,7 +152,7 @@ FString FRigVMPropertyDescription::GetBaseCPPType() const
 		}
 	}
 	
-	return BaseCPPType;
+	return TailCPPType;
 }
 
 ////////////////////////////////////////////////////////////////////////////////////////////////////////////////////////
@@ -167,6 +160,8 @@ FString FRigVMPropertyDescription::GetBaseCPPType() const
 void URigVMMemoryStorageGeneratorClass::PurgeClass(bool bRecompilingOnLoad)
 {
 	Super::PurgeClass(bRecompilingOnLoad);
+
+	// clear our state as well
 	LinkedProperties.Reset();
 	PropertyPaths.Reset();
 	PropertyPathDescriptions.Reset();
@@ -180,15 +175,8 @@ void URigVMMemoryStorageGeneratorClass::Link(FArchive& Ar, bool bRelinkExistingP
 	// garbage collector.
 	AssembleReferenceTokenStream(/*bForce=*/true);
 
-	// Setup the LinkedProperties
-	LinkedProperties.Reset();
-	const FProperty* Property = CastField<FProperty>(ChildProperties);
-	while(Property)
-	{
-		LinkedProperties.Add(Property);
-		Property = CastField<FProperty>(Property->Next);
-	}
-
+	// rebuild property list and property path list
+	RefreshLinkedProperties();
 	RefreshPropertyPaths();
 }
 
@@ -203,6 +191,15 @@ void URigVMMemoryStorageGeneratorClass::Serialize(FArchive& Ar)
 	}
 }
 
+void URigVMMemoryStorageGeneratorClass::PostLoad()
+{
+	Super::PostLoad();
+
+	// rebuild property list and property path list
+	RefreshLinkedProperties();
+	RefreshPropertyPaths();
+}
+
 FString URigVMMemoryStorageGeneratorClass::GetClassName(ERigVMMemoryType InMemoryType)
 {
 	return TEXT("RigVMMemory_") + StaticEnum<ERigVMMemoryType>()->GetDisplayNameTextByValue((int64)InMemoryType).ToString();
@@ -228,7 +225,8 @@ URigVMMemoryStorageGeneratorClass* URigVMMemoryStorageGeneratorClass::CreateStor
 	UClass *SuperClass = URigVMMemoryStorage::StaticClass();
 
 	const FString ClassName = GetClassName(InMemoryType);
-	
+
+	// if there's an old class - remove it from the package and mark it to destroy
 	URigVMMemoryStorageGeneratorClass* OldClass = FindObject<URigVMMemoryStorageGeneratorClass>(Package, *ClassName);
 	if(OldClass)
 	{
@@ -237,24 +235,34 @@ URigVMMemoryStorageGeneratorClass* URigVMMemoryStorageGeneratorClass::CreateStor
 		OldClass->MarkPendingKill();
 	}
 
+	// create the new class
 	URigVMMemoryStorageGeneratorClass* Class = NewObject<URigVMMemoryStorageGeneratorClass>(
 		Package,
 		*ClassName,
 		RF_Standalone | RF_Public
 	);
 
+	// add it to the root so that the garbage collector doesn't pick it up
 	Class->AddToRoot();
 
+	// clear the class (sets relevant flags)
 	Class->PurgeClass(false);
+
+	// setup class to inherit from the super class
 	Class->SetSuperStruct(SuperClass);
 	Class->PropertyLink = SuperClass->PropertyLink;
+
+	// allow to create this class's instances everywhere
 	Class->ClassWithin = UObject::StaticClass();
+
+	// make sure that it cannot be placed / used for class selectors
 	Class->ClassFlags |= CLASS_NotPlaceable;
+
+	// store our custom state
 	Class->MemoryType = InMemoryType;
 
-	// Generate properties
+	// Generate properties by iterating
 	FField** LinkToProperty = &Class->ChildProperties;
-
 	for(const FRigVMPropertyDescription& PropertyDescription : InProperties)
 	{
 		FProperty* CachedProperty = AddProperty(Class, PropertyDescription, LinkToProperty);
@@ -264,7 +272,7 @@ URigVMMemoryStorageGeneratorClass* URigVMMemoryStorageGeneratorClass::CreateStor
 	// Store the property path descriptions
 	Class->PropertyPathDescriptions = InPropertyPaths;
 
-	// Update the class
+	// Update the class's data and link it / build it
 	Class->Bind();
 	Class->StaticLink(true);
 	
@@ -287,24 +295,6 @@ URigVMMemoryStorageGeneratorClass* URigVMMemoryStorageGeneratorClass::CreateStor
 		Property->ImportText(*DefaultValue, ValuePtr, EPropertyPortFlags::PPF_None, nullptr);
 	}
 
-	/*
-	if(OldClass)
-	{
-		UEngine::FCopyPropertiesForUnrelatedObjectsParams Params;
-		Params.bAggressiveDefaultSubobjectReplacement = false;
-		Params.bDoDelta = false;
-		Params.bCopyDeprecatedProperties = true;
-		Params.bSkipCompilerGeneratedDefaults = true;
-		Params.bNotifyObjectReplacement = true;
-
-		Params.bClearReferences = true;
-		UEngine::CopyPropertiesForUnrelatedObjects(OldClass->GetDefaultObject(), CDO, Params);
-
-		// we might need to run this again on the old instances,
-		// but the VM compiler should take care of it anyway
-	}
-	*/
-	
 	return Class;
 }
 
@@ -327,6 +317,7 @@ FProperty* URigVMMemoryStorageGeneratorClass::AddProperty(URigVMMemoryStorageGen
 	FProperty* Result = nullptr;
 	if(InProperty.Property)
 	{
+		// if the property description provides a property simply duplicate that property
 		FProperty* NewProperty = CastFieldChecked<FProperty>(FField::Duplicate(InProperty.Property, InClass, InProperty.Name));
 		check(NewProperty);
 
@@ -336,7 +327,6 @@ FProperty* URigVMMemoryStorageGeneratorClass::AddProperty(URigVMMemoryStorageGen
 	else
 	{
 		FFieldVariant PropertyOwner = InClass;
-		// FProperty** KeyPropertyPtr = nullptr;
 		FProperty** ValuePropertyPtr = &Result;
 
 		for(EPinContainerType Container : InProperty.Containers)
@@ -345,6 +335,7 @@ FProperty* URigVMMemoryStorageGeneratorClass::AddProperty(URigVMMemoryStorageGen
 			{
 				case EPinContainerType::Array:
 				{
+					// create an array property as a container for the tail
 					FArrayProperty* ArrayProperty = new FArrayProperty(PropertyOwner, InProperty.Name, RF_Public);
 					*ValuePropertyPtr = ArrayProperty;
 					ValuePropertyPtr = &ArrayProperty->Inner;
@@ -353,20 +344,17 @@ FProperty* URigVMMemoryStorageGeneratorClass::AddProperty(URigVMMemoryStorageGen
 				}
 				case EPinContainerType::Map:
 				{
+					// create a map property as a container for the tail
 					checkNoEntry(); // this is not implemented yet
 					FMapProperty* MapProperty = new FMapProperty(PropertyOwner, InProperty.Name, RF_Public);
 					*ValuePropertyPtr = MapProperty;
-					// KeyPropertyPtr = &MapProperty->KeyProp;
 					ValuePropertyPtr = &MapProperty->ValueProp;
 					PropertyOwner = MapProperty;
 					break;
 				}
 				case EPinContainerType::Set:
 				{
-					FSetProperty* SetProperty = new FSetProperty(PropertyOwner, InProperty.Name, RF_Public);
-					*ValuePropertyPtr = SetProperty;
-					ValuePropertyPtr = &SetProperty->ElementProp;
-					PropertyOwner = SetProperty;
+					checkNoEntry();
 					break;
 				}
 				case EPinContainerType::None:
@@ -377,6 +365,7 @@ FProperty* URigVMMemoryStorageGeneratorClass::AddProperty(URigVMMemoryStorageGen
 			}
 		}
 
+		// create properties for types which are provided through the CPPTypeObject
 		if(InProperty.CPPTypeObject != nullptr)
 		{
 			if(UEnum* Enum = Cast<UEnum>(InProperty.CPPTypeObject))
@@ -402,7 +391,7 @@ FProperty* URigVMMemoryStorageGeneratorClass::AddProperty(URigVMMemoryStorageGen
 				checkNoEntry();
 			}
 		}
-		else // take care of default types...
+		else // take care of simple types...
 		{
 			static FString BoolString = TEXT("bool");
 			static FString Int32String = TEXT("int32");
@@ -412,7 +401,7 @@ FProperty* URigVMMemoryStorageGeneratorClass::AddProperty(URigVMMemoryStorageGen
 			static FString StringString = TEXT("FString");
 			static FString NameString = TEXT("FName");
 
-			const FString BaseCPPType = InProperty.GetBaseCPPType();
+			const FString BaseCPPType = InProperty.GetTailCPPType();
 			if(BaseCPPType.Equals(BoolString, ESearchCase::IgnoreCase))
 			{
 				(*ValuePropertyPtr) = new FBoolProperty(PropertyOwner, InProperty.Name, RF_Public);;
@@ -444,10 +433,14 @@ FProperty* URigVMMemoryStorageGeneratorClass::AddProperty(URigVMMemoryStorageGen
 			}
 		}
 
+		// mark the property as visible only for details panels
+		// also mark it as non transactional, which means the undo / redo system will ignore changes to it
 		Result->SetPropertyFlags(CPF_Edit | CPF_EditConst | CPF_NonTransactional);
 
 #if WITH_EDITORONLY_DATA
 
+		// store some additional meta data,
+		// mainly for inspecting things in the details panel
 		static const FName NAME_DisplayName(TEXT("DisplayName"));
 		static const FName NAME_ToolTipName(TEXT("ToolTip"));
 
@@ -461,6 +454,8 @@ FProperty* URigVMMemoryStorageGeneratorClass::AddProperty(URigVMMemoryStorageGen
 		Result->SetMetaData(NAME_ToolTipName, *FString::Printf(TEXT("Name %s\nCPPType %s"), *Result->GetName(), *InProperty.CPPType));
 		
 #endif
+
+		// store the property in the cached list
 		InClass->LinkedProperties.Add(Result);
 		(*LinkToProperty) = Result;
 	}
@@ -468,6 +463,17 @@ FProperty* URigVMMemoryStorageGeneratorClass::AddProperty(URigVMMemoryStorageGen
 	return Result;
 }
 
+void URigVMMemoryStorageGeneratorClass::RefreshLinkedProperties()
+{
+	LinkedProperties.Reset();
+	const FProperty* Property = CastField<FProperty>(ChildProperties);
+	while(Property)
+	{
+		LinkedProperties.Add(Property);
+		Property = CastField<FProperty>(Property->Next);
+	}
+}
+
 void URigVMMemoryStorageGeneratorClass::RefreshPropertyPaths()
 {
 	// Update the property paths based on the descriptions
@@ -566,6 +572,9 @@ bool URigVMMemoryStorage::CopyProperty(
 	check(InTargetPtr != nullptr);
 	check(InSourcePtr != nullptr);
 
+	// This block below is there to support Large World Coordinates (LWC).
+	// We allow to link float and double pins (single and arrays) so we need
+	// to support copying values between those as well.
 	if(!InTargetProperty->SameType(InSourceProperty))
 	{
 		if(const FFloatProperty* TargetFloatProperty = CastField<FFloatProperty>(InTargetProperty))
@@ -631,11 +640,14 @@ bool URigVMMemoryStorage::CopyProperty(
 				}
 			}
 		}
-		
+
+		// if we reach this we failed since we are trying to copy
+		// between two properties which are not compatible
 		ensure(InTargetProperty->SameType(InSourceProperty));
 		return false;
 	}
 
+	// rely on the core to copy the property contents
 	InTargetProperty->CopyCompleteValue(InTargetPtr, InSourcePtr);
 	return true;
 }
@@ -661,7 +673,7 @@ bool URigVMMemoryStorage::CopyProperty(
 		}
 
 		MemoryPtr = PropertyPath.GetData<uint8>(MemoryPtr, Property);
-		Property = PropertyPath.GetTargetProperty();
+		Property = PropertyPath.GetTailProperty();
 	};
 
 	uint8* SourcePtr = (uint8*)InSourcePtr;
@@ -693,16 +705,16 @@ bool URigVMMemoryStorage::CopyProperty(
 #if !UE_RIGVM_UCLASS_BASED_STORAGE_DISABLED
 
 bool URigVMMemoryStorage::CopyProperty(
-	FRigVMMemoryHandle& TargetHandle,
-	FRigVMMemoryHandle& SourceHandle)
+	FRigVMMemoryHandle& InTargetHandle,
+	FRigVMMemoryHandle& InSourceHandle)
 {
 	return CopyProperty(
-		TargetHandle.GetProperty(),
-		TargetHandle.GetData(false),
-		TargetHandle.GetPropertyPathRef(),
-		SourceHandle.GetProperty(),
-		SourceHandle.GetData(false),
-		SourceHandle.GetPropertyPathRef());
+		InTargetHandle.GetProperty(),
+		InTargetHandle.GetData(false),
+		InTargetHandle.GetPropertyPathRef(),
+		InSourceHandle.GetProperty(),
+		InSourceHandle.GetData(false),
+		InSourceHandle.GetPropertyPathRef());
 }
 
 #endif
diff --git a/Engine/Source/Runtime/RigVM/Private/RigVMCore/RigVMPropertyPath.cpp b/Engine/Source/Runtime/RigVM/Private/RigVMCore/RigVMPropertyPath.cpp
--- a/Engine/Source/Runtime/RigVM/Private/RigVMCore/RigVMPropertyPath.cpp
+++ b/Engine/Source/Runtime/RigVM/Private/RigVMCore/RigVMPropertyPath.cpp
@@ -19,15 +19,20 @@ FRigVMPropertyPath::FRigVMPropertyPath(const FProperty* InProperty, const FStrin
 	check(InProperty);
 	check(!InSegmentPath.IsEmpty())
 
-	const FProperty* Property = InProperty;
+	// Reformat the path. We'll remove any brackets and clean up
+	// double periods. Turns '[2].Translation.X' to '2.Translation.X'
 	FString WorkPath = InSegmentPath;
 	WorkPath = WorkPath.Replace(TEXT("["), TEXT("."));
 	WorkPath = WorkPath.Replace(TEXT("]"), TEXT("."));
 	WorkPath = WorkPath.Replace(TEXT(".."), TEXT("."));
 	WorkPath.TrimCharInline('.', nullptr);
 
+	// Traverse the provided segment path and build up the segments for
+	// the property path/
+	const FProperty* Property = InProperty;
 	while(!WorkPath.IsEmpty())
 	{
+		// split off the head
 		FString PathSegment, PathRemainder;
 		if(!WorkPath.Split(TEXT("."), &PathSegment, &PathRemainder))
 		{
@@ -73,6 +78,10 @@ FRigVMPropertyPath::FRigVMPropertyPath(const FProperty* InProperty, const FStrin
 			Segment.Property = MapProperty;
 			Property = ValueProperty;
 		}
+		else if(CastField<FSetProperty>(Property))
+		{
+			checkNoEntry();
+		}
 		else
 		{
 			Segments.Empty();
@@ -83,8 +92,8 @@ FRigVMPropertyPath::FRigVMPropertyPath(const FProperty* InProperty, const FStrin
 		WorkPath = PathRemainder;
 	}
 
+	// rebuild the path again from the retrieved segments 
 	Path.Empty();
-
 	for(const FRigVMPropertyPathSegment& Segment : Segments)
 	{
 		switch(Segment.Type)
@@ -122,27 +131,7 @@ FRigVMPropertyPath::FRigVMPropertyPath(const FRigVMPropertyPath& InOther)
 {
 }
 
-bool FRigVMPropertyPath::IsDirect() const
-{
-	for(const FRigVMPropertyPathSegment& Segment : Segments)
-	{
-		switch(Segment.Type)
-		{
-			case ERigVMPropertyPathSegmentType::ArrayElement:
-			case ERigVMPropertyPathSegmentType::MapValue:
-			{
-				return false;
-			}
-			default:
-			{
-				break;
-			}
-		}
-	}
-	return true;
-}
-
-const FProperty* FRigVMPropertyPath::GetTargetProperty() const
+const FProperty* FRigVMPropertyPath::GetTailProperty() const
 {
 	if(Segments.IsEmpty())
 	{
diff --git a/Engine/Source/Runtime/RigVM/Public/RigVMCore/RigVMMemoryStorage.h b/Engine/Source/Runtime/RigVM/Public/RigVMCore/RigVMMemoryStorage.h
--- a/Engine/Source/Runtime/RigVM/Public/RigVMCore/RigVMMemoryStorage.h
+++ b/Engine/Source/Runtime/RigVM/Public/RigVMCore/RigVMMemoryStorage.h
@@ -16,17 +16,28 @@
 
 /**
  * The FRigVMMemoryHandle is used to access the memory used within a URigMemoryStorage.
+ * The Memory Handle caches the pointer of the head property, and can rely on a
+ * RigVMPropertyPath to traverse towards a tail property.
+ * For example it can cache the pointer of a TArray<FTransform> property's memory
+ * and use a property path of '[2].Translation.X' to retrieve the element 2's translation
+ * X float / double.
+ * Additionally the FRigVMMemoryHandle can be used to access 'sliced' memory. Sliced
+ * memory in this case simply means another level of a nested TArray. So a 'sliced'
+ * FTransform is stored as a TArray<FTransform>, which the handle can traverse into and
+ * return the memory of a specific FTransform within that array.
  */
 struct FRigVMMemoryHandle
 {
 public:
 
+	// Default constructor
 	FORCEINLINE_DEBUGGABLE FRigVMMemoryHandle()
 		: Ptr(nullptr)
 		, Property(nullptr) 
 		, PropertyPath(nullptr)
 	{}
 
+	// Constructor from complete data
 	FORCEINLINE_DEBUGGABLE FRigVMMemoryHandle(uint8* InPtr, const FProperty* InProperty,  const FRigVMPropertyPath* InPropertyPath)
 		: Ptr(InPtr)
 		, Property(InProperty)
@@ -34,33 +45,48 @@ public:
 	{
 	}
 
+	/**
+	 * Returns the cached pointer stored within the handle.
+	 * @param bFollowPropertyPath If set to true the memory handle will traverse the memory using the property path
+	 * @param InSliceIndex If this is != INDEX_NONE the memory handle will return the slice of the memory requested
+	 * @return The traversed memory cached by this handle.
+	 */
 	FORCEINLINE uint8* GetData(bool bFollowPropertyPath = false, int32 InSliceIndex = INDEX_NONE)
 	{
 		return GetData_Internal(bFollowPropertyPath, InSliceIndex);
 	}
 
+	/**
+	 * Returns the cached pointer stored within the handle.
+	 * @param bFollowPropertyPath If set to true the memory handle will traverse the memory using the property path
+	 * @param InSliceIndex If this is != INDEX_NONE the memory handle will return the slice of the memory requested
+	 * @return The traversed memory cached by this handle (const)
+	 */
 	FORCEINLINE const uint8* GetData(bool bFollowPropertyPath = false, int32 InSliceIndex = INDEX_NONE) const
 	{
 		return GetData_Internal(bFollowPropertyPath, InSliceIndex);
 	}
 
+	// Returns the head property of this handle
 	FORCEINLINE const FProperty* GetProperty() const
 	{
 		return Property;
 	}
 
+	// Returns the [optional] property path used within this handle
 	FORCEINLINE const FRigVMPropertyPath* GetPropertyPath() const
 	{
 		return PropertyPath;
 	}
 
+	// Returns the [optional] property path used within this handle (ref)
 	FORCEINLINE const FRigVMPropertyPath& GetPropertyPathRef() const
 	{
 		if(PropertyPath)
 		{
 			return *PropertyPath;
 		}
-		return EmptyPropertyPath;
+		return FRigVMPropertyPath::Empty;
 	}
 
 private:
@@ -85,18 +111,24 @@ private:
 
 			return ArrayHelper.GetRawPtr(InSliceIndex);
 		}
-		
+
+		// traverse the property path to the tail property
+		// and return its memory instead.
 		if(bFollowPropertyPath && PropertyPath != nullptr)
 		{
 			return PropertyPath->GetData<uint8>(Ptr, Property);
 		}
 		return Ptr;
 	}
-	
+
+	// The pointer of the memory of the head property
 	uint8* Ptr;
+
+	// The head property used by this handle
 	const FProperty* Property;
+
+	// The [optional] property path used by this handle
 	const FRigVMPropertyPath* PropertyPath;
-	static const FRigVMPropertyPath EmptyPropertyPath;
 
 	friend class URigVM;
 };
@@ -107,15 +139,33 @@ private:
 /// Property Management
 //////////////////////////////////////////////////////////////////////////////
 
+/**
+ * The property description is used to provide all required information
+ * to create a property for the memory storage class.
+ */
 struct RIGVM_API FRigVMPropertyDescription
 {
+public:
+	
+	// The name of the property to create
 	FName Name;
+
+	// The property to base a new property off of
 	const FProperty* Property;
+
+	// The complete CPP type to base a new property off of (for ex: 'TArray<TArray<FVector>>')
 	FString CPPType;
+
+	// The tail CPP Type object, for example the UScriptStruct for a struct 
 	UObject* CPPTypeObject;
+
+	// A list of containers to use for this property, for example [Array, Array]
 	TArray<EPinContainerType> Containers;
+
+	// The default value to use for this property (for example: '(((X=1.000000, Y=2.000000, Z=3.000000)))')
 	FString DefaultValue;
 
+	// Default constructor
 	FRigVMPropertyDescription()
 		: Name(NAME_None)
 		, Property(nullptr)
@@ -124,27 +174,39 @@ struct RIGVM_API FRigVMPropertyDescription
 		, Containers()
 		, DefaultValue()
 	{}
+
+	// Constructor from an existing property
 	FRigVMPropertyDescription(const FProperty* InProperty, const FString& InDefaultValue, const FName& InName = NAME_None);
+
+	// Constructor from complete data
 	FRigVMPropertyDescription(const FName& InName, const FString& InCPPType, UObject* InCPPTypeObject, const FString& InDefaultValue);
 
+	// Returns a sanitized, valid name to use for a new property 
 	static FName SanitizeName(const FName& InName);
+
+	// Sanitize the name of this description in line
 	void SanitizeName();
 
+	// Returns true if this property description is valid
 	bool IsValid() const { return !Name.IsNone(); }
-	int32 NumContainers() const { return Containers.Num(); }
-	bool HasContainer() const { return NumContainers() > 0; }
-	bool IsArray(int32 InContainerIndex = 0) const { return Containers[InContainerIndex] == EPinContainerType::Array; }
-	bool IsMap(int32 InContainerIndex = 0) const { return Containers[InContainerIndex] == EPinContainerType::Map; }
-	bool IsSet(int32 InContainerIndex = 0) const { return Containers[InContainerIndex] == EPinContainerType::Set; }
 
-	FString GetBaseCPPType() const;
+	// Returns the CPP type of the tail property, for ex: '[2].Translation' it is 'FVector'
+	FString GetTailCPPType() const;
 
+private:
+	
 	static const FString ArrayPrefix;
 	static const FString MapPrefix;
-	static const FString SetPrefix;
 	static const FString ContainerSuffix;
 };
 
+/**
+ * The URigVMMemoryStorageGeneratorClass is used to create / represent heterogeneous
+ * memory storages. The generator can produce a UClass which contains a series of
+ * properties. This UClass is then used to instantiate URigVMMemoryStorage objects
+ * to be consumed by the RigVM. The memory storage objects can contain the literals
+ * / constant values used by the virtual machine or work state.
+ */
 UCLASS()
 class RIGVM_API URigVMMemoryStorageGeneratorClass :
 	public UClass
@@ -153,6 +215,7 @@ class RIGVM_API URigVMMemoryStorageGeneratorClass :
 
 public:
 
+	// Default constructor
 	URigVMMemoryStorageGeneratorClass()
 		: Super()
 		, MemoryType(ERigVMMemoryType::Literal)
@@ -161,38 +224,67 @@ public:
 	// UClass overrides
 	void PurgeClass(bool bRecompilingOnLoad) override;
 	void Link(FArchive& Ar, bool bRelinkExistingProperties) override;
-	void Serialize(FArchive& Ar) override;;
+	void Serialize(FArchive& Ar) override;
+	virtual void PostLoad() override;
 
-	// URigVMMemoryStorageGeneratorClass specific
+	// Returns the name of a storage class for a given memory type. The name is unique within the package.
 	static FString GetClassName(ERigVMMemoryType InMemoryType);
+
+	// Returns an existing class for a memory type within the package (or nullptr) 
 	static URigVMMemoryStorageGeneratorClass* GetStorageClass(UObject* InOuter, ERigVMMemoryType InMemoryType);
+
+	/**
+	 * Creates a new class provided a list of properties and property paths
+	 * @param InOuter The package or the outer object of this class
+	 * @param InMemoryType The memory type to use for this class (affects class name)
+	 * @param InProperties The descriptions for the properties to create
+	 * @param InPropertyPaths The descriptions for the property paths to create
+	 */
 	static URigVMMemoryStorageGeneratorClass* CreateStorageClass(
 		UObject* InOuter,
 		ERigVMMemoryType InMemoryType,
 		const TArray<FRigVMPropertyDescription>& InProperties,
 		const TArray<FRigVMPropertyPathDescription>& InPropertyPaths = TArray<FRigVMPropertyPathDescription>());
 
+	// Returns the type of memory of this class (literal, work, etc)
 	ERigVMMemoryType GetMemoryType() const { return MemoryType; }
+
+	// The properties stored within this class
 	const TArray<const FProperty*>& GetProperties() const { return LinkedProperties; }
+
+	// The property paths stored within this class
 	const TArray<FRigVMPropertyPath>& GetPropertyPaths() const { return PropertyPaths; }
 
 protected:
 
+	// Adds a single property to a class.
 	static FProperty* AddProperty(URigVMMemoryStorageGeneratorClass* InClass, const FRigVMPropertyDescription& InProperty, FField** LinkToProperty = nullptr);
 
 private:
 
+	void RefreshLinkedProperties();
 	void RefreshPropertyPaths();
 
+	// The type of memory of this class
 	ERigVMMemoryType MemoryType;
+
+	// A cached list of all linked properties (created by RefreshLinkedProperties)
 	TArray<const FProperty*> LinkedProperties;
+
+	// A cached list of all property paths (created by RefreshPropertyPaths)
 	TArray<FRigVMPropertyPath> PropertyPaths;
+
+	// A list of decriptions for the property paths - used for serialization
 	TArray<FRigVMPropertyPathDescription> PropertyPathDescriptions;
 
 	friend class URigVMMemoryStorage;
 	friend class URigVMCompiler;
 };
 
+/**
+ * The URigVMMemoryStorage represents an instance of heterogeneous memory.
+ * The memory layout is defined by the URigVMMemoryStorageGeneratorClass.
+ */
 UCLASS()
 class RIGVM_API URigVMMemoryStorage : public UObject
 {
@@ -204,6 +296,7 @@ public:
 	/// Memory Access
 	//////////////////////////////////////////////////////////////////////////////
 
+	// Returns the memory type of this memory
 	FORCEINLINE ERigVMMemoryType GetMemoryType() const
 	{
 		if(URigVMMemoryStorageGeneratorClass* Class = Cast<URigVMMemoryStorageGeneratorClass>(GetClass()))
@@ -213,94 +306,80 @@ public:
 		// empty debug containers don't have a generator class
 		return ERigVMMemoryType::Debug;
 	}
-	
+
+	// Returns the number of properties stored in this instance
 	FORCEINLINE int32 Num() const
 	{
 		return GetProperties().Num();
 	}
 
+	// Returns true if a provided property index is valid
 	FORCEINLINE bool IsValidIndex(int32 InIndex) const
 	{
 		return GetProperties().IsValidIndex(InIndex);
 	}
 
+	// Returns the properties provided by this instance
 	const TArray<const FProperty*>& GetProperties() const;
+
+	// Returns the property paths provided by this instance
 	const TArray<FRigVMPropertyPath>& GetPropertyPaths() const;
 
+	// Returns the index of a property given the property itself
 	int32 GetPropertyIndex(const FProperty* InProperty) const;
 
+	// Returns the index of a property given its name
 	int32 GetPropertyIndexByName(const FName& InName) const;
 
+	// Returns a property given its index
 	FORCEINLINE const FProperty* GetProperty(int32 InPropertyIndex) const
 	{
 		return GetProperties()[InPropertyIndex];
 	}
 
+	// Returns a property given its name (or nullptr if the name wasn't found)
 	const FProperty* FindPropertyByName(const FName& InName) const;
 
+	// Creates and returns a new operand for a property (and optionally a property path)
 	FRigVMOperand GetOperand(int32 InPropertyIndex, int32 InPropertyPathIndex = INDEX_NONE) const;
+
+	// Creates and returns a new operand for a property (and optionally a property path)
 	FRigVMOperand GetOperandByName(const FName& InName, int32 InPropertyPathIndex = INDEX_NONE) const;
-	
+
+	// Returns true if the property at a given index is a TArray
 	FORCEINLINE bool IsArray(int32 InPropertyIndex) const
 	{
 		return GetProperty(InPropertyIndex)->IsA<FArrayProperty>();
 	}
-	
-	FORCEINLINE bool IsArrayByName(const FName& InName) const
-	{
-		if(const FProperty* Property = FindPropertyByName(InName))
-		{
-			return Property->IsA<FArrayProperty>();
-		}
-		return false;
-	}
-	
+
+	// Returns true if the property at a given index is a TMap
 	FORCEINLINE bool IsMap(int32 InPropertyIndex) const
 	{
 		return GetProperty(InPropertyIndex)->IsA<FMapProperty>();
 	}
-	
-	FORCEINLINE bool IsMapByName(const FName& InName) const
-	{
-		if(const FProperty* Property = FindPropertyByName(InName))
-		{
-			return Property->IsA<FMapProperty>();
-		}
-		return false;
-	}
-	
-	FORCEINLINE bool IsSet(int32 InPropertyIndex) const
-	{
-		return GetProperty(InPropertyIndex)->IsA<FSetProperty>();
-	}
-	
-	FORCEINLINE bool IsSetByName(const FName& InName) const
-	{
-		if(const FProperty* Property = FindPropertyByName(InName))
-		{
-			return Property->IsA<FSetProperty>();
-		}
-		return false;
-	}
 
+	// Returns the memory for a property given its index
 	template<typename T>
 	FORCEINLINE T* GetData(int32 InPropertyIndex)
 	{
 		const TArray<const FProperty*>& Properties = GetProperties();
-		if(Properties.IsValidIndex(InPropertyIndex))
-		{
-			return Properties[InPropertyIndex]->ContainerPtrToValuePtr<T>(this);
-		}
-		return nullptr;
+		check(Properties.IsValidIndex(InPropertyIndex));
+		return Properties[InPropertyIndex]->ContainerPtrToValuePtr<T>(this);
 	}
 
+	// Returns the memory for a property given its name (or nullptr)
 	template<typename T>
 	FORCEINLINE T* GetDataByName(const FName& InName)
 	{
 		const int32 PropertyIndex = GetPropertyIndexByName(InName);
+		if(PropertyIndex == INDEX_NONE)
+		{
+			return nullptr;
+		}
 		return GetData<T>(PropertyIndex);
 	}
 
+	// Returns the memory for a property given its index and a matching property path
 	template<typename T>
 	FORCEINLINE T* GetData(int32 InPropertyIndex, const FRigVMPropertyPath& InPropertyPath)
 	{
@@ -308,13 +387,19 @@ public:
 		return InPropertyPath.GetData<T>(GetData<uint8>(InPropertyIndex), Property);
 	}
 
+	// Returns the memory for a property given its name and a matching property path (or nullptr)
 	template<typename T>
 	FORCEINLINE T* GetDataByName(const FName& InName, const FRigVMPropertyPath& InPropertyPath)
 	{
 		const int32 PropertyIndex = GetPropertyIndexByName(InName);
+		if(PropertyIndex == INDEX_NONE)
+		{
+			return nullptr;
+		}
 		return GetData<T>(PropertyIndex, InPropertyPath);
 	}
 
+	// Returns the memory for a property (and optionally a property path) given an operand
 	template<typename T>
 	FORCEINLINE T* GetData(const FRigVMOperand& InOperand)
 	{
@@ -332,47 +417,58 @@ public:
 		return GetData<T>(PropertyIndex, GetPropertyPaths()[PropertyPathIndex]); 
 	}
 
+	// Returns the ref of an element stored at a given property index
 	template<typename T>
 	FORCEINLINE T& GetRef(int32 InPropertyIndex)
 	{
 		return *GetData<T>(InPropertyIndex);
 	}
 
+	// Returns the ref of an element stored at a given property name (throws if name is invalid)
 	template<typename T>
 	FORCEINLINE T& GetRefByName(const FName& InName)
 	{
 		return *GetDataByName<T>(InName);
 	}
 
+	// Returns the ref of an element stored at a given property index and a property path
 	template<typename T>
 	FORCEINLINE T& GetRef(int32 InPropertyIndex, const FRigVMPropertyPath& InPropertyPath)
 	{
 		return *GetData<T>(InPropertyIndex, InPropertyPath);
 	}
 
+	// Returns the ref of an element stored at a given property name and a property path (throws if name is invalid)
 	template<typename T>
 	FORCEINLINE T& GetRefByName(const FName& InName, const FRigVMPropertyPath& InPropertyPath)
 	{
 		return *GetDataByName<T>(InName, InPropertyPath);
 	}
 
+	// Returns the ref of an element stored for a given operand
 	template<typename T>
 	FORCEINLINE T& GetRef(const FRigVMOperand& InOperand)
 	{
 		return *GetData<T>(InOperand);
 	}
 
+	// Returns the exported text for a given property index
 	FString GetDataAsString(int32 InPropertyIndex);
-	FString GetDataAsString(const FRigVMOperand& InOperand);
-	
+
+	// Returns the exported text for given property name 
 	FORCEINLINE FString GetDataAsStringByName(const FName& InName)
 	{
 		const int32 PropertyIndex = GetPropertyIndexByName(InName);
 		return GetDataAsString(PropertyIndex);
 	}
 
+	// Returns the exported text for a given operand
+	FString GetDataAsString(const FRigVMOperand& InOperand);
+
+	// Sets the content of a property by index given an exported string. Returns true if succeeded
 	bool SetDataFromString(int32 InPropertyIndex, const FString& InValue);
 
+	// Sets the content of a property by name given an exported string. Returns true if succeeded
 	FORCEINLINE bool SetDataFromStringByName(const FName& InName, const FString& InValue)
 	{
 		const int32 PropertyIndex = GetPropertyIndexByName(InName);
@@ -380,9 +476,11 @@ public:
 	}
 
 #if !UE_RIGVM_UCLASS_BASED_STORAGE_DISABLED
-	
+
+	// Returns the handle for a given property by index (and optionally property path)
 	FRigVMMemoryHandle GetHandle(int32 InPropertyIndex, const FRigVMPropertyPath* InPropertyPath = nullptr);
 
+	// Returns the handle for a given property by name (and optionally property path)
 	FORCEINLINE FRigVMMemoryHandle GetHandleByName(const FName& InName, const FRigVMPropertyPath* InPropertyPath = nullptr)
 	{
 		const int32 PropertyIndex = GetPropertyIndexByName(InName);
@@ -390,13 +488,32 @@ public:
 	}
 	
 #endif
-	
+
+	/**
+	 * Copies the content of a source property and memory into the memory of a target property
+	 * @param InTargetProperty The target property to copy into
+	 * @param InTargetPtr The memory of the target value to copy into
+	 * @param InSourceProperty The source property to copy from
+	 * @param InSourcePtr The memory of the source value to copy from
+	 * @return true if the copy operations was successful
+	 */ 
 	static bool CopyProperty(
 		const FProperty* InTargetProperty,
 		uint8* InTargetPtr,
 		const FProperty* InSourceProperty,
 		const uint8* InSourcePtr);
 
+	/**
+	 * Copies the content of a source property and memory into the memory of a target property.
+	 * This variant supports property paths for both target and source.
+	 * @param InTargetProperty The target property to copy into
+	 * @param InTargetPtr The memory of the target value to copy into
+	 * @param InTargetPropertyPath The property path to use when traversing the target memory.
+	 * @param InSourceProperty The source property to copy from
+	 * @param InSourcePtr The memory of the source value to copy from
+	 * @param InSourcePropertyPath The property path to use when traversing the source memory.
+	 * @return true if the copy operations was successful
+	 */ 
 	static bool CopyProperty(
 		const FProperty* InTargetProperty,
 		uint8* InTargetPtr,
@@ -405,6 +522,17 @@ public:
 		const uint8* InSourcePtr,
 		const FRigVMPropertyPath& InSourcePropertyPath);
 
+	/**
+	 * Copies the content of a source property and memory into the memory of a target property.
+	 * This variant retrieves the memory pointers given target and source storage objects.
+	 * @param InTargetStorage The target property to copy into
+	 * @param InTargetPropertyIndex The index of the value to use on the InTargetStorage
+	 * @param InTargetPropertyPath The property path to use when traversing the target memory.
+	 * @param InSourceStorage The source property to copy from
+	 * @param InSourcePropertyIndex The index of the value to use on the InSourceStorage
+	 * @param InSourcePropertyPath The property path to use when traversing the source memory.
+	 * @return true if the copy operations was successful
+	 */ 
 	static bool CopyProperty(
 		URigVMMemoryStorage* InTargetStorage,
 		int32 InTargetPropertyIndex,
@@ -414,10 +542,17 @@ public:
 		const FRigVMPropertyPath& InSourcePropertyPath);
 
 #if !UE_RIGVM_UCLASS_BASED_STORAGE_DISABLED
-	
+
+	/**
+	 * Copies the content of a source property and memory into the memory of a target property.
+	 * This variant retrieves the memory pointers given target and source storage objects.
+	 * @param InTargetHandle The handle of the target memory to copy into
+	 * @param InSourceHandle The handle of the source memory to copy from
+	 * @return true if the copy operations was successful
+	 */ 
 	static bool CopyProperty(
-		FRigVMMemoryHandle& TargetHandle,
-		FRigVMMemoryHandle& SourceHandle);
+		FRigVMMemoryHandle& InTargetHandle,
+		FRigVMMemoryHandle& InSourceHandle);
 	
 #endif
 
diff --git a/Engine/Source/Runtime/RigVM/Public/RigVMCore/RigVMPropertyPath.h b/Engine/Source/Runtime/RigVM/Public/RigVMCore/RigVMPropertyPath.h
--- a/Engine/Source/Runtime/RigVM/Public/RigVMCore/RigVMPropertyPath.h
+++ b/Engine/Source/Runtime/RigVM/Public/RigVMCore/RigVMPropertyPath.h
@@ -5,62 +5,120 @@
 #include "CoreMinimal.h"
 #include "UObject/UnrealType.h"
 
+/**
+ * The RigVMPropertyPathDescription is used to provide all of the information
+ * necessary to describe a property path for creating it. This is used by
+ * the RigVMCompiler, which collects all descriptions first and then passes
+ * them onto the RigVMStorageGeneratorClass for construction. Furthermore
+ * the RigVMPropertyPathDescription is used to serialize property paths and
+ * recreate them on load.
+ */
 struct RIGVM_API FRigVMPropertyPathDescription
 {
+	// The index of the property this property path belongs to.
 	int32 PropertyIndex;
-	FString RootCPPType;
+
+	// The CPP type of the head property.
+	// For Transform.Translation.X the HeadCPPType is 'FTransform' 
+	FString HeadCPPType;
+
+	// The segment path of the properties below the head property.
+	// For Transform.Translation.X the SegmentPath is 'Translation.X' 
 	FString SegmentPath;
-	
+
+	// Default constructor
 	FRigVMPropertyPathDescription()
 		: PropertyIndex(INDEX_NONE)
-		, RootCPPType()
+		, HeadCPPType()
 		, SegmentPath()
 	{}
 
-	FRigVMPropertyPathDescription(int32 InPropertyIndex, const FString& InRootCPPType, const FString& InSegmentPath)
+	// Constructor from complete data
+	FRigVMPropertyPathDescription(int32 InPropertyIndex, const FString& InHeadCPPType, const FString& InSegmentPath)
 		: PropertyIndex(InPropertyIndex)
-		, RootCPPType(InRootCPPType)
+		, HeadCPPType(InHeadCPPType)
 		, SegmentPath(InSegmentPath)
 	{}
 
+	// Archive operator for serialization
 	FORCEINLINE friend FArchive& operator<<(FArchive& Ar, FRigVMPropertyPathDescription& Path)
 	{
 		Ar << Path.PropertyIndex;
-		Ar << Path.RootCPPType;
+		Ar << Path.HeadCPPType;
 		Ar << Path.SegmentPath;
 		return Ar;
 	}
 };
 
+/**
+ * The RigVMPropertyPathSegmentType is used to differentiate
+ * segments under a property path.
+ */
 enum RIGVM_API ERigVMPropertyPathSegmentType
 {
-	StructMember,
-	ArrayElement,
-	MapValue
+	StructMember, // The segment represents a member of a struct
+	ArrayElement, // The segment represents an element of an array
+	MapValue // The segments represents a value within a map. Typically maps supported use FName as the key.
 };
 
+/**
+ * The RigVMPropertyPathSegment represents a single step needed to go from
+ * the head property of a property path to the tail.
+ */
 struct RIGVM_API FRigVMPropertyPathSegment
 {
+	// The type of segment / step.
 	ERigVMPropertyPathSegmentType Type;
+
+	// The name of the segment - can be the name of the struct member of key for a map element.
 	FName Name;
+
+	// The index of the segment - used to represent the array index for array elements.
 	int32 Index;
+
+	// The property for this segment - on the chain towards the tail property. 
 	const FProperty* Property;
 };
 
+/**
+ * The RigVMPropertyPath is used to access a memory pointer for a tail property
+ * given the memory of the head property. A property path allows to traverse
+ * struct members, array elements and map values.
+ * For example given an array of transforms the property path
+ * '[2].Translation.X'
+ * is represented as three segments (1 array element and 2 struct members).
+ */
 class RIGVM_API FRigVMPropertyPath
 {
 public:
 
+	// Default constructor
 	FRigVMPropertyPath();
+
+	/**
+	 * Constructor from complete data
+	 * @param InProperty The head property to base this property path off of
+	 * @param InSegmentPath A . separated string providing all of the segments of this path
+	 */
 	FRigVMPropertyPath(const FProperty* InProperty, const FString& InSegmentPath);
+
+	// Copy constructor
 	FRigVMPropertyPath(const FRigVMPropertyPath& InOther);
 
+	// Returns the property path as a string (the sanitized SegmentPath)
 	FORCEINLINE const FString& ToString() const { return Path; }
+
+	// Returns the number of segments
 	FORCEINLINE int32 Num() const { return Segments.Num(); }
+
+	// Returns true if the property path is valid (has any segments)
 	FORCEINLINE bool IsValid() const { return Num() > 0; }
+
+	// Returns true if the property is empty (no segments)
 	FORCEINLINE bool IsEmpty() const { return Num() == 0; }
+
+	// Index based access operator for a segment
 	FORCEINLINE const FRigVMPropertyPathSegment& operator[](int32 InIndex) const { return Segments[InIndex]; }
-	bool IsDirect() const;
 
 	FORCEINLINE bool operator ==(const FRigVMPropertyPath& Other) const
 	{
@@ -82,12 +140,18 @@ public:
 		return GetTypeHash(this) < GetTypeHash(Other);
 	}
 
+	/**
+	 * Accessor for traversing the path and returning the memory of the tail
+	 * @param InPtr The memory of the head property
+	 * @param InProperty The head property to traverse off of
+	 */
 	template<typename T>
 	FORCEINLINE T* GetData(uint8* InPtr, const FProperty* InProperty) const
 	{
 		return (T*)GetData_Internal(InPtr, InProperty);
 	}
 
+	// Returns a unique hash for the property path
 	friend FORCEINLINE uint32 GetTypeHash(const FRigVMPropertyPath& InPropertyPath)
 	{
 		if(InPropertyPath.IsEmpty())
@@ -101,15 +165,21 @@ public:
 		);
 	}
 
-	const FProperty* GetTargetProperty() const;
-	
+	// Returns the property of the last segment
+	const FProperty* GetTailProperty() const;
+
+	// Static empty property path (used for comparisons)
 	static FRigVMPropertyPath Empty;
 
 private:
 
+	// Internal traversal method
 	uint8* GetData_Internal(uint8* InPtr, const FProperty* InProperty) const;
 
+	// Stored string representation of the segment path
 	FString Path;
+
+	// Resolved segments of the property path
 	TArray<FRigVMPropertyPathSegment> Segments;
 };
 
