Control Rig: RigVM Memory upgrade: Remaining fixes

diff --git a/Engine/Plugins/Animation/ControlRig/Source/ControlRigEditor/Private/Editor/DetailsViewWrapperObject.cpp b/Engine/Plugins/Animation/ControlRig/Source/ControlRigEditor/Private/Editor/DetailsViewWrapperObject.cpp
--- a/Engine/Plugins/Animation/ControlRig/Source/ControlRigEditor/Private/Editor/DetailsViewWrapperObject.cpp
+++ b/Engine/Plugins/Animation/ControlRig/Source/ControlRigEditor/Private/Editor/DetailsViewWrapperObject.cpp
@@ -40,6 +40,7 @@ UClass* UDetailsViewWrapperObject::GetClassForStruct(UScriptStruct* InStruct)
 	WrapperClass->SetSuperStruct(SuperClass);
 	WrapperClass->ClassWithin = UObject::StaticClass();
 	WrapperClass->ClassConfigName = SuperClass->ClassConfigName;
+	WrapperClass->ClassFlags |= CLASS_NotPlaceable | CLASS_Hidden;
 	WrapperClass->SetMetaData(TEXT("DisplayName"), *InStruct->GetDisplayNameText().ToString());
 
 	struct Local
diff --git a/Engine/Source/Developer/RigVMDeveloper/Private/RigVMCompiler/RigVMAST.cpp b/Engine/Source/Developer/RigVMDeveloper/Private/RigVMCompiler/RigVMAST.cpp
--- a/Engine/Source/Developer/RigVMDeveloper/Private/RigVMCompiler/RigVMAST.cpp
+++ b/Engine/Source/Developer/RigVMDeveloper/Private/RigVMCompiler/RigVMAST.cpp
@@ -1881,7 +1881,7 @@ bool FRigVMParserAST::FoldConstantValuesToLiterals(URigVMGraph* InGraph, URigVMC
 
 	// build the VM to run this AST
 	TMap<FString, FRigVMOperand> Operands;
-	URigVM* TempVM = NewObject<URigVM>(GetTransientPackage());
+	URigVM* TempVM = NewObject<URigVM>(InGraph);
 	
 	URigVMCompiler* TempCompiler = NewObject<URigVMCompiler>(GetTransientPackage());
 	TempCompiler->Settings.SetupNodeInstructionIndex = false;
@@ -2076,6 +2076,9 @@ bool FRigVMParserAST::FoldConstantValuesToLiterals(URigVMGraph* InGraph, URigVMC
 		}
 	}
 
+	TempVM->Rename(nullptr, GetTransientPackage(), REN_ForceNoResetLoaders | REN_DoNotDirty | REN_DontCreateRedirectors | REN_NonTransactional);
+	TempVM->MarkPendingKill();
+
 	RemoveExpressions(ExpressionsToRemove);
 
 	return ExpressionsToRemove.Num() > 0;
diff --git a/Engine/Source/Developer/RigVMDeveloper/Private/RigVMCompiler/RigVMCompiler.cpp b/Engine/Source/Developer/RigVMDeveloper/Private/RigVMCompiler/RigVMCompiler.cpp
--- a/Engine/Source/Developer/RigVMDeveloper/Private/RigVMCompiler/RigVMCompiler.cpp
+++ b/Engine/Source/Developer/RigVMDeveloper/Private/RigVMCompiler/RigVMCompiler.cpp
@@ -384,43 +384,18 @@ bool URigVMCompiler::Compile(URigVMGraph* InGraph, URigVMController* InControlle
 
 	for(ERigVMMemoryType MemoryType : MemoryTypes)
 	{
-		URigVMMemoryStorage** MemoryStorageObject = nullptr;
-		switch(MemoryType)
-		{
-			case ERigVMMemoryType::Literal:
-			{
-				MemoryStorageObject = &WorkData.VM->LiteralMemoryStorageObject;
-				break;
-			}
-			case ERigVMMemoryType::Work:
-			{
-				MemoryStorageObject = &WorkData.VM->WorkMemoryStorageObject;
-				break;
-			}
-			case ERigVMMemoryType::Debug:
-			{
-				MemoryStorageObject = &WorkData.VM->DebugMemoryStorageObject;
-				break;
-			}
-			default:
-			{
-				checkNoEntry();
-				break;
-			}
-		}
-
 		TArray<FRigVMPropertyDescription>* Properties = WorkData.PropertyDescriptions.Find(MemoryType);
 		if(Properties == nullptr)
 		{
-			*MemoryStorageObject = NewObject<URigVMMemoryStorage>(WorkData.VM, NAME_None, RF_Public);
 			continue;
 		}
 
 		UPackage* Package = InGraph->GetOutermost();
-		UClass* MemoryClass = URigVMMemoryStorageGeneratorClass::CreateStorageClass(Package, MemoryType, *Properties);
-		*MemoryStorageObject = NewObject<URigVMMemoryStorage>(WorkData.VM, MemoryClass, NAME_None, RF_Public); 
+		URigVMMemoryStorageGeneratorClass::CreateStorageClass(Package, MemoryType, *Properties);
 	}
 
+	WorkData.VM->ClearMemory();
+
 #endif
 
 	WorkData.bSetupMemory = false;
diff --git a/Engine/Source/Runtime/RigVM/Private/RigVMCore/RigVM.cpp b/Engine/Source/Runtime/RigVM/Private/RigVMCore/RigVM.cpp
--- a/Engine/Source/Runtime/RigVM/Private/RigVMCore/RigVM.cpp
+++ b/Engine/Source/Runtime/RigVM/Private/RigVMCore/RigVM.cpp
@@ -144,19 +144,6 @@ void URigVM::Save(FArchive& Ar)
 
 #else
 
-	bool bWorkMemoryStorageObjectValid = WorkMemoryStorageObject != nullptr;
-	Ar << bWorkMemoryStorageObjectValid;
-	if(bWorkMemoryStorageObjectValid)
-	{
-		WorkMemoryStorageObject->Serialize(Ar);
-	}
-	bool bLiteralMemoryStorageObjectValid = LiteralMemoryStorageObject != nullptr;
-	Ar << bLiteralMemoryStorageObjectValid;
-	if(bLiteralMemoryStorageObjectValid)
-	{
-		LiteralMemoryStorageObject->Serialize(Ar);
-	}
-
 	Ar << ExternalPropertyPathDescriptions;
 	
 #endif
@@ -204,51 +191,10 @@ void URigVM::Load(FArchive& Ar)
 
 #if !UE_RIGVM_UCLASS_BASED_STORAGE_DISABLED
 
-	if(WorkMemoryStorageObject == nullptr)
-	{
-		UClass* StorageClass = URigVMMemoryStorageGeneratorClass::GetStorageClass(GetOutermost(), ERigVMMemoryType::Work);
-		if(StorageClass)
-		{
-			WorkMemoryStorageObject = NewObject<URigVMMemoryStorage>(this, StorageClass, NAME_None, RF_Public);
-		}
-		else
-		{
-			WorkMemoryStorageObject = NewObject<URigVMMemoryStorage>(this, NAME_None, RF_Public);
-		}
-	}
-	if(LiteralMemoryStorageObject == nullptr)
-	{
-		UClass* StorageClass = URigVMMemoryStorageGeneratorClass::GetStorageClass(GetOutermost(), ERigVMMemoryType::Literal);
-		if(StorageClass)
-		{
-			LiteralMemoryStorageObject = NewObject<URigVMMemoryStorage>(this, StorageClass, NAME_None, RF_Public);
-		}
-		else
-		{
-			LiteralMemoryStorageObject = NewObject<URigVMMemoryStorage>(this, NAME_None, RF_Public);
-		}
-	}
-	if(DebugMemoryStorageObject == nullptr)
-	{
-		DebugMemoryStorageObject = NewObject<URigVMMemoryStorage>(this, NAME_None, RF_Public);
-	}
-
-	bool bWorkMemoryStorageObjectValid = false;
-	Ar << bWorkMemoryStorageObjectValid;
-	if(bWorkMemoryStorageObjectValid)
-	{
-		WorkMemoryStorageObject->Serialize(Ar);
-	}
-	
-	bool bLiteralMemoryStorageObjectValid = false;
-	Ar << bLiteralMemoryStorageObjectValid;
-	if(bLiteralMemoryStorageObjectValid)
-	{
-		LiteralMemoryStorageObject->Serialize(Ar);
-	}
+	// requesting the memory types will create them
+	ClearMemory();
 	
 	Ar << ExternalPropertyPathDescriptions;
-
 	Ar << FunctionNamesStorage;
 	Ar << ByteCodeStorage;
 	Ar << Parameters;
@@ -601,6 +547,8 @@ void URigVM::CopyFrom(URigVM* InVM, bool bDeferCopy, bool bReferenceLiteralMemor
 			}
 			else if(TargetMemory->GetClass() != SourceMemory->GetClass())
 			{
+				TargetMemory->Rename(nullptr, GetTransientPackage(), REN_ForceNoResetLoaders | REN_DoNotDirty | REN_DontCreateRedirectors | REN_NonTransactional);
+				TargetMemory->MarkPendingKill();
 				TargetMemory = NewObject<URigVMMemoryStorage>(Outer, SourceMemory->GetClass());
 			}
 
@@ -624,9 +572,11 @@ void URigVM::CopyFrom(URigVM* InVM, bool bDeferCopy, bool bReferenceLiteralMemor
 		}
 	};
 
-	CopyMemoryStorage(WorkMemoryStorageObject, InVM->WorkMemoryStorageObject, this);
-	CopyMemoryStorage(LiteralMemoryStorageObject, InVM->LiteralMemoryStorageObject, this);
-	CopyMemoryStorage(DebugMemoryStorageObject, InVM->DebugMemoryStorageObject, this);
+	// we don't need to copy the literals since they are shared
+	// between all instances of the VM
+	LiteralMemoryStorageObject = Cast<URigVMMemoryStorage>(InVM->GetLiteralMemory()->GetClass()->GetDefaultObject());
+	CopyMemoryStorage(WorkMemoryStorageObject, InVM->GetWorkMemory(), this);
+	CopyMemoryStorage(DebugMemoryStorageObject, InVM->GetDebugMemory(), this);
 
 	ExternalPropertyPathDescriptions = InVM->ExternalPropertyPathDescriptions;
 	ExternalPropertyPaths.Reset();
@@ -702,21 +652,82 @@ FString URigVM::GetRigVMFunctionName(int32 InFunctionIndex) const
 
 #if !UE_RIGVM_UCLASS_BASED_STORAGE_DISABLED
 
-URigVMMemoryStorage* URigVM::GetMemoryByType(ERigVMMemoryType InMemoryType) const
+URigVMMemoryStorage* URigVM::GetMemoryByType(ERigVMMemoryType InMemoryType, bool bCreateIfNeeded)
 {
 	switch(InMemoryType)
 	{
 		case ERigVMMemoryType::Literal:
 		{
-			return GetLiteralMemory();
+			if(bCreateIfNeeded)
+			{
+				if(LiteralMemoryStorageObject == nullptr)
+				{
+					if(UClass* Class = URigVMMemoryStorageGeneratorClass::GetStorageClass(this, InMemoryType))
+					{
+						// for literals we share the CDO between all VMs
+						LiteralMemoryStorageObject = Cast<URigVMMemoryStorage>(Class->GetDefaultObject(true));
+					}
+					else
+					{
+						LiteralMemoryStorageObject = NewObject<URigVMMemoryStorage>(this);
+					}
+				}
+			}
+			return LiteralMemoryStorageObject;
 		}
 		case ERigVMMemoryType::Work:
 		{
-			return GetWorkMemory();
+			if(bCreateIfNeeded)
+			{
+				if(WorkMemoryStorageObject)
+				{
+					if(WorkMemoryStorageObject->GetOuter() != this)
+					{
+						WorkMemoryStorageObject = nullptr;
+					}
+				}
+				if(WorkMemoryStorageObject == nullptr)
+				{
+					if(UClass* Class = URigVMMemoryStorageGeneratorClass::GetStorageClass(this, InMemoryType))
+					{
+						WorkMemoryStorageObject = NewObject<URigVMMemoryStorage>(this, Class);
+					}
+					else
+					{
+						WorkMemoryStorageObject = NewObject<URigVMMemoryStorage>(this);
+					}
+				}
+			}
+			check(WorkMemoryStorageObject->GetOuter() == this);
+			return WorkMemoryStorageObject;
 		}
 		case ERigVMMemoryType::Debug:
 		{
-			return GetDebugMemory();
+			if(bCreateIfNeeded)
+			{
+				if(DebugMemoryStorageObject)
+				{
+					if(DebugMemoryStorageObject->GetOuter() != this)
+					{
+						DebugMemoryStorageObject = nullptr;
+					}
+				}
+				if(DebugMemoryStorageObject == nullptr)
+				{
+#if WITH_EDITOR
+					if(UClass* Class = URigVMMemoryStorageGeneratorClass::GetStorageClass(this, InMemoryType))
+					{
+						DebugMemoryStorageObject = NewObject<URigVMMemoryStorage>(this, Class);
+					}
+					else
+#endif
+					{
+						DebugMemoryStorageObject = NewObject<URigVMMemoryStorage>(this);
+					}
+				}
+			}
+			check(WorkMemoryStorageObject->GetOuter() == this);
+			return DebugMemoryStorageObject;
 		}
 		default:
 		{
@@ -726,6 +737,27 @@ URigVMMemoryStorage* URigVM::GetMemoryByType(ERigVMMemoryType InMemoryType) cons
 	return nullptr;
 }
 
+void URigVM::ClearMemory()
+{
+	LiteralMemoryStorageObject = nullptr;
+
+	if(WorkMemoryStorageObject)
+	{
+		WorkMemoryStorageObject->Rename(nullptr, GetTransientPackage(), REN_ForceNoResetLoaders | REN_DoNotDirty | REN_DontCreateRedirectors | REN_NonTransactional);
+		WorkMemoryStorageObject->MarkPendingKill();
+		WorkMemoryStorageObject = nullptr;
+	}
+
+	if(DebugMemoryStorageObject)
+	{
+		DebugMemoryStorageObject->Rename(nullptr, GetTransientPackage(), REN_ForceNoResetLoaders | REN_DoNotDirty | REN_DontCreateRedirectors | REN_NonTransactional);
+		DebugMemoryStorageObject->MarkPendingKill();
+		DebugMemoryStorageObject = nullptr;
+	}
+
+	InvalidateCachedMemory();
+}
+
 #endif
 
 const FRigVMInstructionArray& URigVM::GetInstructions()
@@ -2548,7 +2580,7 @@ FString URigVM::DumpByteCodeAsText(const TArray<int32>& InInstructionOrder, bool
 	return FString::Join(DumpByteCodeAsTextArray(InInstructionOrder, bIncludeLineNumbers), TEXT("\n"));
 }
 
-FString URigVM::GetOperandLabel(const FRigVMOperand& InOperand, TFunction<FString(const FString& RegisterName, const FString& RegisterOffsetName)> FormatFunction) const
+FString URigVM::GetOperandLabel(const FRigVMOperand& InOperand, TFunction<FString(const FString& RegisterName, const FString& RegisterOffsetName)> FormatFunction)
 {
 #if UE_RIGVM_UCLASS_BASED_STORAGE_DISABLED
 	
@@ -2714,6 +2746,10 @@ void URigVM::CacheSingleMemoryHandle(const FRigVMOperand& InArg, bool bForExecut
 		check(Memory->GetPropertyPaths().IsValidIndex(InArg.GetRegisterOffset()));
 		PropertyPath = &Memory->GetPropertyPaths()[InArg.GetRegisterOffset()];
 	}
+
+	// if you are hitting this it's likely that the VM was created outside of a valid
+	// package. the compiler bases the memory class construction on the package the VM
+	// is in - so a VM under GetTransientPackage() can be created - but not run.
 	uint8* Data = Memory->GetData<uint8>(InArg.GetRegisterIndex());
 	const FProperty* Property = Memory->GetProperties()[InArg.GetRegisterIndex()];
 	const FRigVMMemoryHandle Handle(Data, Property, PropertyPath);
diff --git a/Engine/Source/Runtime/RigVM/Public/RigVMCore/RigVM.h b/Engine/Source/Runtime/RigVM/Public/RigVMCore/RigVM.h
--- a/Engine/Source/Runtime/RigVM/Public/RigVMCore/RigVM.h
+++ b/Engine/Source/Runtime/RigVM/Public/RigVMCore/RigVM.h
@@ -213,41 +213,39 @@ public:
 #else
 
 	// Returns a memory storage by type
-	URigVMMemoryStorage* GetMemoryByType(ERigVMMemoryType InMemoryType) const;
+	URigVMMemoryStorage* GetMemoryByType(ERigVMMemoryType InMemoryType, bool bCreateIfNeeded = true);
 	
 	// The default mutable work memory
-	FORCEINLINE URigVMMemoryStorage* GetWorkMemory() const { return WorkMemoryStorageObject; }
+	FORCEINLINE URigVMMemoryStorage* GetWorkMemory(bool bCreateIfNeeded = true) { return GetMemoryByType(ERigVMMemoryType::Work, bCreateIfNeeded); }
 
 	// The default const literal memory
-	FORCEINLINE URigVMMemoryStorage* GetLiteralMemory() const { return LiteralMemoryStorageObject; }
+	FORCEINLINE URigVMMemoryStorage* GetLiteralMemory(bool bCreateIfNeeded = true) { return GetMemoryByType(ERigVMMemoryType::Literal, bCreateIfNeeded); }
 
 	// The default debug watch memory
-	FORCEINLINE URigVMMemoryStorage* GetDebugMemory() const { return DebugMemoryStorageObject; }
+	FORCEINLINE URigVMMemoryStorage* GetDebugMemory(bool bCreateIfNeeded = true) { return GetMemoryByType(ERigVMMemoryType::Debug, bCreateIfNeeded); }
 
 	// returns all memory storages as an array
-	FORCEINLINE TArray<URigVMMemoryStorage*> GetLocalMemoryArray() const
+	FORCEINLINE TArray<URigVMMemoryStorage*> GetLocalMemoryArray()
 	{
-		check(GetWorkMemory());
-		check(GetLiteralMemory());
-		check(GetDebugMemory());
-		
 		TArray<URigVMMemoryStorage*> LocalMemory;
-		LocalMemory.Add(GetWorkMemory());
-		LocalMemory.Add(GetLiteralMemory());
-		LocalMemory.Add(GetDebugMemory());
-		
+		LocalMemory.Add(GetWorkMemory(true));
+		LocalMemory.Add(GetLiteralMemory(true));
+		LocalMemory.Add(GetDebugMemory(true));
 		return LocalMemory;
 	}
 
+	// Removes all memory from this VM
+	void ClearMemory();
+
 #endif
 
-	UPROPERTY()
+	UPROPERTY(transient, DuplicateTransient)
 	URigVMMemoryStorage* WorkMemoryStorageObject;
 
-	UPROPERTY()
+	UPROPERTY(transient, DuplicateTransient)
 	URigVMMemoryStorage* LiteralMemoryStorageObject;
 
-	UPROPERTY()
+	UPROPERTY(transient, DuplicateTransient)
 	URigVMMemoryStorage* DebugMemoryStorageObject;
 
 	TArray<FRigVMPropertyPathDescription> ExternalPropertyPathDescriptions;
@@ -349,7 +347,7 @@ public:
 	}
 
 	// Retrieve the array size of the parameter
-	int32 GetParameterArraySize(const FRigVMParameter& InParameter) const
+	int32 GetParameterArraySize(const FRigVMParameter& InParameter)
 	{
 		return (int32)GetWorkMemory()[InParameter.GetRegisterIndex()].GetTotalElementCount();
 	}
@@ -380,7 +378,7 @@ public:
 	}
 
 	// Retrieve the array size of the parameter
-	FORCEINLINE int32 GetParameterArraySize(const FRigVMParameter& InParameter) const
+	FORCEINLINE int32 GetParameterArraySize(const FRigVMParameter& InParameter)
 	{
 		const int32 PropertyIndex = InParameter.GetRegisterIndex();
 		const FProperty* Property = GetWorkMemory()->GetProperties()[PropertyIndex];
@@ -396,13 +394,13 @@ public:
 #endif
 
 	// Retrieve the array size of the parameter
-	int32 GetParameterArraySize(int32 InParameterIndex) const
+	int32 GetParameterArraySize(int32 InParameterIndex)
 	{
 		return GetParameterArraySize(Parameters[InParameterIndex]);
 	}
 
 	// Retrieve the array size of the parameter
-	int32 GetParameterArraySize(const FName& InParameterName) const
+	int32 GetParameterArraySize(const FName& InParameterName)
 	{
 		int32 ParameterIndex = ParametersNameMap.FindChecked(InParameterName);
 		return GetParameterArraySize(ParameterIndex);
@@ -682,7 +680,7 @@ public:
 
 #if WITH_EDITOR
 	// FormatFunction is an optional argument that allows you to override how operands are displayed, for example, see SControlRigStackView::PopulateStackView
-	FString GetOperandLabel(const FRigVMOperand & InOperand, TFunction<FString(const FString& RegisterName, const FString& RegisterOffsetName)> FormatFunction = nullptr) const;
+	FString GetOperandLabel(const FRigVMOperand & InOperand, TFunction<FString(const FString& RegisterName, const FString& RegisterOffsetName)> FormatFunction = nullptr);
 #endif
 
 	FExecutionReachedExitEvent& ExecutionReachedExit() { return OnExecutionReachedExit; }
