Control Rig: RigVM: Mark temporary classes as CLASS_Hidden

diff --git a/Engine/Source/Runtime/RigVM/Private/RigVMCore/RigVMMemoryStorage.cpp b/Engine/Source/Runtime/RigVM/Private/RigVMCore/RigVMMemoryStorage.cpp
--- a/Engine/Source/Runtime/RigVM/Private/RigVMCore/RigVMMemoryStorage.cpp
+++ b/Engine/Source/Runtime/RigVM/Private/RigVMCore/RigVMMemoryStorage.cpp
@@ -256,7 +256,7 @@ URigVMMemoryStorageGeneratorClass* URigVMMemoryStorageGeneratorClass::CreateStor
 	Class->ClassWithin = UObject::StaticClass();
 
 	// make sure that it cannot be placed / used for class selectors
-	Class->ClassFlags |= CLASS_NotPlaceable;
+	Class->ClassFlags |= CLASS_NotPlaceable | CLASS_Hidden;
 
 	// store our custom state
 	Class->MemoryType = InMemoryType;
