Control Rig: RigVM: Fix serialization for new RigVMMemoryStorage

diff --git a/Engine/Source/Runtime/RigVM/Private/RigVMCore/RigVMMemoryDeprecated.cpp b/Engine/Source/Runtime/RigVM/Private/RigVMCore/RigVMMemoryDeprecated.cpp
--- a/Engine/Source/Runtime/RigVM/Private/RigVMCore/RigVMMemoryDeprecated.cpp
+++ b/Engine/Source/Runtime/RigVM/Private/RigVMCore/RigVMMemoryDeprecated.cpp
@@ -351,8 +351,6 @@ FRigVMRegisterOffset::FRigVMRegisterOffset(UScriptStruct* InScriptStruct, const
 	ensure(ElementSize > 0);
 }
 
-#if UE_RIGVM_UCLASS_BASED_STORAGE_DISABLED
-
 uint8* FRigVMRegisterOffset::GetData(uint8* InContainer) const
 {
 	uint8* Data = InContainer;
@@ -372,8 +370,6 @@ uint8* FRigVMRegisterOffset::GetData(uint8* InContainer) const
 	return Data;
 }
 
-#endif
-
 bool FRigVMRegisterOffset::ContainsArraySegment() const
 {
 	for (int32 SegmentIndex : Segments)
@@ -469,8 +465,6 @@ FRigVMMemoryContainer::FRigVMMemoryContainer(bool bInUseNames)
 {
 }
 
-#if UE_RIGVM_UCLASS_BASED_STORAGE_DISABLED
-
 FRigVMMemoryContainer::FRigVMMemoryContainer(const FRigVMMemoryContainer& Other)
 {
 	*this = Other;
@@ -486,10 +480,14 @@ bool FRigVMMemoryContainer::CopyRegisters(const FRigVMMemoryContainer &InOther)
 	ensure(Registers.Num() == InOther.Registers.Num());
 	for (int32 Index = 0; Index < Registers.Num(); ++Index)
 	{
+#if UE_RIGVM_UCLASS_BASED_STORAGE_DISABLED
+		
 		if (!Copy(Index, Index, &InOther))
 		{
 			return false;
 		}
+
+#endif
 	}
 	return true;
 }
@@ -526,17 +524,19 @@ FRigVMMemoryContainer& FRigVMMemoryContainer::operator= (const FRigVMMemoryConta
 
 		// Why was this introduced? Are we trying to save time here?
 		// An assignment should be copying everything.
+#if UE_RIGVM_UCLASS_BASED_STORAGE_DISABLED
+
 		if (MemoryType == ERigVMMemoryType::Literal)
 		{
 			Copy(Index, Index, &InOther);
 		}
+
+#endif
 	}
 
 	return *this;
 }
 
-#endif
-
 void FRigVMMemoryContainer::Serialize(FArchive& Ar)
 {
 	Ar.UsingCustomVersion(FAnimObjectVersion::GUID);
@@ -755,12 +755,10 @@ void FRigVMMemoryContainer::Save(FArchive& Ar)
 
 void FRigVMMemoryContainer::Load(FArchive& Ar)
 {
-#if UE_RIGVM_UCLASS_BASED_STORAGE_DISABLED
 	for (int32 RegisterIndex = 0; RegisterIndex < Registers.Num(); RegisterIndex++)
 	{
 		Destroy(RegisterIndex);
 	}
-#endif
 	
 	Ar << bUseNameMap;
 	Ar << MemoryType;
@@ -777,8 +775,6 @@ void FRigVMMemoryContainer::Load(FArchive& Ar)
 	TArray<FString> ScriptStructPaths;
 	Ar << ScriptStructPaths;
 
-#if UE_RIGVM_UCLASS_BASED_STORAGE_DISABLED
-
 	for (const FString& ScriptStructPath : ScriptStructPaths)
 	{
 		UScriptStruct* ScriptStruct = FindObject<UScriptStruct>(nullptr, *ScriptStructPath);
@@ -795,13 +791,9 @@ void FRigVMMemoryContainer::Load(FArchive& Ar)
 		ScriptStructs.Add(ScriptStruct);
 	}
 
-#endif
-
 	uint64 TotalBytes = 0;
 	Ar << TotalBytes;
 
-#if UE_RIGVM_UCLASS_BASED_STORAGE_DISABLED
-
 	Data.Empty();
 
 	if (!bEncounteredErrorDuringLoad)
@@ -839,8 +831,8 @@ void FRigVMMemoryContainer::Load(FArchive& Ar)
 				Register.ByteIndex = Data.AddZeroed(Register.GetNumBytesAllSlices());
 			}
 		}
+
 		UpdateRegisters();
-		
 
 		for (int32 RegisterOffsetIndex = 0; RegisterOffsetIndex < RegisterOffsets.Num(); RegisterOffsetIndex++)
 		{
@@ -872,14 +864,10 @@ void FRigVMMemoryContainer::Load(FArchive& Ar)
 		}
 	}
 
-#endif
-
 	for (int32 RegisterIndex = 0; RegisterIndex < Registers.Num(); RegisterIndex++)
 	{
 		FRigVMRegister& Register = Registers[RegisterIndex];
 
-#if UE_RIGVM_UCLASS_BASED_STORAGE_DISABLED
-
 		if (Register.ElementCount == 0 && !Register.IsDynamic())
 		{
 			continue;
@@ -892,8 +880,6 @@ void FRigVMMemoryContainer::Load(FArchive& Ar)
 		{
 			check(!Register.IsDynamic());
 		}
-
-#endif
 		
 		if (!Register.IsDynamic())
 		{
@@ -904,15 +890,11 @@ void FRigVMMemoryContainer::Load(FArchive& Ar)
 					FRigVMByteArray View;
 					Ar << View;
 
-#if UE_RIGVM_UCLASS_BASED_STORAGE_DISABLED
-
 					if (!bEncounteredErrorDuringLoad)
 					{
 						ensure(View.Num() <= Register.GetAllocatedBytes());
 						FMemory::Memcpy(&Data[Register.GetWorkByteIndex()], View.GetData(), View.Num());
 					}
-
-#endif
 					break;
 				}
 				case ERigVMRegisterType::Name:
@@ -920,15 +902,11 @@ void FRigVMMemoryContainer::Load(FArchive& Ar)
 					TArray<FName> View;
 					Ar << View;
 
-#if UE_RIGVM_UCLASS_BASED_STORAGE_DISABLED
-
 					if (!bEncounteredErrorDuringLoad)
 					{
 						ensure(View.Num() == Register.GetTotalElementCount());
 						RigVMCopy<FName>(&Data[Register.GetWorkByteIndex()], View.GetData(), View.Num());
 					}
-
-#endif
 					break;
 				}
 				case ERigVMRegisterType::String:
@@ -936,15 +914,11 @@ void FRigVMMemoryContainer::Load(FArchive& Ar)
 					TArray<FString> View;
 					Ar << View;
 
-#if UE_RIGVM_UCLASS_BASED_STORAGE_DISABLED
-						
 					if (!bEncounteredErrorDuringLoad)
 					{
 						ensure(View.Num() == Register.GetTotalElementCount());
 						RigVMCopy<FString>(&Data[Register.GetWorkByteIndex()], View.GetData(), View.Num());
 					}
-						
-#endif
 					break;
 				}
 				case ERigVMRegisterType::Struct:
@@ -952,8 +926,6 @@ void FRigVMMemoryContainer::Load(FArchive& Ar)
 					TArray<FString> View;
 					Ar << View;
 
-#if UE_RIGVM_UCLASS_BASED_STORAGE_DISABLED
-
 					if (!bEncounteredErrorDuringLoad)
 					{
 						ensure(View.Num() == Register.GetTotalElementCount());
@@ -977,8 +949,6 @@ void FRigVMMemoryContainer::Load(FArchive& Ar)
 							}
 						}
 					}
-
-#endif
 					break;
 				}
 			}
@@ -986,8 +956,6 @@ void FRigVMMemoryContainer::Load(FArchive& Ar)
 		else if (!Register.IsNestedDynamic())
 		{
 
-#if UE_RIGVM_UCLASS_BASED_STORAGE_DISABLED
-			
 			FRigVMByteArray* ArrayStorage = nullptr;
 
 			if (!bEncounteredErrorDuringLoad)
@@ -995,8 +963,6 @@ void FRigVMMemoryContainer::Load(FArchive& Ar)
 				ArrayStorage = (FRigVMByteArray*)&Data[Register.GetWorkByteIndex()];
 			}
 			
-#endif
-			
 			// notice that all dynamic registers' memory is not really serialized
 			// serialize empty arrays here just to avoid having to increment serialization version
 			// the loaded data is not getting used at all, it will be recomputed at execution time
@@ -1030,8 +996,6 @@ void FRigVMMemoryContainer::Load(FArchive& Ar)
 		}
 		else
 		{
-
-#if UE_RIGVM_UCLASS_BASED_STORAGE_DISABLED
 			
 			FRigVMNestedByteArray* NestedArrayStorage = nullptr;
 			if (!bEncounteredErrorDuringLoad)
@@ -1041,20 +1005,14 @@ void FRigVMMemoryContainer::Load(FArchive& Ar)
 				NestedArrayStorage->SetNumZeroed(Register.SliceCount);
 			}
 			
-#endif
-
 			for (int32 SliceIndex = 0; SliceIndex < Register.SliceCount; SliceIndex++)
 			{
-#if UE_RIGVM_UCLASS_BASED_STORAGE_DISABLED
-				
 				FRigVMByteArray* ArrayStorage = nullptr;
 				if (NestedArrayStorage)
 				{
 					ArrayStorage = &(*NestedArrayStorage)[SliceIndex];
 				}
 
-#endif
-
 				// notice that all dynamic registers' memory is not really serialized
 				// serialize empty arrays here just to avoid having to increment serialization version
 				// the loaded data is not getting used at all, it will be recomputed at execution time
@@ -1111,8 +1069,6 @@ void FRigVMMemoryContainer::Load(FArchive& Ar)
 #endif
 }
 
-#if UE_RIGVM_UCLASS_BASED_STORAGE_DISABLED
-
 void FRigVMMemoryContainer::Reset()
 {
 	if (Data.Num() > 0)
@@ -1139,6 +1095,8 @@ void FRigVMMemoryContainer::Empty()
 	NameMap.Empty();
 }
 
+#if UE_RIGVM_UCLASS_BASED_STORAGE_DISABLED
+
 bool FRigVMMemoryContainer::Copy(
 	int32 InTargetRegisterIndex,
 	int32 InTargetRegisterOffset,
@@ -1336,6 +1294,8 @@ bool FRigVMMemoryContainer::Copy(
 	return Copy(InSourceOperand.GetRegisterIndex(), InTargetOperand.GetRegisterIndex(), InSourceMemory, InSourceOperand.GetRegisterOffset(), InTargetOperand.GetRegisterOffset(), InSourceSliceIndex, InTargetSliceIndex);
 }
 
+#endif
+
 int32 FRigVMMemoryContainer::Allocate(const FName& InNewName, int32 InElementSize, int32 InElementCount, int32 InSliceCount, const uint8* InDataPtr, bool bUpdateRegisters)
 {
 	FName Name = InNewName;
@@ -1714,6 +1674,8 @@ int32 FRigVMMemoryContainer::GetOrAddRegisterOffset(int32 InRegisterIndex, UScri
 	return ExistingIndex;
 }
 
+#if UE_RIGVM_UCLASS_BASED_STORAGE_DISABLED
+
 void FRigVMMemoryContainer::SetRegisterValueFromString(const FRigVMOperand& InOperand, const FString& InCPPType, const UObject* InCPPTypeObject, const TArray<FString>& InDefaultValues)
 {
 	if (InOperand.GetRegisterIndex() < 0 || InOperand.GetRegisterIndex() >= Registers.Num())
@@ -1892,6 +1854,8 @@ const FRigVMRegisterOffset& FRigVMMemoryContainer::GetRegisterOffsetForOperand(c
 	return GetRegisterOffset(InOperand.GetRegisterOffset());
 }
 
+#endif
+
 void FRigVMMemoryContainer::UpdateRegisters()
 {
 	int32 AlignmentShift = 0;
@@ -1965,12 +1929,16 @@ void FRigVMMemoryContainer::UpdateRegisters()
 	}
 }
 
+#if UE_RIGVM_UCLASS_BASED_STORAGE_DISABLED
+
 void FRigVMMemoryContainer::FillWithZeroes(int32 InRegisterIndex)
 {
 	ensure(Registers.IsValidIndex(InRegisterIndex));
 	FMemory::Memzero(GetData(InRegisterIndex), Registers[InRegisterIndex].GetNumBytesAllSlices());
 }
 
+#endif
+
 int32 FRigVMMemoryContainer::FindOrAddScriptStruct(UScriptStruct* InScriptStruct)
 {
 	if (InScriptStruct == nullptr)
@@ -1985,5 +1953,3 @@ int32 FRigVMMemoryContainer::FindOrAddScriptStruct(UScriptStruct* InScriptStruct
 	}
 	return ScriptStructs.Add(InScriptStruct);
 }
-
-#endif
\ No newline at end of file
diff --git a/Engine/Source/Runtime/RigVM/Public/RigVMCore/RigVMMemoryDeprecated.h b/Engine/Source/Runtime/RigVM/Public/RigVMCore/RigVMMemoryDeprecated.h
--- a/Engine/Source/Runtime/RigVM/Public/RigVMCore/RigVMMemoryDeprecated.h
+++ b/Engine/Source/Runtime/RigVM/Public/RigVMCore/RigVMMemoryDeprecated.h
@@ -215,12 +215,8 @@ public:
 	// constructs a path given a struct and a segment path
 	FRigVMRegisterOffset(UScriptStruct* InScriptStruct, const FString& InSegmentPath, int32 InInitialOffset = 0, uint16 InElementSize = 0, const FName& InCPPType = NAME_None);
 
-#if UE_RIGVM_UCLASS_BASED_STORAGE_DISABLED
-
 	// returns the data pointer within a container
 	uint8* GetData(uint8* InContainer) const;
-
-#endif
 	
 	// returns the segments of this path
 	const TArray<int32>& GetSegments() const { return Segments; }
@@ -506,8 +502,6 @@ public:
 
 	FRigVMMemoryContainer(bool bInUseNames = true);
 	
-#if UE_RIGVM_UCLASS_BASED_STORAGE_DISABLED
-	
 	FRigVMMemoryContainer(const FRigVMMemoryContainer& Other);
 	~FRigVMMemoryContainer();
 
@@ -572,8 +566,6 @@ public:
 
 	// accessor for a register based on a a name. note: only works if SupportsNames() == true
 	FORCEINLINE_DEBUGGABLE FRigVMRegister& GetRegister(const FName& InName) { return Registers[GetIndex(InName)]; }
-
-#endif
 	
 	void Serialize(FArchive& Ar);
 	void Save(FArchive& Ar);
@@ -584,8 +576,6 @@ public:
 		return Ar;
 	}
 
-#if UE_RIGVM_UCLASS_BASED_STORAGE_DISABLED
-
 	// Returns an argument for a given register.
 	// This is typically used to store a light weight address for use within a VM.
 	FORCEINLINE_DEBUGGABLE FRigVMOperand GetOperand(int32 InRegisterIndex, int32 InRegisterOffset)
@@ -666,6 +656,8 @@ private:
 
 public:
 
+#if UE_RIGVM_UCLASS_BASED_STORAGE_DISABLED
+
 	// Returns a memory handle for a given register
 	FORCEINLINE_DEBUGGABLE FRigVMMemoryHandle GetHandle(const FRigVMRegister& Register, int32 InRegisterOffset = INDEX_NONE) const
 	{
@@ -927,6 +919,8 @@ public:
 		return GetDynamicArray<T>(InOperand.GetRegisterIndex(), InOperand.GetRegisterOffset(), InSliceIndex);
 	}
 
+#endif
+	
 	// Returns the script struct used for a given register (can be nullptr for non-struct-registers).
 	FORCEINLINE_DEBUGGABLE UScriptStruct* GetScriptStruct(const FRigVMRegister& Register) const
 	{
@@ -952,6 +946,8 @@ public:
 		return Path.GetScriptStruct();
 	}
 
+#if UE_RIGVM_UCLASS_BASED_STORAGE_DISABLED
+
 	// Copies the content of a source register to a target register.
 	// The source register can optionally be referencing a specific source memory container.
 	bool Copy(
@@ -984,6 +980,8 @@ public:
 		int32 InSourceSliceIndex = 0,
 		int32 InTargetSliceIndex = 0);
 
+#endif
+	
 	// Returns the index of a register based on the register name.
 	// Note: This only works if SupportsNames() == true
 	FORCEINLINE_DEBUGGABLE int32 GetIndex(const FName& InName) const
@@ -1025,6 +1023,8 @@ public:
 		return GetIndex(InPotentialNewName) == INDEX_NONE;
 	}
 
+#if UE_RIGVM_UCLASS_BASED_STORAGE_DISABLED
+
 	// Adds a new named register for a typed array from an array view (used by compiler)
 	template<typename T>
 	FORCEINLINE_DEBUGGABLE int32 AddFixedArray(const FName& InNewName, const FRigVMFixedArray<T>& InArrayView, int32 InSliceCount = 1)
@@ -1095,6 +1095,8 @@ public:
 		return AddDynamicValue<T>(NAME_None, InValue, InSliceCount);
 	}
 
+#endif
+	
 	// Adds a register path and returns its index
 	int32 GetOrAddRegisterOffset(int32 InRegisterIndex, int32 InArrayElement = 0);
 
@@ -1104,6 +1106,8 @@ public:
 	// Adds a register path and returns its index
 	int32 GetOrAddRegisterOffset(int32 InRegisterIndex, UScriptStruct* InScriptStruct, const FString& InSegmentPath, int32 InInitialOffset = 0, int32 InElementSize = 0);
 
+#if UE_RIGVM_UCLASS_BASED_STORAGE_DISABLED
+
 	void SetRegisterValueFromString(const FRigVMOperand& InOperand, const FString& InCPPType, const UObject* InCPPTypeObject, const TArray<FString>& InDefaultValues);
 	TArray<FString> GetRegisterValueAsString(const FRigVMOperand& InOperand, const FString& InCPPType, const UObject* InCPPTypeObject);
 
@@ -1313,6 +1317,8 @@ private:
 		return Register;
 	}
 
+#endif
+
 	// Updates internal data for topological changes
 	void UpdateRegisters();
 
@@ -1328,13 +1334,15 @@ private:
 	// Performs optional destruction of data within a struct register
 	bool Destroy(int32 InRegisterIndex, int32 InElementIndex = INDEX_NONE, int32 InSliceIndex = 0);
 
+#if UE_RIGVM_UCLASS_BASED_STORAGE_DISABLED
+	
 	// Fills a register with zero memory
 	void FillWithZeroes(int32 InRegisterIndex);
 
+#endif
+
 	// Ensures to add a script struct to the internal map if needed
 	int32 FindOrAddScriptStruct(UScriptStruct* InScriptStruct);
-
-#endif
 	
 	UPROPERTY()
 	bool bUseNameMap;
