Control Rig: RigVM: RigVMMemoryStorage support for cooked builds

diff --git a/Engine/Source/Runtime/RigVM/Private/RigVMCore/RigVM.cpp b/Engine/Source/Runtime/RigVM/Private/RigVMCore/RigVM.cpp
--- a/Engine/Source/Runtime/RigVM/Private/RigVMCore/RigVM.cpp
+++ b/Engine/Source/Runtime/RigVM/Private/RigVMCore/RigVM.cpp
@@ -202,12 +202,64 @@ void URigVM::Load(FArchive& Ar)
 #endif
 
 #if UE_RIGVM_UCLASS_BASED_STORAGE_DISABLED
+	
 	if (WorkMemoryStorage.bEncounteredErrorDuringLoad ||
 		LiteralMemoryStorage.bEncounteredErrorDuringLoad ||
 		!ValidateAllOperandsDuringLoad())
-#else
-	if (!ValidateAllOperandsDuringLoad())
+	{
+		Reset();
+	}
+	else
+	{
+		Instructions.Reset();
+		FunctionsStorage.Reset();
+		ParametersNameMap.Reset();
+
+		for (int32 Index = 0; Index < Parameters.Num(); Index++)
+		{
+			ParametersNameMap.Add(Parameters[Index].Name, Index);
+		}
+
+		// rebuild the bytecode to adjust for byte shifts in shipping
+		RebuildByteCodeOnLoad();
+
+		InvalidateCachedMemory();
+	}
+
 #endif
+}
+
+void URigVM::PostLoad()
+{
+	Super::PostLoad();
+	
+#if !UE_RIGVM_UCLASS_BASED_STORAGE_DISABLED
+
+	ClearMemory();
+
+	TArray<ERigVMMemoryType> MemoryTypes;
+	MemoryTypes.Add(ERigVMMemoryType::Literal);
+	MemoryTypes.Add(ERigVMMemoryType::Work);
+	MemoryTypes.Add(ERigVMMemoryType::Debug);
+
+	for(ERigVMMemoryType MemoryType : MemoryTypes)
+	{
+		if(URigVMMemoryStorageGeneratorClass* Class = URigVMMemoryStorageGeneratorClass::GetStorageClass(this, MemoryType))
+		{
+			if(Class->LinkedProperties.Num() == 0)
+			{
+				Class->RefreshLinkedProperties();
+			}
+			if(Class->PropertyPathDescriptions.Num() != Class->PropertyPaths.Num())
+			{
+				Class->RefreshPropertyPaths();
+			}
+		}
+	}
+	
+	RefreshExternalPropertyPaths();
+
+	if (!ValidateAllOperandsDuringLoad())
 	{
 		Reset();
 	}
@@ -227,6 +279,8 @@ void URigVM::Load(FArchive& Ar)
 
 		InvalidateCachedMemory();
 	}
+
+#endif
 }
 
 bool URigVM::ValidateAllOperandsDuringLoad()
@@ -944,23 +998,7 @@ void URigVM::CacheMemoryHandlesIfRequired(TArrayView<URigVMMemoryStorage*> InMem
 
 #if !UE_RIGVM_UCLASS_BASED_STORAGE_DISABLED
 
-	ExternalPropertyPaths.Reset();
-
-	ExternalPropertyPaths.SetNumZeroed(ExternalPropertyPathDescriptions.Num());
-	for(int32 PropertyPathIndex = 0; PropertyPathIndex < ExternalPropertyPaths.Num(); PropertyPathIndex++)
-	{
-		ExternalPropertyPaths[PropertyPathIndex] = FRigVMPropertyPath();
-
-		const int32 PropertyIndex = ExternalPropertyPathDescriptions[PropertyPathIndex].PropertyIndex;
-		if(ExternalVariables.IsValidIndex(PropertyIndex))
-		{
-			check(ExternalVariables[PropertyIndex].Property);
-			
-			ExternalPropertyPaths[PropertyPathIndex] = FRigVMPropertyPath(
-				ExternalVariables[PropertyIndex].Property,
-				ExternalPropertyPathDescriptions[PropertyPathIndex].SegmentPath);
-		}
-	}
+	RefreshExternalPropertyPaths();
 	
 #endif		
 
@@ -2958,6 +2996,27 @@ FRigVMCopyOp URigVM::GetCopyOpForOperands(const FRigVMOperand& InSource, const F
 	return FRigVMCopyOp(InSource, InTarget);
 }
 
+void URigVM::RefreshExternalPropertyPaths()
+{
+	ExternalPropertyPaths.Reset();
+
+	ExternalPropertyPaths.SetNumZeroed(ExternalPropertyPathDescriptions.Num());
+	for(int32 PropertyPathIndex = 0; PropertyPathIndex < ExternalPropertyPaths.Num(); PropertyPathIndex++)
+	{
+		ExternalPropertyPaths[PropertyPathIndex] = FRigVMPropertyPath();
+
+		const int32 PropertyIndex = ExternalPropertyPathDescriptions[PropertyPathIndex].PropertyIndex;
+		if(ExternalVariables.IsValidIndex(PropertyIndex))
+		{
+			check(ExternalVariables[PropertyIndex].Property);
+			
+			ExternalPropertyPaths[PropertyPathIndex] = FRigVMPropertyPath(
+				ExternalVariables[PropertyIndex].Property,
+				ExternalPropertyPathDescriptions[PropertyPathIndex].SegmentPath);
+		}
+	}
+}
+
 #endif
 
 #if UE_RIGVM_UCLASS_BASED_STORAGE_DISABLED
diff --git a/Engine/Source/Runtime/RigVM/Public/RigVMCore/RigVM.h b/Engine/Source/Runtime/RigVM/Public/RigVMCore/RigVM.h
--- a/Engine/Source/Runtime/RigVM/Public/RigVMCore/RigVM.h
+++ b/Engine/Source/Runtime/RigVM/Public/RigVMCore/RigVM.h
@@ -149,6 +149,7 @@ public:
 	void Serialize(FArchive& Ar);
 	void Save(FArchive& Ar);
 	void Load(FArchive& Ar);
+	void PostLoad() override;
 	
 	// resets the container and maintains all memory
 	void Reset();
@@ -823,6 +824,8 @@ private:
 	FCopyInfoForOperand GetCopyInfoForOperand(const FRigVMOperand& InOperand);
 	UScriptStruct* GetScriptStructForCopyOp(const FRigVMCopyOp& InCopyOp) const;
 	UScriptStruct* GetScripStructForOperand(const FRigVMOperand& InOperand) const;
+#else
+	void RefreshExternalPropertyPaths();
 #endif
 	
 	TMap<FRigVMOperand, TArray<FRigVMOperand>> OperandToDebugRegisters;
diff --git a/Engine/Source/Runtime/RigVM/Public/RigVMCore/RigVMMemoryStorage.h b/Engine/Source/Runtime/RigVM/Public/RigVMCore/RigVMMemoryStorage.h
--- a/Engine/Source/Runtime/RigVM/Public/RigVMCore/RigVMMemoryStorage.h
+++ b/Engine/Source/Runtime/RigVM/Public/RigVMCore/RigVMMemoryStorage.h
@@ -279,6 +279,7 @@ private:
 
 	friend class URigVMMemoryStorage;
 	friend class URigVMCompiler;
+	friend class URigVM;
 };
 
 /**
