Control Rig: Performance metric support / memory statistics support for new RIgVMMemoryStorage

diff --git a/Engine/Plugins/Animation/ControlRig/Source/ControlRig/Private/ControlRig.cpp b/Engine/Plugins/Animation/ControlRig/Source/ControlRig/Private/ControlRig.cpp
--- a/Engine/Plugins/Animation/ControlRig/Source/ControlRig/Private/ControlRig.cpp
+++ b/Engine/Plugins/Animation/ControlRig/Source/ControlRig/Private/ControlRig.cpp
@@ -1215,6 +1215,58 @@ bool UControlRig::CanExecute() const
 	return CVarControlRigDisableExecutionAll->GetInt() == 0;
 }
 
+TArray<UControlRig*> UControlRig::FindControlRigs(UObject* Outer, TSubclassOf<UControlRig> OptionalClass)
+{
+	TArray<UControlRig*> Result;
+	
+	if(Outer == nullptr)
+	{
+		return Result; 
+	}
+	
+	AActor* OuterActor = Cast<AActor>(Outer);
+	if(OuterActor == nullptr)
+	{
+		OuterActor = Outer->GetTypedOuter<AActor>();
+	}
+	
+	for (TObjectIterator<UControlRig> Itr; Itr; ++Itr)
+	{
+		UControlRig* RigInstance = *Itr;
+		if (OptionalClass == nullptr || RigInstance->GetClass()->IsChildOf(OptionalClass))
+		{
+			if(RigInstance->IsInOuter(Outer))
+			{
+				Result.Add(RigInstance);
+				continue;
+			}
+
+			if(OuterActor)
+			{
+				if(RigInstance->IsInOuter(OuterActor))
+				{
+					Result.Add(RigInstance);
+					continue;
+				}
+
+				if (TSharedPtr<IControlRigObjectBinding> Binding = RigInstance->GetObjectBinding())
+				{
+					if (AActor* Actor = Binding->GetHostingActor())
+					{
+						if (Actor == OuterActor)
+						{
+							Result.Add(RigInstance);
+							continue;
+						}
+					}
+				}
+			}
+		}
+	}
+
+	return Result;
+}
+
 void UControlRig::Serialize(FArchive& Ar)
 {
 	DECLARE_SCOPE_HIERARCHICAL_COUNTER_FUNC()
diff --git a/Engine/Plugins/Animation/ControlRig/Source/ControlRig/Private/ControlRigControlActor.cpp b/Engine/Plugins/Animation/ControlRig/Source/ControlRig/Private/ControlRigControlActor.cpp
--- a/Engine/Plugins/Animation/ControlRig/Source/ControlRig/Private/ControlRigControlActor.cpp
+++ b/Engine/Plugins/Animation/ControlRig/Source/ControlRig/Private/ControlRigControlActor.cpp
@@ -128,43 +128,23 @@ void AControlRigControlActor::Refresh()
 
 	if (ControlRig == nullptr)
 	{
-		for (TObjectIterator<UControlRig> Itr; Itr; ++Itr)
+		TArray<UControlRig*> Rigs = UControlRig::FindControlRigs(ActorToTrack, ControlRigClass);
+		if(Rigs.Num() > 0)
 		{
-			UControlRig* RigInstance = *Itr;
-			if (ControlRigClass == nullptr || RigInstance->GetClass()->IsChildOf(ControlRigClass))
-			{
-				if (TSharedPtr<IControlRigObjectBinding> Binding = RigInstance->GetObjectBinding())
-				{
-					if (AActor* Actor = Binding->GetHostingActor())
-					{
-						if (Actor == ActorToTrack)
-						{
-							ControlRig = RigInstance;
-							RemoveUnbindDelegate();
-							OnUnbindDelegate =Binding->OnControlRigUnbind().AddLambda([ this ]( ) { this->Clear(); this->Refresh(); });
-							break;
-						}
-					}
-				}
-			}
-
-			if(ControlRig == nullptr)
-			{
-				UObject* Outer = RigInstance->GetOuter();
-				if(RigInstance->IsInOuter(ActorToTrack))
-				{
-					ControlRig = RigInstance;
-					RemoveUnbindDelegate();
-					break;
-				}
-			}
+			ControlRig = Rigs[0];
 		}
-
+		
 		if (ControlRig == nullptr)
 		{
 			return;
 		}
 
+		RemoveUnbindDelegate();
+		if (TSharedPtr<IControlRigObjectBinding> Binding = ControlRig->GetObjectBinding())
+		{
+			OnUnbindDelegate = Binding->OnControlRigUnbind().AddLambda([ this ]( ) { this->Clear(); this->Refresh(); });
+		}
+
 		UControlRigGizmoLibrary* GizmoLibrary = ControlRig->GetGizmoLibrary();
 		if (GizmoLibrary == nullptr)
 		{
diff --git a/Engine/Plugins/Animation/ControlRig/Source/ControlRig/Public/ControlRig.h b/Engine/Plugins/Animation/ControlRig/Source/ControlRig/Public/ControlRig.h
--- a/Engine/Plugins/Animation/ControlRig/Source/ControlRig/Public/ControlRig.h
+++ b/Engine/Plugins/Animation/ControlRig/Source/ControlRig/Public/ControlRig.h
@@ -64,6 +64,9 @@ public:
 
 	static const FName OwnerComponent;
 
+	UFUNCTION(BlueprintCallable, Category = ControlRig)
+	static TArray<UControlRig*> FindControlRigs(UObject* Outer, TSubclassOf<UControlRig> OptionalClass);
+
 private:
 	/** Current delta time */
 	float DeltaTime;
@@ -262,7 +265,8 @@ public:
 
 	/** Set the VM, and bind its commands */
 	void SetVM(URigVM* VM);
-	
+
+	UFUNCTION(BlueprintPure, Category = "Control Rig")
 	URigVM* GetVM();
 
 	/** INodeMappingInterface implementation */
diff --git a/Engine/Source/Runtime/RigVM/Private/RigVMCore/RigVMMemoryStorage.cpp b/Engine/Source/Runtime/RigVM/Private/RigVMCore/RigVMMemoryStorage.cpp
--- a/Engine/Source/Runtime/RigVM/Private/RigVMCore/RigVMMemoryStorage.cpp
+++ b/Engine/Source/Runtime/RigVM/Private/RigVMCore/RigVMMemoryStorage.cpp
@@ -298,6 +298,15 @@ URigVMMemoryStorageGeneratorClass* URigVMMemoryStorageGeneratorClass::CreateStor
 	return Class;
 }
 
+FRigVMMemoryStatistics URigVMMemoryStorageGeneratorClass::GetStatistics() const
+{
+	FRigVMMemoryStatistics Statistics;
+	Statistics.RegisterCount = GetProperties().Num();
+	Statistics.TotalBytes = Statistics.DataBytes + sizeof(URigVMMemoryStorage);
+	Statistics.TotalBytes += sizeof(FProperty) * Statistics.RegisterCount;
+	return Statistics;
+}
+
 FProperty* URigVMMemoryStorageGeneratorClass::AddProperty(URigVMMemoryStorageGeneratorClass* InClass, const FRigVMPropertyDescription& InProperty, FField** LinkToProperty)
 {
 	UClass *SuperClass = URigVMMemoryStorage::StaticClass();
diff --git a/Engine/Source/Runtime/RigVM/Public/RigVMCore/RigVM.h b/Engine/Source/Runtime/RigVM/Public/RigVMCore/RigVM.h
--- a/Engine/Source/Runtime/RigVM/Public/RigVMCore/RigVM.h
+++ b/Engine/Source/Runtime/RigVM/Public/RigVMCore/RigVM.h
@@ -650,6 +650,7 @@ public:
 #endif
 
 	// returns the statistics information
+	UFUNCTION(BlueprintPure, Category = "RigVM")
 	FRigVMStatistics GetStatistics() const
 	{
 		FRigVMStatistics Statistics;
@@ -657,6 +658,15 @@ public:
 		Statistics.LiteralMemory = LiteralMemoryPtr->GetStatistics();
 		Statistics.WorkMemory = WorkMemoryPtr->GetStatistics();
 		Statistics.DebugMemory = DebugMemoryPtr->GetStatistics();
+#else
+		if(LiteralMemoryStorageObject)
+		{
+			Statistics.LiteralMemory = LiteralMemoryStorageObject->GetStatistics();
+		}
+		if(WorkMemoryStorageObject)
+		{
+			Statistics.WorkMemory = WorkMemoryStorageObject->GetStatistics();
+		}
 #endif
 		Statistics.ByteCode = ByteCodePtr->GetStatistics();
 		Statistics.BytesForCaching = FirstHandleForInstruction.GetAllocatedSize() + CachedMemoryHandles.GetAllocatedSize();
@@ -665,9 +675,16 @@ public:
 			Statistics.WorkMemory.TotalBytes +
 			Statistics.ByteCode.DataBytes +
 			Statistics.BytesForCaching;
+		
+#if UE_RIGVM_UCLASS_BASED_STORAGE_DISABLED
 		Statistics.BytesPerInstance =
 			Statistics.WorkMemory.TotalBytes +
 			Statistics.BytesForCaching;
+#else
+		Statistics.BytesPerInstance =
+			Statistics.WorkMemory.DataBytes +
+			Statistics.BytesForCaching;
+#endif
 
 		return Statistics;
 	}
diff --git a/Engine/Source/Runtime/RigVM/Public/RigVMCore/RigVMMemoryStorage.h b/Engine/Source/Runtime/RigVM/Public/RigVMCore/RigVMMemoryStorage.h
--- a/Engine/Source/Runtime/RigVM/Public/RigVMCore/RigVMMemoryStorage.h
+++ b/Engine/Source/Runtime/RigVM/Public/RigVMCore/RigVMMemoryStorage.h
@@ -255,6 +255,9 @@ public:
 	// The property paths stored within this class
 	const TArray<FRigVMPropertyPath>& GetPropertyPaths() const { return PropertyPaths; }
 
+	// returns the statistics information
+	FRigVMMemoryStatistics GetStatistics() const;
+
 protected:
 
 	// Adds a single property to a class.
@@ -347,6 +350,21 @@ public:
 	// Creates and returns a new operand for a property (and optionally a property path)
 	FRigVMOperand GetOperandByName(const FName& InName, int32 InPropertyPathIndex = INDEX_NONE) const;
 
+	// returns the statistics information
+	FRigVMMemoryStatistics GetStatistics() const
+	{
+		if(URigVMMemoryStorageGeneratorClass* Class = Cast<URigVMMemoryStorageGeneratorClass>(GetClass()))
+		{
+			return Class->GetStatistics();
+		}
+
+		FRigVMMemoryStatistics Statistics;
+		Statistics.RegisterCount = 0;
+		Statistics.DataBytes = URigVMMemoryStorage::StaticClass()->GetStructureSize();
+		Statistics.TotalBytes = Statistics.DataBytes;
+		return Statistics;
+	}
+
 	// Returns true if the property at a given index is a TArray
 	FORCEINLINE bool IsArray(int32 InPropertyIndex) const
 	{
diff --git a/Engine/Source/Runtime/RigVM/Public/RigVMCore/RigVMStatistics.h b/Engine/Source/Runtime/RigVM/Public/RigVMCore/RigVMStatistics.h
--- a/Engine/Source/Runtime/RigVM/Public/RigVMCore/RigVMStatistics.h
+++ b/Engine/Source/Runtime/RigVM/Public/RigVMCore/RigVMStatistics.h
@@ -5,62 +5,62 @@
 #include "CoreMinimal.h"
 #include "RigVMStatistics.generated.h"
 
-USTRUCT()
+USTRUCT(BlueprintType)
 struct FRigVMMemoryStatistics
 {
 	GENERATED_BODY()
 
 public:
 
-	UPROPERTY(VisibleAnywhere, Category = "Statistics")
-	uint32 RegisterCount = 0;
+	UPROPERTY(VisibleAnywhere, BlueprintReadOnly, Category = "Statistics")
+	int32 RegisterCount = 0;
 
-	UPROPERTY(VisibleAnywhere, Category = "Statistics")
-	uint32 DataBytes = 0;
+	UPROPERTY(VisibleAnywhere, BlueprintReadOnly, Category = "Statistics")
+	int32 DataBytes = 0;
 
-	UPROPERTY(VisibleAnywhere, Category = "Statistics")
-	uint32 TotalBytes = 0;
+	UPROPERTY(VisibleAnywhere, BlueprintReadOnly, Category = "Statistics")
+	int32 TotalBytes = 0;
 };
 
-USTRUCT()
+USTRUCT(BlueprintType)
 struct FRigVMByteCodeStatistics
 {
 	GENERATED_BODY()
 
 public:
 
-	UPROPERTY(VisibleAnywhere, Category = "Statistics")
-	uint32 InstructionCount = 0;
+	UPROPERTY(VisibleAnywhere, BlueprintReadOnly, Category = "Statistics")
+	int32 InstructionCount = 0;
 
-	UPROPERTY(VisibleAnywhere, Category = "Statistics")
-	uint32 DataBytes = 0;
+	UPROPERTY(VisibleAnywhere, BlueprintReadOnly, Category = "Statistics")
+	int32 DataBytes = 0;
 };
 
-USTRUCT()
+USTRUCT(BlueprintType)
 struct FRigVMStatistics
 {
 	GENERATED_BODY()
 
 public:
 
-	UPROPERTY(VisibleAnywhere, Category = "Statistics")
-	uint32 BytesForCDO = 0;
+	UPROPERTY(VisibleAnywhere, BlueprintReadOnly, Category = "Statistics")
+	int32 BytesForCDO = 0;
 
-	UPROPERTY(VisibleAnywhere, Category = "Statistics")
-	uint32 BytesPerInstance = 0;
+	UPROPERTY(VisibleAnywhere, BlueprintReadOnly, Category = "Statistics")
+	int32 BytesPerInstance = 0;
 
-	UPROPERTY(VisibleAnywhere, Category = "Statistics")
+	UPROPERTY(VisibleAnywhere, BlueprintReadOnly, Category = "Statistics")
 	FRigVMMemoryStatistics LiteralMemory;
 
-	UPROPERTY(VisibleAnywhere, Category = "Statistics")
+	UPROPERTY(VisibleAnywhere, BlueprintReadOnly, Category = "Statistics")
 	FRigVMMemoryStatistics WorkMemory;
 
-	UPROPERTY(VisibleAnywhere, Category = "Statistics")
+	UPROPERTY(VisibleAnywhere, BlueprintReadOnly, Category = "Statistics")
 	FRigVMMemoryStatistics DebugMemory;
 
-	UPROPERTY(VisibleAnywhere, Category = "Statistics")
-	uint32 BytesForCaching = 0;
+	UPROPERTY(VisibleAnywhere, BlueprintReadOnly, Category = "Statistics")
+	int32 BytesForCaching = 0;
 
-	UPROPERTY(VisibleAnywhere, Category = "Statistics")
+	UPROPERTY(VisibleAnywhere, BlueprintReadOnly, Category = "Statistics")
 	FRigVMByteCodeStatistics ByteCode;
 };
\ No newline at end of file
