RigVM: Define change to enable uclass based memory storage

diff --git a/Engine/Source/Runtime/RigVM/Public/RigVMDefines.h b/Engine/Source/Runtime/RigVM/Public/RigVMDefines.h
--- a/Engine/Source/Runtime/RigVM/Public/RigVMDefines.h
+++ b/Engine/Source/Runtime/RigVM/Public/RigVMDefines.h
@@ -7,5 +7,5 @@
 #pragma once
 
 #ifndef UE_RIGVM_UCLASS_BASED_STORAGE_DISABLED
-#define UE_RIGVM_UCLASS_BASED_STORAGE_DISABLED 1
+#define UE_RIGVM_UCLASS_BASED_STORAGE_DISABLED 0
 #endif
