RigVM: Byte Code operators for Array operations Reset, GetNum, SetNum etc

diff --git a/Engine/Plugins/Animation/ControlRig/Source/ControlRig/Private/ControlRig.cpp b/Engine/Plugins/Animation/ControlRig/Source/ControlRig/Private/ControlRig.cpp
--- a/Engine/Plugins/Animation/ControlRig/Source/ControlRig/Private/ControlRig.cpp
+++ b/Engine/Plugins/Animation/ControlRig/Source/ControlRig/Private/ControlRig.cpp
@@ -216,6 +216,9 @@ void UControlRig::InitializeFromCDO()
 
 		// copy draw container
 		DrawContainer = CDO->DrawContainer;
+
+		// copy vm settings
+		VMRuntimeSettings = CDO->VMRuntimeSettings; 
 	}
 }
 
@@ -786,37 +789,7 @@ void UControlRig::Execute(const EControlRigState InState, const FName& InEventNa
 					*Entry.Message
 				);
 
-			if(LoggedMessages.Contains(PerInstructionMessage))
-			{
-				continue;
-			}
-
-			switch (Entry.Severity)
-			{
-				case EMessageSeverity::CriticalError:
-				case EMessageSeverity::Error:
-				{
-					UE_LOG(LogControlRig, Error, TEXT("%s"), *PerInstructionMessage);
-					break;
-				}
-				case EMessageSeverity::PerformanceWarning:
-				case EMessageSeverity::Warning:
-				{
-					UE_LOG(LogControlRig, Warning, TEXT("%s"), *PerInstructionMessage);
-					break;
-				}
-				case EMessageSeverity::Info:
-				{
-					UE_LOG(LogControlRig, Display, TEXT("%s"), *PerInstructionMessage);
-					break;
-				}
-				default:
-				{
-					break;
-				}
-			}
-
-			LoggedMessages.Add(PerInstructionMessage, true);
+			LogOnce(Entry.Severity, Entry.InstructionIndex, PerInstructionMessage);
 		}
 	}
 
@@ -1118,6 +1091,7 @@ void UControlRig::SetVM(URigVM* NewVM)
 	if (VM)
 	{
 		VM->ExecutionReachedExit().RemoveAll(this);
+		VM->SetRuntimeSettings(FRigVMRuntimeSettings());
 	}
 	
 	if (NewVM)
@@ -1126,6 +1100,13 @@ void UControlRig::SetVM(URigVM* NewVM)
 		{
 			NewVM->ExecutionReachedExit().AddUObject(this, &UControlRig::HandleExecutionReachedExit);
 		}
+
+		// setup array handling and error reporting on the VM
+		VMRuntimeSettings.LogFunction = [this](EMessageSeverity::Type InSeverity, int32 InInstructionIndex, const FString& Message)
+		{
+			LogOnce(InSeverity, InInstructionIndex, Message);
+		};
+		NewVM->SetRuntimeSettings(VMRuntimeSettings);
 	}
 
 	VM = NewVM;
@@ -2576,6 +2557,41 @@ URigVM* UControlRig::GetSnapshotVM(bool bCreateIfNeeded)
 #endif
 }
 
+void UControlRig::LogOnce(EMessageSeverity::Type InSeverity, int32 InInstructionIndex, const FString& InMessage)
+{
+	if(LoggedMessages.Contains(InMessage))
+	{
+		return;
+	}
+
+	switch (InSeverity)
+	{
+		case EMessageSeverity::CriticalError:
+		case EMessageSeverity::Error:
+		{
+			UE_LOG(LogControlRig, Error, TEXT("%s"), *InMessage);
+			break;
+		}
+		case EMessageSeverity::PerformanceWarning:
+		case EMessageSeverity::Warning:
+		{
+			UE_LOG(LogControlRig, Warning, TEXT("%s"), *InMessage);
+			break;
+		}
+		case EMessageSeverity::Info:
+		{
+			UE_LOG(LogControlRig, Display, TEXT("%s"), *InMessage);
+			break;
+		}
+		default:
+		{
+			break;
+		}
+	}
+
+	LoggedMessages.Add(InMessage, true);
+}
+
 void UControlRig::AddBreakpoint(int32 InstructionIndex, URigVMNode* InNode, uint16 InDepth)
 {
 	DebugInfo.AddBreakpoint(InstructionIndex, InNode, InDepth);
diff --git a/Engine/Plugins/Animation/ControlRig/Source/ControlRig/Public/ControlRig.h b/Engine/Plugins/Animation/ControlRig/Source/ControlRig/Public/ControlRig.h
--- a/Engine/Plugins/Animation/ControlRig/Source/ControlRig/Public/ControlRig.h
+++ b/Engine/Plugins/Animation/ControlRig/Source/ControlRig/Public/ControlRig.h
@@ -240,6 +240,9 @@ public:
 	UPROPERTY(transient)
 	ERigExecutionType ExecutionType;
 
+	UPROPERTY(EditAnywhere, Category = "VM")
+	FRigVMRuntimeSettings VMRuntimeSettings;
+
 	/** Execute */
 	UFUNCTION(BlueprintCallable, Category = "Control Rig")
 	void Execute(const EControlRigState State, const FName& InEventName);
@@ -693,6 +696,7 @@ private:
 
 	FRigVMDebugInfo DebugInfo;
 	TMap<FString, bool> LoggedMessages;
+	void LogOnce(EMessageSeverity::Type InSeverity, int32 InInstructionIndex, const FString& InMessage);
 	
 public:
 
diff --git a/Engine/Plugins/Animation/ControlRig/Source/ControlRigDeveloper/Private/ControlRigBlueprint.cpp b/Engine/Plugins/Animation/ControlRig/Source/ControlRigDeveloper/Private/ControlRigBlueprint.cpp
--- a/Engine/Plugins/Animation/ControlRig/Source/ControlRigDeveloper/Private/ControlRigBlueprint.cpp
+++ b/Engine/Plugins/Animation/ControlRig/Source/ControlRigDeveloper/Private/ControlRigBlueprint.cpp
@@ -530,6 +530,7 @@ void UControlRigBlueprint::RecompileVM()
 		TGuardValue<bool> ReentrantGuardSelf(bSuspendModelNotificationsForSelf, true);
 		TGuardValue<bool> ReentrantGuardOthers(bSuspendModelNotificationsForOthers, true);
 
+		CDO->VMRuntimeSettings = VMRuntimeSettings;
 		CDO->GetHierarchy()->CopyHierarchy(Hierarchy);
 
 		if (CDO->VM->GetOuter() != CDO)
diff --git a/Engine/Plugins/Animation/ControlRig/Source/ControlRigDeveloper/Public/ControlRigBlueprint.h b/Engine/Plugins/Animation/ControlRig/Source/ControlRigDeveloper/Public/ControlRigBlueprint.h
--- a/Engine/Plugins/Animation/ControlRig/Source/ControlRigDeveloper/Public/ControlRigBlueprint.h
+++ b/Engine/Plugins/Animation/ControlRig/Source/ControlRigDeveloper/Public/ControlRigBlueprint.h
@@ -280,6 +280,9 @@ public:
 	UPROPERTY(EditAnywhere, Category = "User Interface")
 	FRigGraphDisplaySettings RigGraphDisplaySettings;
 
+	UPROPERTY(EditAnywhere, Category = "VM")
+	FRigVMRuntimeSettings VMRuntimeSettings;
+
 	UPROPERTY(EditAnywhere, BlueprintReadWrite, Category = "VM")
 	FRigVMCompileSettings VMCompileSettings;
 
diff --git a/Engine/Source/Runtime/RigVM/Private/RigVMCore/RigVM.cpp b/Engine/Source/Runtime/RigVM/Private/RigVMCore/RigVM.cpp
--- a/Engine/Source/Runtime/RigVM/Private/RigVMCore/RigVM.cpp
+++ b/Engine/Source/Runtime/RigVM/Private/RigVMCore/RigVM.cpp
@@ -430,6 +430,7 @@ bool URigVM::ValidateAllOperandsDuringLoad()
 			case ERigVMOpCode::BoolTrue:
 			case ERigVMOpCode::Increment:
 			case ERigVMOpCode::Decrement:
+			case ERigVMOpCode::ArrayReset:
 			{
 				const FRigVMUnaryOp& Op = ByteCodeStorage.GetOpAt<FRigVMUnaryOp>(ByteCodeInstruction);
 				CheckOperandValidity(Op.Arg);
@@ -460,12 +461,48 @@ bool URigVM::ValidateAllOperandsDuringLoad()
 				break;
 			}
 			case ERigVMOpCode::BeginBlock:
+			case ERigVMOpCode::ArrayGetNum:
+			case ERigVMOpCode::ArraySetNum:
+			case ERigVMOpCode::ArrayAppend:
+			case ERigVMOpCode::ArrayClone:
+			case ERigVMOpCode::ArrayRemove:
 			{
 				const FRigVMBinaryOp& Op = ByteCodeStorage.GetOpAt<FRigVMBinaryOp>(ByteCodeInstruction);
 				CheckOperandValidity(Op.ArgA);
 				CheckOperandValidity(Op.ArgB);
 				break;
 			}
+			case ERigVMOpCode::ArrayAdd:
+			case ERigVMOpCode::ArrayGetAtIndex:
+			case ERigVMOpCode::ArraySetAtIndex:
+			case ERigVMOpCode::ArrayInsert:
+			{
+				const FRigVMTernaryOp& Op = ByteCodeStorage.GetOpAt<FRigVMTernaryOp>(ByteCodeInstruction);
+				CheckOperandValidity(Op.ArgA);
+				CheckOperandValidity(Op.ArgB);
+				CheckOperandValidity(Op.ArgC);
+				break;
+			}
+			case ERigVMOpCode::ArrayFind:
+			{
+				const FRigVMQuaternaryOp& Op = ByteCodeStorage.GetOpAt<FRigVMQuaternaryOp>(ByteCodeInstruction);
+				CheckOperandValidity(Op.ArgA);
+				CheckOperandValidity(Op.ArgB);
+				CheckOperandValidity(Op.ArgC);
+				CheckOperandValidity(Op.ArgD);
+				break;
+			}
+			case ERigVMOpCode::ArrayIterator:
+			{
+				const FRigVMSenaryOp& Op = ByteCodeStorage.GetOpAt<FRigVMSenaryOp>(ByteCodeInstruction);
+				CheckOperandValidity(Op.ArgA);
+				CheckOperandValidity(Op.ArgB);
+				CheckOperandValidity(Op.ArgC);
+				CheckOperandValidity(Op.ArgD);
+				CheckOperandValidity(Op.ArgE);
+				CheckOperandValidity(Op.ArgF);
+				break;
+			}
 			case ERigVMOpCode::Invalid:
 			{
 				ensure(false);
@@ -1093,6 +1130,7 @@ void URigVM::CacheMemoryHandlesIfRequired(TArrayView<URigVMMemoryStorage*> InMem
 			case ERigVMOpCode::BoolTrue:
 			case ERigVMOpCode::Increment:
 			case ERigVMOpCode::Decrement:
+			case ERigVMOpCode::ArrayReset:
 			{
 				const FRigVMUnaryOp& Op = ByteCode.GetOpAt<FRigVMUnaryOp>(Instructions[InstructionIndex]);
 				CacheSingleMemoryHandle(Op.Arg);
@@ -1159,6 +1197,11 @@ void URigVM::CacheMemoryHandlesIfRequired(TArrayView<URigVMMemoryStorage*> InMem
 				break;
 			}
 			case ERigVMOpCode::BeginBlock:
+			case ERigVMOpCode::ArrayGetNum:
+			case ERigVMOpCode::ArraySetNum:
+			case ERigVMOpCode::ArrayAppend:
+			case ERigVMOpCode::ArrayClone:
+			case ERigVMOpCode::ArrayRemove:
 			{
 				const FRigVMBinaryOp& Op = ByteCode.GetOpAt<FRigVMBinaryOp>(Instructions[InstructionIndex]);
 				CacheSingleMemoryHandle(Op.ArgA);
@@ -1166,6 +1209,40 @@ void URigVM::CacheMemoryHandlesIfRequired(TArrayView<URigVMMemoryStorage*> InMem
 				InstructionIndex++;
 				break;
 			}
+			case ERigVMOpCode::ArrayAdd:
+			case ERigVMOpCode::ArrayGetAtIndex:
+			case ERigVMOpCode::ArraySetAtIndex:
+			case ERigVMOpCode::ArrayInsert:
+			{
+				const FRigVMTernaryOp& Op = ByteCode.GetOpAt<FRigVMTernaryOp>(Instructions[InstructionIndex]);
+				CacheSingleMemoryHandle(Op.ArgA);
+				CacheSingleMemoryHandle(Op.ArgB);
+				CacheSingleMemoryHandle(Op.ArgC);
+				InstructionIndex++;
+				break;
+			}
+			case ERigVMOpCode::ArrayFind:
+			{
+				const FRigVMQuaternaryOp& Op = ByteCode.GetOpAt<FRigVMQuaternaryOp>(Instructions[InstructionIndex]);
+				CacheSingleMemoryHandle(Op.ArgA);
+				CacheSingleMemoryHandle(Op.ArgB);
+				CacheSingleMemoryHandle(Op.ArgC);
+				CacheSingleMemoryHandle(Op.ArgD);
+				InstructionIndex++;
+				break;
+			}
+			case ERigVMOpCode::ArrayIterator:
+			{
+				const FRigVMSenaryOp& Op = ByteCode.GetOpAt<FRigVMSenaryOp>(Instructions[InstructionIndex]);
+				CacheSingleMemoryHandle(Op.ArgA);
+				CacheSingleMemoryHandle(Op.ArgB);
+				CacheSingleMemoryHandle(Op.ArgC);
+				CacheSingleMemoryHandle(Op.ArgD);
+				CacheSingleMemoryHandle(Op.ArgE);
+				CacheSingleMemoryHandle(Op.ArgF);
+				InstructionIndex++;
+				break;
+			}
 			case ERigVMOpCode::EndBlock:
 			{
 				InstructionIndex++;
@@ -1695,6 +1772,17 @@ bool URigVM::Initialize(TArrayView<URigVMMemoryStorage*> Memory, TArrayView<void
 			case ERigVMOpCode::BeginBlock:
 			case ERigVMOpCode::EndBlock:
 			case ERigVMOpCode::Exit:
+			case ERigVMOpCode::ArrayGetNum:
+			case ERigVMOpCode::ArraySetNum:
+			case ERigVMOpCode::ArrayAppend:
+			case ERigVMOpCode::ArrayClone:
+			case ERigVMOpCode::ArrayGetAtIndex:
+			case ERigVMOpCode::ArraySetAtIndex:
+			case ERigVMOpCode::ArrayInsert:
+			case ERigVMOpCode::ArrayRemove:
+			case ERigVMOpCode::ArrayAdd:
+			case ERigVMOpCode::ArrayFind:
+			case ERigVMOpCode::ArrayIterator:
 			{
 				break;
 			}
@@ -2283,6 +2371,264 @@ bool URigVM::Execute(TArrayView<URigVMMemoryStorage*> Memory, TArrayView<void*>
 				Context.InstructionIndex++;
 				break;
 			}
+			case ERigVMOpCode::ArrayReset:
+			{
+#if !UE_RIGVM_UCLASS_BASED_STORAGE_DISABLED
+				FRigVMMemoryHandle& ArrayHandle = CachedMemoryHandles[FirstHandleForInstruction[Context.InstructionIndex]]; 					
+				FScriptArrayHelper ArrayHelper(CastFieldChecked<FArrayProperty>(ArrayHandle.GetProperty()), ArrayHandle.GetData());
+				ArrayHelper.Resize(0);
+#endif	
+				Context.InstructionIndex++;
+				break;
+			}
+			case ERigVMOpCode::ArrayGetNum:
+			{
+#if !UE_RIGVM_UCLASS_BASED_STORAGE_DISABLED
+				FRigVMMemoryHandle& ArrayHandle = CachedMemoryHandles[FirstHandleForInstruction[Context.InstructionIndex]]; 					
+				FScriptArrayHelper ArrayHelper(CastFieldChecked<FArrayProperty>(ArrayHandle.GetProperty()), ArrayHandle.GetData());
+				int32& Count = (*((int32*)CachedMemoryHandles[FirstHandleForInstruction[Context.InstructionIndex] + 1].GetData()));
+				Count = ArrayHelper.Num();
+#endif	
+				Context.InstructionIndex++;
+				break;
+			}
+			case ERigVMOpCode::ArraySetNum:
+			{
+#if !UE_RIGVM_UCLASS_BASED_STORAGE_DISABLED
+				FRigVMMemoryHandle& ArrayHandle = CachedMemoryHandles[FirstHandleForInstruction[Context.InstructionIndex]]; 					
+				FScriptArrayHelper ArrayHelper(CastFieldChecked<FArrayProperty>(ArrayHandle.GetProperty()), ArrayHandle.GetData());
+				const int32 Count = (*((int32*)CachedMemoryHandles[FirstHandleForInstruction[Context.InstructionIndex] + 1].GetData()));
+				if(Context.IsValidArraySize(Count))
+				{
+					ArrayHelper.Resize(Count);
+				}
+#endif
+				Context.InstructionIndex++;
+				break;
+			}
+			case ERigVMOpCode::ArrayAppend:
+			{
+#if !UE_RIGVM_UCLASS_BASED_STORAGE_DISABLED
+				FRigVMMemoryHandle& ArrayHandle = CachedMemoryHandles[FirstHandleForInstruction[Context.InstructionIndex]]; 					
+				FRigVMMemoryHandle& OtherArrayHandle = CachedMemoryHandles[FirstHandleForInstruction[Context.InstructionIndex] + 1];
+				const FArrayProperty* ArrayProperty = CastFieldChecked<FArrayProperty>(ArrayHandle.GetProperty()); 
+				const FArrayProperty* OtherArrayProperty = CastFieldChecked<FArrayProperty>(OtherArrayHandle.GetProperty()); 
+
+				FScriptArrayHelper ArrayHelper(ArrayProperty, ArrayHandle.GetData());
+				FScriptArrayHelper OtherArrayHelper(OtherArrayProperty, OtherArrayHandle.GetData());
+
+				if(OtherArrayHelper.Num() > 0)
+				{
+					if(Context.IsValidArraySize(ArrayHelper.Num() + OtherArrayHelper.Num()))
+					{
+						const FProperty* TargetProperty = ArrayProperty->Inner;
+						const FProperty* SourceProperty = OtherArrayProperty->Inner;
+
+						int32 TargetIndex = ArrayHelper.AddValues(OtherArrayHelper.Num());
+						for(int32 SourceIndex = 0; SourceIndex < OtherArrayHelper.Num(); SourceIndex++, TargetIndex++)
+						{
+							uint8* TargetMemory = ArrayHelper.GetRawPtr(TargetIndex);
+							const uint8* SourceMemory = OtherArrayHelper.GetRawPtr(SourceIndex);
+							URigVMMemoryStorage::CopyProperty(TargetProperty, TargetMemory, SourceProperty, SourceMemory);
+						}
+					}
+				}
+#endif
+				Context.InstructionIndex++;
+				break;
+			}
+			case ERigVMOpCode::ArrayClone:
+			{
+#if !UE_RIGVM_UCLASS_BASED_STORAGE_DISABLED
+				FRigVMMemoryHandle& ArrayHandle = CachedMemoryHandles[FirstHandleForInstruction[Context.InstructionIndex]]; 					
+				FRigVMMemoryHandle& ClonedArrayHandle = CachedMemoryHandles[FirstHandleForInstruction[Context.InstructionIndex] + 1];
+				const FArrayProperty* ArrayProperty = CastFieldChecked<FArrayProperty>(ArrayHandle.GetProperty()); 
+				const FArrayProperty* ClonedArrayProperty = CastFieldChecked<FArrayProperty>(ClonedArrayHandle.GetProperty()); 
+				FScriptArrayHelper ArrayHelper(ArrayProperty, ArrayHandle.GetData());
+				FScriptArrayHelper ClonedArrayHelper(ClonedArrayProperty, ClonedArrayHandle.GetData());
+
+				ClonedArrayHelper.Resize(ArrayHelper.Num());
+				if(ArrayHelper.Num() > 0)
+				{
+					const FProperty* TargetProperty = ClonedArrayProperty->Inner;
+					const FProperty* SourceProperty = ArrayProperty->Inner;
+					for(int32 ElementIndex = 0; ElementIndex < ArrayHelper.Num(); ElementIndex++)
+					{
+						uint8* TargetMemory = ClonedArrayHelper.GetRawPtr(ElementIndex);
+						const uint8* SourceMemory = ArrayHelper.GetRawPtr(ElementIndex);
+						URigVMMemoryStorage::CopyProperty(TargetProperty, TargetMemory, SourceProperty, SourceMemory);
+					}
+				}
+#endif
+				Context.InstructionIndex++;
+				break;
+			}
+			case ERigVMOpCode::ArrayGetAtIndex:
+			{
+#if !UE_RIGVM_UCLASS_BASED_STORAGE_DISABLED
+				FRigVMMemoryHandle& ArrayHandle = CachedMemoryHandles[FirstHandleForInstruction[Context.InstructionIndex]];
+				const FArrayProperty* ArrayProperty = CastFieldChecked<FArrayProperty>(ArrayHandle.GetProperty());
+				FScriptArrayHelper ArrayHelper(ArrayProperty, ArrayHandle.GetData());
+				const int32 Index = (*((int32*)CachedMemoryHandles[FirstHandleForInstruction[Context.InstructionIndex] + 1].GetData()));
+				if(Context.IsValidArrayIndex(Index, ArrayHelper.Num()))
+				{
+					FRigVMMemoryHandle& ElementHandle = CachedMemoryHandles[FirstHandleForInstruction[Context.InstructionIndex] + 2];
+					uint8* TargetMemory = ElementHandle.GetData();
+					const uint8* SourceMemory = ArrayHelper.GetRawPtr(Index);
+					URigVMMemoryStorage::CopyProperty(ElementHandle.GetProperty(), TargetMemory, ArrayProperty->Inner, SourceMemory);
+				}
+#endif
+				Context.InstructionIndex++;
+				break;
+			}
+			case ERigVMOpCode::ArraySetAtIndex:
+			{
+#if !UE_RIGVM_UCLASS_BASED_STORAGE_DISABLED
+				FRigVMMemoryHandle& ArrayHandle = CachedMemoryHandles[FirstHandleForInstruction[Context.InstructionIndex]];
+				const FArrayProperty* ArrayProperty = CastFieldChecked<FArrayProperty>(ArrayHandle.GetProperty());
+				FScriptArrayHelper ArrayHelper(ArrayProperty, ArrayHandle.GetData());
+				const int32 Index = (*((int32*)CachedMemoryHandles[FirstHandleForInstruction[Context.InstructionIndex] + 1].GetData()));
+				if(Context.IsValidArrayIndex(Index, ArrayHelper.Num()))
+				{
+					FRigVMMemoryHandle& ElementHandle = CachedMemoryHandles[FirstHandleForInstruction[Context.InstructionIndex] + 2];
+					uint8* TargetMemory = ArrayHelper.GetRawPtr(Index);
+					const uint8* SourceMemory = ElementHandle.GetData();
+					URigVMMemoryStorage::CopyProperty(ArrayProperty->Inner, TargetMemory, ElementHandle.GetProperty(), SourceMemory);
+				}
+#endif
+				Context.InstructionIndex++;
+				break;
+			}
+			case ERigVMOpCode::ArrayInsert:
+			{
+#if !UE_RIGVM_UCLASS_BASED_STORAGE_DISABLED
+				FRigVMMemoryHandle& ArrayHandle = CachedMemoryHandles[FirstHandleForInstruction[Context.InstructionIndex]];
+				const FArrayProperty* ArrayProperty = CastFieldChecked<FArrayProperty>(ArrayHandle.GetProperty());
+				FScriptArrayHelper ArrayHelper(ArrayProperty, ArrayHandle.GetData());
+				if(Context.IsValidArraySize(ArrayHelper.Num() + 1))
+				{
+					int32 Index = (*((int32*)CachedMemoryHandles[FirstHandleForInstruction[Context.InstructionIndex] + 1].GetData()));
+					Index = FMath::Clamp<int32>(Index, 0, ArrayHelper.Num());
+					ArrayHelper.InsertValues(Index, 1);
+
+					FRigVMMemoryHandle& ElementHandle = CachedMemoryHandles[FirstHandleForInstruction[Context.InstructionIndex] + 2];
+					uint8* TargetMemory = ArrayHelper.GetRawPtr(Index);
+					const uint8* SourceMemory = ElementHandle.GetData();
+					URigVMMemoryStorage::CopyProperty(ArrayProperty->Inner, TargetMemory, ElementHandle.GetProperty(), SourceMemory);
+				}
+#endif
+				Context.InstructionIndex++;
+				break;
+			}
+			case ERigVMOpCode::ArrayRemove:
+			{
+#if !UE_RIGVM_UCLASS_BASED_STORAGE_DISABLED
+				FRigVMMemoryHandle& ArrayHandle = CachedMemoryHandles[FirstHandleForInstruction[Context.InstructionIndex]];
+				const FArrayProperty* ArrayProperty = CastFieldChecked<FArrayProperty>(ArrayHandle.GetProperty());
+				FScriptArrayHelper ArrayHelper(ArrayProperty, ArrayHandle.GetData());
+				const int32 Index = (*((int32*)CachedMemoryHandles[FirstHandleForInstruction[Context.InstructionIndex] + 1].GetData()));
+				if(Context.IsValidArrayIndex(Index, ArrayHelper.Num()))
+				{
+					ArrayHelper.RemoveValues(Index, 1);
+				}
+#endif
+				Context.InstructionIndex++;
+				break;
+			}
+			case ERigVMOpCode::ArrayAdd:
+			{
+#if !UE_RIGVM_UCLASS_BASED_STORAGE_DISABLED
+				FRigVMMemoryHandle& ArrayHandle = CachedMemoryHandles[FirstHandleForInstruction[Context.InstructionIndex]];
+				const FArrayProperty* ArrayProperty = CastFieldChecked<FArrayProperty>(ArrayHandle.GetProperty());
+				FScriptArrayHelper ArrayHelper(ArrayProperty, ArrayHandle.GetData());
+				int32& Index = (*((int32*)CachedMemoryHandles[FirstHandleForInstruction[Context.InstructionIndex] + 2].GetData()));
+				if(Context.IsValidArraySize(ArrayHelper.Num() + 1))
+				{
+					FRigVMMemoryHandle& ElementHandle = CachedMemoryHandles[FirstHandleForInstruction[Context.InstructionIndex] + 1];
+					Index = ArrayHelper.AddValue();
+
+					uint8* TargetMemory = ArrayHelper.GetRawPtr(Index);
+					const uint8* SourceMemory = ElementHandle.GetData();
+					URigVMMemoryStorage::CopyProperty(ArrayProperty->Inner, TargetMemory, ElementHandle.GetProperty(), SourceMemory);
+				}
+				else
+				{
+					Index = INDEX_NONE;
+				}
+#endif
+				Context.InstructionIndex++;
+				break;
+			}
+			case ERigVMOpCode::ArrayFind:
+			{
+#if !UE_RIGVM_UCLASS_BASED_STORAGE_DISABLED
+				FRigVMMemoryHandle& ArrayHandle = CachedMemoryHandles[FirstHandleForInstruction[Context.InstructionIndex]];
+				const FArrayProperty* ArrayProperty = CastFieldChecked<FArrayProperty>(ArrayHandle.GetProperty());
+				FScriptArrayHelper ArrayHelper(ArrayProperty, ArrayHandle.GetData());
+
+				FRigVMMemoryHandle& ElementHandle = CachedMemoryHandles[FirstHandleForInstruction[Context.InstructionIndex] + 1];
+				int32& FoundIndex = (*((int32*)CachedMemoryHandles[FirstHandleForInstruction[Context.InstructionIndex] + 2].GetData()));
+				bool& bFound = (*((bool*)CachedMemoryHandles[FirstHandleForInstruction[Context.InstructionIndex] + 3].GetData()));
+
+				FoundIndex = INDEX_NONE;
+				bFound = false;
+
+				const FProperty* PropertyA = ElementHandle.GetProperty();
+				const FProperty* PropertyB = ArrayProperty->Inner;
+
+				if(PropertyA->SameType(PropertyB))
+				{
+					const uint8* MemoryA = ElementHandle.GetData();
+
+					for(int32 Index = 0; Index < ArrayHelper.Num(); Index++)
+					{
+						const uint8* MemoryB = ArrayHelper.GetRawPtr(Index);
+						if(PropertyA->Identical(MemoryA, MemoryB))
+						{
+							FoundIndex = Index;
+							bFound = true;
+							break;
+						}
+					}
+				}
+				else
+				{
+					static const TCHAR IncompatibleTypes[] = TEXT("Array('%s') doesn't support searching for element('%$s').");
+					Context.Logf(EMessageSeverity::Error, IncompatibleTypes, *PropertyB->GetCPPType(), *PropertyA->GetCPPType());
+				}
+#endif
+				Context.InstructionIndex++;
+				break;
+			}
+			case ERigVMOpCode::ArrayIterator:
+			{
+#if !UE_RIGVM_UCLASS_BASED_STORAGE_DISABLED
+				FRigVMMemoryHandle& ArrayHandle = CachedMemoryHandles[FirstHandleForInstruction[Context.InstructionIndex]];
+				const FArrayProperty* ArrayProperty = CastFieldChecked<FArrayProperty>(ArrayHandle.GetProperty());
+				FScriptArrayHelper ArrayHelper(ArrayProperty, ArrayHandle.GetData());
+
+				FRigVMMemoryHandle& ElementHandle = CachedMemoryHandles[FirstHandleForInstruction[Context.InstructionIndex] + 1];
+				const int32& Index = (*((int32*)CachedMemoryHandles[FirstHandleForInstruction[Context.InstructionIndex] + 2].GetData()));
+				int32& Count = (*((int32*)CachedMemoryHandles[FirstHandleForInstruction[Context.InstructionIndex] + 3].GetData()));
+				float& Ratio = (*((float*)CachedMemoryHandles[FirstHandleForInstruction[Context.InstructionIndex] + 4].GetData()));
+				bool& bContinue = (*((bool*)CachedMemoryHandles[FirstHandleForInstruction[Context.InstructionIndex] + 5].GetData()));
+
+				Count = ArrayHelper.Num();
+				bContinue = Index >=0 && Index < Count;
+
+				if((Count <= 0) || !bContinue)
+				{
+					Ratio = 0.f;
+				}
+				else
+				{
+					Ratio = float(Index) / float(Count - 1);
+				}
+
+#endif
+				Context.InstructionIndex++;
+				break;
+			}
 			case ERigVMOpCode::Invalid:
 			{
 				ensure(false);
diff --git a/Engine/Source/Runtime/RigVM/Private/RigVMCore/RigVMByteCode.cpp b/Engine/Source/Runtime/RigVM/Private/RigVMCore/RigVMByteCode.cpp
--- a/Engine/Source/Runtime/RigVM/Private/RigVMCore/RigVMByteCode.cpp
+++ b/Engine/Source/Runtime/RigVM/Private/RigVMCore/RigVMByteCode.cpp
@@ -24,6 +24,44 @@ void FRigVMBinaryOp::Serialize(FArchive& Ar)
 	Ar << ArgB;
 }
 
+void FRigVMTernaryOp::Serialize(FArchive& Ar)
+{
+	Ar << OpCode;
+	Ar << ArgA;
+	Ar << ArgB;
+	Ar << ArgC;
+}
+
+void FRigVMQuaternaryOp::Serialize(FArchive& Ar)
+{
+	Ar << OpCode;
+	Ar << ArgA;
+	Ar << ArgB;
+	Ar << ArgC;
+	Ar << ArgD;
+}
+
+void FRigVMQuinaryOp::Serialize(FArchive& Ar)
+{
+	Ar << OpCode;
+	Ar << ArgA;
+	Ar << ArgB;
+	Ar << ArgC;
+	Ar << ArgD;
+	Ar << ArgE;
+}
+
+void FRigVMSenaryOp::Serialize(FArchive& Ar)
+{
+	Ar << OpCode;
+	Ar << ArgA;
+	Ar << ArgB;
+	Ar << ArgC;
+	Ar << ArgD;
+	Ar << ArgE;
+	Ar << ArgF;
+}
+
 void FRigVMCopyOp::Serialize(FArchive& Ar)
 {
 	Ar.UsingCustomVersion(FUE5MainStreamObjectVersion::GUID);
@@ -285,6 +323,7 @@ void FRigVMByteCode::Save(FArchive& Ar) const
 			case ERigVMOpCode::BoolTrue:
 			case ERigVMOpCode::Increment:
 			case ERigVMOpCode::Decrement:
+			case ERigVMOpCode::ArrayReset:
 			{
 				FRigVMUnaryOp Op = GetOpAt<FRigVMUnaryOp>(Instruction.ByteCodeIndex);
 				Ar << Op;
@@ -320,11 +359,37 @@ void FRigVMByteCode::Save(FArchive& Ar) const
 				break;
 			}
 			case ERigVMOpCode::BeginBlock:
+			case ERigVMOpCode::ArrayGetNum:
+			case ERigVMOpCode::ArraySetNum:
+			case ERigVMOpCode::ArrayAppend:
+			case ERigVMOpCode::ArrayClone:
+			case ERigVMOpCode::ArrayRemove:
 			{
 				FRigVMBinaryOp Op = GetOpAt<FRigVMBinaryOp>(Instruction.ByteCodeIndex);
 				Ar << Op;
 				break;
 			}
+			case ERigVMOpCode::ArrayAdd:
+			case ERigVMOpCode::ArrayGetAtIndex:
+			case ERigVMOpCode::ArraySetAtIndex:
+			case ERigVMOpCode::ArrayInsert:
+			{
+				FRigVMTernaryOp Op = GetOpAt<FRigVMTernaryOp>(Instruction.ByteCodeIndex);
+				Ar << Op;
+				break;
+			}
+			case ERigVMOpCode::ArrayFind:
+			{
+				FRigVMQuaternaryOp Op = GetOpAt<FRigVMQuaternaryOp>(Instruction.ByteCodeIndex);
+				Ar << Op;
+				break;
+			}
+			case ERigVMOpCode::ArrayIterator:
+			{
+				FRigVMSenaryOp Op = GetOpAt<FRigVMSenaryOp>(Instruction.ByteCodeIndex);
+				Ar << Op;
+				break;
+			}
 			case ERigVMOpCode::EndBlock:
 			{
 				// nothing todo, the EndBlock has no custom data inside of it
@@ -472,6 +537,7 @@ void FRigVMByteCode::Load(FArchive& Ar)
 			case ERigVMOpCode::BoolTrue:
 			case ERigVMOpCode::Increment:
 			case ERigVMOpCode::Decrement:
+			case ERigVMOpCode::ArrayReset:
 			{
 				FRigVMUnaryOp Op;
 				Ar << Op;
@@ -510,12 +576,41 @@ void FRigVMByteCode::Load(FArchive& Ar)
 				break;
 			}
 			case ERigVMOpCode::BeginBlock:
+			case ERigVMOpCode::ArrayGetNum:
+			case ERigVMOpCode::ArraySetNum:
+			case ERigVMOpCode::ArrayAppend:
+			case ERigVMOpCode::ArrayClone:
+			case ERigVMOpCode::ArrayRemove:
 			{
 				FRigVMBinaryOp Op;
 				Ar << Op;
 				AddOp<FRigVMBinaryOp>(Op);
 				break;
 			}
+			case ERigVMOpCode::ArrayAdd:
+			case ERigVMOpCode::ArrayGetAtIndex:
+			case ERigVMOpCode::ArraySetAtIndex:
+			case ERigVMOpCode::ArrayInsert:
+			{
+				FRigVMTernaryOp Op;
+				Ar << Op;
+				AddOp<FRigVMTernaryOp>(Op);
+				break;
+			}
+			case ERigVMOpCode::ArrayFind:
+			{
+				FRigVMQuaternaryOp Op;
+				Ar << Op;
+				AddOp<FRigVMQuaternaryOp>(Op);
+				break;
+			}
+			case ERigVMOpCode::ArrayIterator:
+			{
+				FRigVMSenaryOp Op;
+				Ar << Op;
+				AddOp<FRigVMSenaryOp>(Op);
+				break;
+			}
 			case ERigVMOpCode::EndBlock:
 			{
 				AddEndBlockOp();
@@ -711,6 +806,7 @@ uint64 FRigVMByteCode::GetOpNumBytesAt(uint64 InByteCodeIndex, bool bIncludeOper
 		case ERigVMOpCode::BoolTrue:
 		case ERigVMOpCode::Increment:
 		case ERigVMOpCode::Decrement:
+		case ERigVMOpCode::ArrayReset:
 		{
 			return (uint64)sizeof(FRigVMUnaryOp);
 		}
@@ -745,9 +841,29 @@ uint64 FRigVMByteCode::GetOpNumBytesAt(uint64 InByteCodeIndex, bool bIncludeOper
 			return (uint64)sizeof(FRigVMBaseOp);
 		}
 		case ERigVMOpCode::BeginBlock:
+		case ERigVMOpCode::ArrayGetNum:
+		case ERigVMOpCode::ArraySetNum:
+		case ERigVMOpCode::ArrayAppend:
+		case ERigVMOpCode::ArrayClone:
+		case ERigVMOpCode::ArrayRemove:
 		{
 			return (uint64)sizeof(FRigVMBinaryOp);
 		}
+		case ERigVMOpCode::ArrayAdd:
+		case ERigVMOpCode::ArrayGetAtIndex:
+		case ERigVMOpCode::ArraySetAtIndex:
+		case ERigVMOpCode::ArrayInsert:
+		{				
+			return (uint64)sizeof(FRigVMTernaryOp);
+		}
+		case ERigVMOpCode::ArrayFind:
+		{				
+			return (uint64)sizeof(FRigVMQuaternaryOp);
+		}
+		case ERigVMOpCode::ArrayIterator:
+		{				
+			return (uint64)sizeof(FRigVMSenaryOp);
+		}
 		case ERigVMOpCode::EndBlock:
 		{
 			return (uint64)sizeof(FRigVMBaseOp);
@@ -993,6 +1109,7 @@ FString FRigVMByteCode::DumpToText() const
 			case ERigVMOpCode::BoolTrue:
 			case ERigVMOpCode::Increment:
 			case ERigVMOpCode::Decrement:
+			case ERigVMOpCode::ArrayReset:
 			{
 				const FRigVMUnaryOp& Op = GetOpAt<FRigVMUnaryOp>(Instruction.ByteCodeIndex);
 				FString ArgContent;
@@ -1039,10 +1156,76 @@ FString FRigVMByteCode::DumpToText() const
 			{
 				break;
 			}
+			case ERigVMOpCode::BeginBlock:
+			case ERigVMOpCode::ArrayGetNum:
+			case ERigVMOpCode::ArraySetNum:
+			case ERigVMOpCode::ArrayAppend:
+			case ERigVMOpCode::ArrayClone:
+			case ERigVMOpCode::ArrayRemove:
+			{
+				const FRigVMBinaryOp& Op = GetOpAt<FRigVMBinaryOp>(Instruction.ByteCodeIndex);
+				FString ArgA, ArgB;
+				FRigVMOperand::StaticStruct()->ExportText(ArgA, &Op.ArgA, nullptr, nullptr, PPF_None, nullptr);
+				FRigVMOperand::StaticStruct()->ExportText(ArgB, &Op.ArgB, nullptr, nullptr, PPF_None, nullptr);
+				Line += FString::Printf(TEXT(", ArgA %s"), *ArgA);
+				Line += FString::Printf(TEXT(", ArgB %s"), *ArgB);
+				break;
+			}
+			case ERigVMOpCode::ArrayAdd:
+			case ERigVMOpCode::ArrayGetAtIndex:
+			case ERigVMOpCode::ArraySetAtIndex:
+			case ERigVMOpCode::ArrayInsert:
+			{
+				const FRigVMTernaryOp& Op = GetOpAt<FRigVMTernaryOp>(Instruction.ByteCodeIndex);
+				FString ArgA, ArgB, ArgC;
+				FRigVMOperand::StaticStruct()->ExportText(ArgA, &Op.ArgA, nullptr, nullptr, PPF_None, nullptr);
+				FRigVMOperand::StaticStruct()->ExportText(ArgB, &Op.ArgB, nullptr, nullptr, PPF_None, nullptr);
+				FRigVMOperand::StaticStruct()->ExportText(ArgC, &Op.ArgC, nullptr, nullptr, PPF_None, nullptr);
+				Line += FString::Printf(TEXT(", ArgA %s"), *ArgA);
+				Line += FString::Printf(TEXT(", ArgB %s"), *ArgB);
+				Line += FString::Printf(TEXT(", ArgC %s"), *ArgC);
+				break;
+			}
+			case ERigVMOpCode::ArrayFind:
+			{
+				const FRigVMQuaternaryOp& Op = GetOpAt<FRigVMQuaternaryOp>(Instruction.ByteCodeIndex);
+				FString ArgA, ArgB, ArgC, ArgD;
+				FRigVMOperand::StaticStruct()->ExportText(ArgA, &Op.ArgA, nullptr, nullptr, PPF_None, nullptr);
+				FRigVMOperand::StaticStruct()->ExportText(ArgB, &Op.ArgB, nullptr, nullptr, PPF_None, nullptr);
+				FRigVMOperand::StaticStruct()->ExportText(ArgC, &Op.ArgC, nullptr, nullptr, PPF_None, nullptr);
+				FRigVMOperand::StaticStruct()->ExportText(ArgD, &Op.ArgD, nullptr, nullptr, PPF_None, nullptr);
+				Line += FString::Printf(TEXT(", ArgA %s"), *ArgA);
+				Line += FString::Printf(TEXT(", ArgB %s"), *ArgB);
+				Line += FString::Printf(TEXT(", ArgC %s"), *ArgC);
+				Line += FString::Printf(TEXT(", ArgD %s"), *ArgD);
+				break;
+			}
+			case ERigVMOpCode::ArrayIterator:
+			{
+				const FRigVMSenaryOp& Op = GetOpAt<FRigVMSenaryOp>(Instruction.ByteCodeIndex);
+				FString ArgA, ArgB, ArgC, ArgD, ArgE, ArgF;
+				FRigVMOperand::StaticStruct()->ExportText(ArgA, &Op.ArgA, nullptr, nullptr, PPF_None, nullptr);
+				FRigVMOperand::StaticStruct()->ExportText(ArgB, &Op.ArgB, nullptr, nullptr, PPF_None, nullptr);
+				FRigVMOperand::StaticStruct()->ExportText(ArgC, &Op.ArgC, nullptr, nullptr, PPF_None, nullptr);
+				FRigVMOperand::StaticStruct()->ExportText(ArgD, &Op.ArgD, nullptr, nullptr, PPF_None, nullptr);
+				FRigVMOperand::StaticStruct()->ExportText(ArgE, &Op.ArgE, nullptr, nullptr, PPF_None, nullptr);
+				FRigVMOperand::StaticStruct()->ExportText(ArgF, &Op.ArgF, nullptr, nullptr, PPF_None, nullptr);
+				Line += FString::Printf(TEXT(", ArgA %s"), *ArgA);
+				Line += FString::Printf(TEXT(", ArgB %s"), *ArgB);
+				Line += FString::Printf(TEXT(", ArgC %s"), *ArgC);
+				Line += FString::Printf(TEXT(", ArgD %s"), *ArgD);
+				Line += FString::Printf(TEXT(", ArgE %s"), *ArgE);
+				Line += FString::Printf(TEXT(", ArgF %s"), *ArgF);
+				break;
+			}
 			case ERigVMOpCode::Invalid:
 			{
 				ensure(false);
 			}
+			default:
+			{
+				break;
+			}
 		}
 
 		Lines.Add(Line);
@@ -1068,6 +1251,65 @@ uint64 FRigVMByteCode::AddEndBlockOp()
 	return AddOp(Op);
 }
 
+uint64 FRigVMByteCode::AddArrayResetOp(FRigVMOperand InArrayArg)
+{
+	return AddOp(FRigVMUnaryOp(ERigVMOpCode::ArrayReset, InArrayArg));
+}
+
+uint64 FRigVMByteCode::AddArrayGetNumOp(FRigVMOperand InArrayArg, FRigVMOperand InNumArg)
+{
+	return AddOp(FRigVMBinaryOp(ERigVMOpCode::ArrayGetNum, InArrayArg, InNumArg));
+}
+
+uint64 FRigVMByteCode::AddArraySetNumOp(FRigVMOperand InArrayArg, FRigVMOperand InNumArg)
+{
+	return AddOp(FRigVMBinaryOp(ERigVMOpCode::ArraySetNum, InArrayArg, InNumArg));
+}
+
+uint64 FRigVMByteCode::AddArrayGetAtIndexOp(FRigVMOperand InArrayArg, FRigVMOperand InIndexArg, FRigVMOperand InElementArg)
+{
+	return AddOp(FRigVMTernaryOp(ERigVMOpCode::ArrayGetAtIndex, InArrayArg, InIndexArg, InElementArg));
+}
+
+uint64 FRigVMByteCode::AddArraySetAtIndexOp(FRigVMOperand InArrayArg, FRigVMOperand InIndexArg, FRigVMOperand InElementArg)
+{
+	return AddOp(FRigVMTernaryOp(ERigVMOpCode::ArraySetAtIndex, InArrayArg, InIndexArg, InElementArg));
+}
+
+uint64 FRigVMByteCode::AddArrayAddOp(FRigVMOperand InArrayArg, FRigVMOperand InElementArg, FRigVMOperand InIndexArg)
+{
+	return AddOp(FRigVMTernaryOp(ERigVMOpCode::ArrayAdd, InArrayArg, InElementArg, InIndexArg));	
+}
+
+uint64 FRigVMByteCode::AddArrayInsertOp(FRigVMOperand InArrayArg, FRigVMOperand InIndexArg, FRigVMOperand InElementArg)
+{
+	return AddOp(FRigVMTernaryOp(ERigVMOpCode::ArrayInsert, InArrayArg, InIndexArg, InElementArg));
+}
+
+uint64 FRigVMByteCode::AddArrayRemoveOp(FRigVMOperand InArrayArg, FRigVMOperand InIndexArg)
+{
+	return AddOp(FRigVMBinaryOp(ERigVMOpCode::ArrayRemove, InArrayArg, InIndexArg));
+}
+
+uint64 FRigVMByteCode::AddArrayFindOp(FRigVMOperand InArrayArg, FRigVMOperand InElementArg, FRigVMOperand InIndexArg, FRigVMOperand InSuccessArg)
+{
+	return AddOp(FRigVMQuaternaryOp(ERigVMOpCode::ArrayFind, InArrayArg, InElementArg, InIndexArg, InSuccessArg));	
+}
+
+uint64 FRigVMByteCode::AddArrayAppendOp(FRigVMOperand InArrayArg, FRigVMOperand InOtherArrayArg)
+{
+	return AddOp(FRigVMBinaryOp(ERigVMOpCode::ArrayAppend, InArrayArg, InOtherArrayArg));
+}
+
+uint64 FRigVMByteCode::AddArrayCloneOp(FRigVMOperand InArrayArg, FRigVMOperand InClonedArrayArg)
+{
+	return AddOp(FRigVMBinaryOp(ERigVMOpCode::ArrayClone, InArrayArg, InClonedArrayArg));
+}
+
+uint64 FRigVMByteCode::AddArrayIteratorOp(FRigVMOperand InArrayArg, FRigVMOperand InElementArg, FRigVMOperand InIndexArg, FRigVMOperand InCountArg, FRigVMOperand InRatioArg, FRigVMOperand InContinueArg)
+{
+	return AddOp(FRigVMSenaryOp(ERigVMOpCode::ArrayIterator, InArrayArg, InElementArg, InIndexArg, InCountArg, InRatioArg, InContinueArg));
+}
 
 uint64 FRigVMByteCode::GetOpAlignment(ERigVMOpCode InOpCode) const
 {
@@ -1152,6 +1394,7 @@ uint64 FRigVMByteCode::GetOpAlignment(ERigVMOpCode InOpCode) const
 		case ERigVMOpCode::BoolTrue:
 		case ERigVMOpCode::Increment:
 		case ERigVMOpCode::Decrement:
+		case ERigVMOpCode::ArrayReset:
 		{
 			static const uint64 Alignment = FRigVMUnaryOp::StaticStruct()->GetCppStructOps()->GetAlignment();
 			return Alignment;
@@ -1192,10 +1435,33 @@ uint64 FRigVMByteCode::GetOpAlignment(ERigVMOpCode InOpCode) const
 			return Alignment;
 		}
 		case ERigVMOpCode::BeginBlock:
+		case ERigVMOpCode::ArrayGetNum:
+		case ERigVMOpCode::ArraySetNum:
+		case ERigVMOpCode::ArrayAppend:
+		case ERigVMOpCode::ArrayClone:
+		case ERigVMOpCode::ArrayRemove:
 		{
 			static const uint64 Alignment = FRigVMBinaryOp::StaticStruct()->GetCppStructOps()->GetAlignment();
 			return Alignment;
 		}
+		case ERigVMOpCode::ArrayAdd:
+		case ERigVMOpCode::ArrayGetAtIndex:
+		case ERigVMOpCode::ArraySetAtIndex:
+		case ERigVMOpCode::ArrayInsert:
+		{				
+			static const uint64 Alignment = FRigVMTernaryOp::StaticStruct()->GetCppStructOps()->GetAlignment();
+			return Alignment;
+		}
+		case ERigVMOpCode::ArrayFind:
+		{				
+			static const uint64 Alignment = FRigVMQuaternaryOp::StaticStruct()->GetCppStructOps()->GetAlignment();
+			return Alignment;
+		}
+		case ERigVMOpCode::ArrayIterator:
+		{				
+			static const uint64 Alignment = FRigVMSenaryOp::StaticStruct()->GetCppStructOps()->GetAlignment();
+			return Alignment;
+		}
 		case ERigVMOpCode::EndBlock:
 		{
 			static const uint64 Alignment = FRigVMBaseOp::StaticStruct()->GetCppStructOps()->GetAlignment();
diff --git a/Engine/Source/Runtime/RigVM/Public/RigVMCore/RigVM.h b/Engine/Source/Runtime/RigVM/Public/RigVMCore/RigVM.h
--- a/Engine/Source/Runtime/RigVM/Public/RigVMCore/RigVM.h
+++ b/Engine/Source/Runtime/RigVM/Public/RigVMCore/RigVM.h
@@ -160,6 +160,12 @@ public:
 	// resets the container and clones the input VM
 	void CopyFrom(URigVM* InVM, bool bDeferCopy = false, bool bReferenceLiteralMemory = false, bool bReferenceByteCode = false, bool bCopyExternalVariables = false, bool bCopyDynamicRegisters = false);
 
+	// sets the max array size allowed by this VM
+	FORCEINLINE void SetRuntimeSettings(FRigVMRuntimeSettings InRuntimeSettings)
+	{
+		Context.SetRuntimeSettings(InRuntimeSettings);
+	}
+
 	// Initializes all execute ops and their memory.
 #if UE_RIGVM_UCLASS_BASED_STORAGE_DISABLED
 	bool Initialize(FRigVMMemoryContainerPtrArray Memory, FRigVMFixedArray<void*> AdditionalArguments);
diff --git a/Engine/Source/Runtime/RigVM/Public/RigVMCore/RigVMByteCode.h b/Engine/Source/Runtime/RigVM/Public/RigVMCore/RigVMByteCode.h
--- a/Engine/Source/Runtime/RigVM/Public/RigVMCore/RigVMByteCode.h
+++ b/Engine/Source/Runtime/RigVM/Public/RigVMCore/RigVMByteCode.h
@@ -99,6 +99,18 @@ enum class ERigVMOpCode : uint8
 	Exit, // exit the execution loop
 	BeginBlock, // begins a new memory slice / block
 	EndBlock, // ends the last memory slice / block
+	ArrayReset, // reads and returns the size of an array (unary op, in out array)
+	ArrayGetNum, // reads and returns the size of an array (binary op, in array, out int32) 
+	ArraySetNum, // resizes an array (binary op, in out array, in int32)
+	ArrayGetAtIndex, // returns an array element by index (ternary op, in array, in int32, out element)  
+	ArraySetAtIndex, // sets an array element by index (ternary op, in out array, in int32, in element)
+	ArrayAdd, // adds an element to an array (ternary op, in out array, in element, out int32 index)
+	ArrayInsert, // inserts an element to an array (ternary op, in out array, in int32, in element)
+	ArrayRemove, // removes an element from an array (binary op, in out array, in inindex)
+	ArrayFind, // finds and returns the index of an element (quaternery op, in array, in element, out int32 index, out bool success)
+	ArrayAppend, // appends an array to another (binary op, in out array, in array)
+	ArrayClone, // clones an array (binary op, in array, out array)
+	ArrayIterator, // iterates over an array (senary op, in array, out element, out index, out count, out ratio, out continue)
 	Invalid
 };
 
@@ -172,7 +184,8 @@ struct RIGVM_API FRigVMUnaryOp : public FRigVMBaseOp
 			uint8(InOpCode) == uint8(ERigVMOpCode::JumpAbsoluteIf) ||
 			uint8(InOpCode) == uint8(ERigVMOpCode::JumpForwardIf) ||
 			uint8(InOpCode) == uint8(ERigVMOpCode::JumpBackwardIf) ||
-			uint8(InOpCode) == uint8(ERigVMOpCode::ChangeType)
+			uint8(InOpCode) == uint8(ERigVMOpCode::ChangeType) ||
+			uint8(InOpCode) == uint8(ERigVMOpCode::ArrayReset)
 		);
 	}
 
@@ -186,7 +199,7 @@ struct RIGVM_API FRigVMUnaryOp : public FRigVMBaseOp
 	}
 };
 
-// operator used for beginblock
+// operator used for beginblock and array reset
 USTRUCT()
 struct RIGVM_API FRigVMBinaryOp : public FRigVMBaseOp
 {
@@ -205,7 +218,12 @@ struct RIGVM_API FRigVMBinaryOp : public FRigVMBaseOp
 		, ArgB(InArgB)
 	{
 		ensure(
-			uint8(InOpCode) == uint8(ERigVMOpCode::BeginBlock)
+			uint8(InOpCode) == uint8(ERigVMOpCode::BeginBlock) ||
+			uint8(InOpCode) == uint8(ERigVMOpCode::ArrayGetNum) ||
+			uint8(InOpCode) == uint8(ERigVMOpCode::ArraySetNum) ||
+			uint8(InOpCode) == uint8(ERigVMOpCode::ArrayAppend) ||
+			uint8(InOpCode) == uint8(ERigVMOpCode::ArrayClone) ||
+			uint8(InOpCode) == uint8(ERigVMOpCode::ArrayRemove)
 		);
 	}
 
@@ -220,6 +238,172 @@ struct RIGVM_API FRigVMBinaryOp : public FRigVMBaseOp
 	}
 };
 
+// operator used for some array operations
+USTRUCT()
+struct RIGVM_API FRigVMTernaryOp : public FRigVMBaseOp
+{
+	GENERATED_USTRUCT_BODY()
+
+	FRigVMTernaryOp()
+		: FRigVMBaseOp(ERigVMOpCode::Invalid)
+		, ArgA()
+		, ArgB()
+		, ArgC()
+	{
+	}
+
+	FRigVMTernaryOp(ERigVMOpCode InOpCode, FRigVMOperand InArgA, FRigVMOperand InArgB, FRigVMOperand InArgC)
+		: FRigVMBaseOp(InOpCode)
+		, ArgA(InArgA)
+		, ArgB(InArgB)
+		, ArgC(InArgC)
+	{
+		ensure(
+			uint8(InOpCode) == uint8(ERigVMOpCode::ArrayAdd) ||
+			uint8(InOpCode) == uint8(ERigVMOpCode::ArrayGetAtIndex) ||
+			uint8(InOpCode) == uint8(ERigVMOpCode::ArraySetAtIndex) ||
+			uint8(InOpCode) == uint8(ERigVMOpCode::ArrayInsert)
+		);
+	}
+
+	FRigVMOperand ArgA;
+	FRigVMOperand ArgB;
+	FRigVMOperand ArgC;
+
+	void Serialize(FArchive& Ar);
+	FORCEINLINE friend FArchive& operator<<(FArchive& Ar, FRigVMTernaryOp& P)
+	{
+		P.Serialize(Ar);
+		return Ar;
+	}
+};
+
+// operator used for some array operations
+USTRUCT()
+struct RIGVM_API FRigVMQuaternaryOp : public FRigVMBaseOp
+{
+	GENERATED_USTRUCT_BODY()
+
+	FRigVMQuaternaryOp()
+		: FRigVMBaseOp(ERigVMOpCode::Invalid)
+		, ArgA()
+		, ArgB()
+		, ArgC()
+		, ArgD()
+	{
+	}
+
+	FRigVMQuaternaryOp(ERigVMOpCode InOpCode, FRigVMOperand InArgA, FRigVMOperand InArgB, FRigVMOperand InArgC, FRigVMOperand InArgD)
+		: FRigVMBaseOp(InOpCode)
+		, ArgA(InArgA)
+		, ArgB(InArgB)
+		, ArgC(InArgC)
+		, ArgD(InArgD)
+	{
+		ensure(
+			uint8(InOpCode) == uint8(ERigVMOpCode::ArrayFind)
+		);
+	}
+
+	FRigVMOperand ArgA;
+	FRigVMOperand ArgB;
+	FRigVMOperand ArgC;
+	FRigVMOperand ArgD;
+
+	void Serialize(FArchive& Ar);
+	FORCEINLINE friend FArchive& operator<<(FArchive& Ar, FRigVMQuaternaryOp& P)
+	{
+		P.Serialize(Ar);
+		return Ar;
+	}
+};
+
+// operator used for some array operations
+USTRUCT()
+struct RIGVM_API FRigVMQuinaryOp : public FRigVMBaseOp
+{
+	GENERATED_USTRUCT_BODY()
+
+	FRigVMQuinaryOp()
+		: FRigVMBaseOp(ERigVMOpCode::Invalid)
+		, ArgA()
+		, ArgB()
+		, ArgC()
+		, ArgD()
+		, ArgE()
+	{
+	}
+
+	FRigVMQuinaryOp(ERigVMOpCode InOpCode, FRigVMOperand InArgA, FRigVMOperand InArgB, FRigVMOperand InArgC, FRigVMOperand InArgD, FRigVMOperand InArgE)
+		: FRigVMBaseOp(InOpCode)
+		, ArgA(InArgA)
+		, ArgB(InArgB)
+		, ArgC(InArgC)
+		, ArgD(InArgD)
+		, ArgE(InArgE)
+	{
+	}
+
+	FRigVMOperand ArgA;
+	FRigVMOperand ArgB;
+	FRigVMOperand ArgC;
+	FRigVMOperand ArgD;
+	FRigVMOperand ArgE;
+
+	void Serialize(FArchive& Ar);
+	FORCEINLINE friend FArchive& operator<<(FArchive& Ar, FRigVMQuinaryOp& P)
+	{
+		P.Serialize(Ar);
+		return Ar;
+	}
+};
+
+// operator used for some array operations
+USTRUCT()
+struct RIGVM_API FRigVMSenaryOp : public FRigVMBaseOp
+{
+	GENERATED_USTRUCT_BODY()
+
+	FRigVMSenaryOp()
+		: FRigVMBaseOp(ERigVMOpCode::Invalid)
+		, ArgA()
+		, ArgB()
+		, ArgC()
+		, ArgD()
+		, ArgE()
+		, ArgF()
+	{
+	}
+
+	FRigVMSenaryOp(ERigVMOpCode InOpCode, FRigVMOperand InArgA, FRigVMOperand InArgB, FRigVMOperand InArgC, FRigVMOperand InArgD, FRigVMOperand InArgE, FRigVMOperand InArgF)
+		: FRigVMBaseOp(InOpCode)
+		, ArgA(InArgA)
+		, ArgB(InArgB)
+		, ArgC(InArgC)
+		, ArgD(InArgD)
+		, ArgE(InArgE)
+		, ArgF(InArgF)
+	{
+		ensure(
+			uint8(InOpCode) == uint8(ERigVMOpCode::ArrayIterator)
+		);
+	}
+
+	FRigVMOperand ArgA;
+	FRigVMOperand ArgB;
+	FRigVMOperand ArgC;
+	FRigVMOperand ArgD;
+	FRigVMOperand ArgE;
+	FRigVMOperand ArgF;
+
+	void Serialize(FArchive& Ar);
+	FORCEINLINE friend FArchive& operator<<(FArchive& Ar, FRigVMSenaryOp& P)
+	{
+		P.Serialize(Ar);
+		return Ar;
+	}
+};
+
 // The kind of copy operation to perform
 UENUM()
 enum class ERigVMCopyType : uint8
@@ -631,6 +815,42 @@ public:
 	// adds an operator to end the last memory slice
 	uint64 AddEndBlockOp();
 
+	// adds an array reset operator
+	uint64 AddArrayResetOp(FRigVMOperand InArrayArg);
+
+	// adds an array get num operator
+	uint64 AddArrayGetNumOp(FRigVMOperand InArrayArg, FRigVMOperand InNumArg);
+
+	// adds an array set num operator
+	uint64 AddArraySetNumOp(FRigVMOperand InArrayArg, FRigVMOperand InNumArg);
+
+	// adds an array get at index operator
+	uint64 AddArrayGetAtIndexOp(FRigVMOperand InArrayArg, FRigVMOperand InIndexArg, FRigVMOperand InElementArg);
+
+	// adds an array set at index operator
+	uint64 AddArraySetAtIndexOp(FRigVMOperand InArrayArg, FRigVMOperand InIndexArg, FRigVMOperand InElementArg);
+
+	// adds an array add operator
+	uint64 AddArrayAddOp(FRigVMOperand InArrayArg, FRigVMOperand InElementArg, FRigVMOperand InIndexArg);
+
+	// adds an array insert operator
+	uint64 AddArrayInsertOp(FRigVMOperand InArrayArg, FRigVMOperand InIndexArg, FRigVMOperand InElementArg);
+
+	// adds an array remove operator
+	uint64 AddArrayRemoveOp(FRigVMOperand InArrayArg, FRigVMOperand InIndexArg);
+
+	// adds an array find operator
+	uint64 AddArrayFindOp(FRigVMOperand InArrayArg, FRigVMOperand InElementArg, FRigVMOperand InIndexArg, FRigVMOperand InSuccessArg);
+
+	// adds an array append operator
+	uint64 AddArrayAppendOp(FRigVMOperand InArrayArg, FRigVMOperand InOtherArrayArg);
+
+	// adds an array clone operator
+	uint64 AddArrayCloneOp(FRigVMOperand InArrayArg, FRigVMOperand InClonedArrayArg);
+
+	// adds an array iterator operator
+	uint64 AddArrayIteratorOp(FRigVMOperand InArrayArg, FRigVMOperand InElementArg, FRigVMOperand InIndexArg, FRigVMOperand InCountArg, FRigVMOperand InRatioArg, FRigVMOperand InContinueArg);
+	
 	// returns an instruction array for iterating over all operators
 	FORCEINLINE FRigVMInstructionArray GetInstructions() const
 	{
diff --git a/Engine/Source/Runtime/RigVM/Public/RigVMCore/RigVMExecuteContext.h b/Engine/Source/Runtime/RigVM/Public/RigVMCore/RigVMExecuteContext.h
--- a/Engine/Source/Runtime/RigVM/Public/RigVMCore/RigVMExecuteContext.h
+++ b/Engine/Source/Runtime/RigVM/Public/RigVMCore/RigVMExecuteContext.h
@@ -6,6 +6,7 @@
 #include "RigVMDefines.h"
 #include "RigVMArray.h"
 #include "RigVMExternalVariable.h"
+#include "RigVMModule.h"
 #include "RigVMExecuteContext.generated.h"
 
 USTRUCT()
@@ -147,6 +148,24 @@ private:
 	int32 Index;
 };
 
+USTRUCT()
+struct RIGVM_API FRigVMRuntimeSettings
+{
+	GENERATED_BODY()
+
+	/**
+	* The largest allowed size for arrays within the Control Rig VM.
+	* Accessing or creating larger arrays will cause runtime errors in the rig.
+	*/
+	UPROPERTY(EditAnywhere, Category = "VM")
+	int32 MaximumArraySize = 2048;
+
+	/*
+	 * The function to use for logging anything from the VM to the host
+	 */
+	TFunction<void(EMessageSeverity::Type,int32,const FString&)> LogFunction = nullptr;
+};
+
 /**
  * The execute context is used for mutable nodes to
  * indicate execution order.
@@ -160,6 +179,7 @@ struct RIGVM_API FRigVMExecuteContext
 		: EventName(NAME_None)
 		, FunctionName(NAME_None)
 		, InstructionIndex(0)
+		, RuntimeSettings() 
 	{
 		Reset();
 	}
@@ -178,6 +198,7 @@ struct RIGVM_API FRigVMExecuteContext
 		EventName = Other.EventName;
 		FunctionName = Other.FunctionName;
 		InstructionIndex = Other.InstructionIndex;
+		RuntimeSettings = Other.RuntimeSettings;
 		ExternalVariables = Other.ExternalVariables;
 		OpaqueArguments = Other.OpaqueArguments;
 		Slices = Other.Slices;
@@ -243,9 +264,67 @@ struct RIGVM_API FRigVMExecuteContext
 		return nullptr;
 	}
 
+	FORCEINLINE void Log(EMessageSeverity::Type InSeverity, const FString& InMessage) const
+	{
+		if(RuntimeSettings.LogFunction)
+		{
+			RuntimeSettings.LogFunction(InSeverity, InstructionIndex, InMessage);
+		}
+		else
+		{
+			if(InSeverity == EMessageSeverity::Error)
+			{
+				UE_LOG(LogRigVM, Error, TEXT("Instruction %d: %s"), InstructionIndex, *InMessage);
+			}
+			else if(InSeverity == EMessageSeverity::Warning)
+			{
+				UE_LOG(LogRigVM, Warning, TEXT("Instruction %d: %s"), InstructionIndex, *InMessage);
+			}
+			else
+			{
+				UE_LOG(LogRigVM, Display, TEXT("Instruction %d: %s"), InstructionIndex, *InMessage);
+			}
+		}
+	}
+
+	template <typename FmtType, typename... Types>
+	FORCEINLINE void Logf(EMessageSeverity::Type InSeverity, const FmtType& Fmt, Types... Args) const
+	{
+		Log(InSeverity, FString::Printf(Fmt, Args...));
+	}
+
+	FORCEINLINE bool IsValidArrayIndex(int32 InIndex, int32 InArraySize)
+	{
+		if(InIndex < 0 || InIndex >= InArraySize)
+		{
+			static const TCHAR OutOfBoundsFormat[] = TEXT("Array Index (%d) out of bounds (count %d).");
+			Logf(EMessageSeverity::Error, OutOfBoundsFormat, InIndex, InArraySize);
+			return false;
+		}
+		return true;
+	}
+
+	FORCEINLINE bool IsValidArraySize(int32 InSize) const
+	{
+		if(InSize < 0 || InSize > RuntimeSettings.MaximumArraySize)
+		{
+			static const TCHAR OutOfBoundsFormat[] = TEXT("Array Size (%d) larger than allowed maximum (%d).\nCheck VMRuntimeSettings in class settings.");
+			Logf(EMessageSeverity::Error, OutOfBoundsFormat, InSize, RuntimeSettings.MaximumArraySize);
+			return false;
+		}
+		return true;
+	}
+
+	FORCEINLINE void SetRuntimeSettings(FRigVMRuntimeSettings InRuntimeSettings)
+	{
+		RuntimeSettings = InRuntimeSettings;
+		check(RuntimeSettings.MaximumArraySize > 0);
+	}
+
 	FName EventName;
 	FName FunctionName;
 	uint16 InstructionIndex;
+	FRigVMRuntimeSettings RuntimeSettings;
 #if UE_RIGVM_UCLASS_BASED_STORAGE_DISABLED
 	FRigVMFixedArray<void*> OpaqueArguments;
 #else
