Control Rig: Look for memory classes in all potential class container packages

diff --git a/Engine/Source/Runtime/RigVM/Private/RigVMCore/RigVMMemoryStorage.cpp b/Engine/Source/Runtime/RigVM/Private/RigVMCore/RigVMMemoryStorage.cpp
--- a/Engine/Source/Runtime/RigVM/Private/RigVMCore/RigVMMemoryStorage.cpp
+++ b/Engine/Source/Runtime/RigVM/Private/RigVMCore/RigVMMemoryStorage.cpp
@@ -208,10 +208,37 @@ FString URigVMMemoryStorageGeneratorClass::GetClassName(ERigVMMemoryType InMemor
 URigVMMemoryStorageGeneratorClass* URigVMMemoryStorageGeneratorClass::GetStorageClass(UObject* InOuter, ERigVMMemoryType InMemoryType)
 {
 	check(InOuter);
-	UPackage* Package = InOuter->GetOutermost();
+
+	TArray<UObject*> PotentialClassContainers;
+	PotentialClassContainers.Add(InOuter);
+	
+	UObject* Outer = InOuter->GetOuter();
+	do
+	{
+		if(Outer->IsA<UPackage>())
+		{
+			break;
+		}
+
+		PotentialClassContainers.Add(Outer);
+		PotentialClassContainers.Add(Outer->GetClass());
+
+		Outer = Outer->GetOuter();
+	}
+	while (Outer);
 
 	const FString ClassName = GetClassName(InMemoryType);
-	return FindObject<URigVMMemoryStorageGeneratorClass>(Package, *ClassName);
+
+	for(UObject* PotentialClassContainer : PotentialClassContainers)
+	{
+		UPackage* Package = PotentialClassContainer->GetOutermost();
+		if(URigVMMemoryStorageGeneratorClass* Class = FindObject<URigVMMemoryStorageGeneratorClass>(Package, *ClassName))
+		{
+			return Class;
+		}
+	}
+
+	return nullptr;
 }
 
 URigVMMemoryStorageGeneratorClass* URigVMMemoryStorageGeneratorClass::CreateStorageClass(
@@ -523,6 +550,21 @@ FString URigVMMemoryStorage::GetDataAsString(const FRigVMOperand& InOperand)
 	return GetDataAsString(PropertyIndex);
 }
 
+FString URigVMMemoryStorage::GetDataAsStringSafe(int32 InPropertyIndex)
+{
+	if(!IsValidIndex(InPropertyIndex))
+	{
+		return FString();
+	}
+	return GetDataAsString(InPropertyIndex);
+}
+
+FString URigVMMemoryStorage::GetDataAsStringSafe(const FRigVMOperand& InOperand)
+{
+	const int32 PropertyIndex = InOperand.GetRegisterIndex();
+	return GetDataAsStringSafe(PropertyIndex);
+}
+
 class FRigVMMemoryStorageImportErrorContext : public FOutputDevice
 {
 public:
diff --git a/Engine/Source/Runtime/RigVM/Public/RigVMCore/RigVMMemoryStorage.h b/Engine/Source/Runtime/RigVM/Public/RigVMCore/RigVMMemoryStorage.h
--- a/Engine/Source/Runtime/RigVM/Public/RigVMCore/RigVMMemoryStorage.h
+++ b/Engine/Source/Runtime/RigVM/Public/RigVMCore/RigVMMemoryStorage.h
@@ -484,6 +484,19 @@ public:
 	// Returns the exported text for a given operand
 	FString GetDataAsString(const FRigVMOperand& InOperand);
 
+	// Returns the exported text for a given property index
+	FString GetDataAsStringSafe(int32 InPropertyIndex);
+
+	// Returns the exported text for given property name 
+	FORCEINLINE FString GetDataAsStringByNameSafe(const FName& InName)
+	{
+		const int32 PropertyIndex = GetPropertyIndexByName(InName);
+		return GetDataAsStringSafe(PropertyIndex);
+	}
+
+	// Returns the exported text for a given operand
+	FString GetDataAsStringSafe(const FRigVMOperand& InOperand);
+
 	// Sets the content of a property by index given an exported string. Returns true if succeeded
 	bool SetDataFromString(int32 InPropertyIndex, const FString& InValue);
 
