RigVM: Fix non-unity build

diff --git a/Engine/Source/Runtime/RigVM/Public/RigVMCore/RigVMExecuteContext.h b/Engine/Source/Runtime/RigVM/Public/RigVMCore/RigVMExecuteContext.h
--- a/Engine/Source/Runtime/RigVM/Public/RigVMCore/RigVMExecuteContext.h
+++ b/Engine/Source/Runtime/RigVM/Public/RigVMCore/RigVMExecuteContext.h
@@ -7,6 +7,7 @@
 #include "RigVMArray.h"
 #include "RigVMExternalVariable.h"
 #include "RigVMModule.h"
+#include "Logging/TokenizedMessage.h"
 #include "RigVMExecuteContext.generated.h"
 
 USTRUCT()
