Control Rig: RigVM: Avoid crash for empty arrays

diff --git a/Engine/Source/Runtime/RigVM/Private/RigVMCore/RigVMMemoryStorage.cpp b/Engine/Source/Runtime/RigVM/Private/RigVMCore/RigVMMemoryStorage.cpp
--- a/Engine/Source/Runtime/RigVM/Private/RigVMCore/RigVMMemoryStorage.cpp
+++ b/Engine/Source/Runtime/RigVM/Private/RigVMCore/RigVMMemoryStorage.cpp
@@ -311,7 +311,9 @@ URigVMMemoryStorageGeneratorClass* URigVMMemoryStorageGeneratorClass::CreateStor
 	for(int32 PropertyIndex = 0; PropertyIndex < LinkedProperties.Num(); PropertyIndex++)
 	{
 		const FString& DefaultValue = InProperties[PropertyIndex].DefaultValue;
-		if(DefaultValue.IsEmpty())
+
+		static const FString EmptyBraces = TEXT("()");
+		if(DefaultValue.IsEmpty() || DefaultValue == EmptyBraces)
 		{
 			continue;
 		}
