RigVM: Fix crash for debugging base structure pins

diff --git a/Engine/Source/Runtime/RigVM/Private/RigVMCore/RigVMMemoryStorage.cpp b/Engine/Source/Runtime/RigVM/Private/RigVMCore/RigVMMemoryStorage.cpp
--- a/Engine/Source/Runtime/RigVM/Private/RigVMCore/RigVMMemoryStorage.cpp
+++ b/Engine/Source/Runtime/RigVM/Private/RigVMCore/RigVMMemoryStorage.cpp
@@ -487,7 +487,33 @@ FProperty* URigVMMemoryStorageGeneratorClass::AddProperty(URigVMMemoryStorageGen
 			}
 			else
 			{
-				checkNoEntry();
+				// we might have hit a base structure
+				static TArray<UScriptStruct*> BaseStructures;
+				if(BaseStructures.IsEmpty())
+				{
+					BaseStructures.Add(TBaseStructure<FVector>::Get());
+					BaseStructures.Add(TBaseStructure<FVector2D>::Get());
+					BaseStructures.Add(TBaseStructure<FRotator>::Get());
+					BaseStructures.Add(TBaseStructure<FQuat>::Get());
+					BaseStructures.Add(TBaseStructure<FTransform>::Get());
+					BaseStructures.Add(TBaseStructure<FLinearColor>::Get());
+				}
+
+				for(UScriptStruct* BaseStructure : BaseStructures)
+				{
+					if(BaseStructure->GetStructCPPName() == BaseCPPType)
+					{
+						FStructProperty* StructProperty = new FStructProperty(PropertyOwner, InProperty.Name, RF_Public);
+						StructProperty->Struct = BaseStructure;
+						(*ValuePropertyPtr) = StructProperty;
+						break;
+					}
+				}
+
+				if(*ValuePropertyPtr == nullptr)
+				{
+					return nullptr;
+				}
 			}
 		}
 
