Follow up build break fix for 17249273

diff --git a/Engine/Source/Runtime/RigVM/Private/RigVMCore/RigVM.cpp b/Engine/Source/Runtime/RigVM/Private/RigVMCore/RigVM.cpp
--- a/Engine/Source/Runtime/RigVM/Private/RigVMCore/RigVM.cpp
+++ b/Engine/Source/Runtime/RigVM/Private/RigVMCore/RigVM.cpp
@@ -1900,11 +1900,13 @@ bool URigVM::Execute(TArrayView<URigVMMemoryStorage*> Memory, TArrayView<void*>
 
 	NumExecutions++;
 
+#if WITH_EDITOR
 	uint64 StartCycles = 0;
 	if(Context.RuntimeSettings.bEnableProfiling)
 	{
 		StartCycles = FPlatformTime::Cycles64();
 	}
+#endif
 
 	while (Instructions.IsValidIndex(Context.InstructionIndex))
 	{
