Control Rig: Fix debug not working in backwards and forward solve

diff --git a/Engine/Plugins/Animation/ControlRig/Source/ControlRig/Private/ControlRig.cpp b/Engine/Plugins/Animation/ControlRig/Source/ControlRig/Private/ControlRig.cpp
--- a/Engine/Plugins/Animation/ControlRig/Source/ControlRig/Private/ControlRig.cpp
+++ b/Engine/Plugins/Animation/ControlRig/Source/ControlRig/Private/ControlRig.cpp
@@ -230,6 +230,17 @@ void UControlRig::Evaluate_AnyThread()
 	for (const FName& EventName : EventQueue)
 	{
 		Execute(EControlRigState::Update, EventName);
+
+#if WITH_EDITOR
+		if (VM)
+		{
+			if (VM->GetHaltedAtBreakpoint().IsValid())
+			{
+				break;
+			}
+		}
+#endif
+			
 	}
 }
 
@@ -1401,15 +1412,18 @@ void UControlRig::HandleOnControlModified(UControlRig* Subject, FRigControlEleme
 	}	
 }
 
-void UControlRig::HandleExecutionReachedExit()
+void UControlRig::HandleExecutionReachedExit(const FName& InEventName)
 {
 #if WITH_EDITOR
-	if(URigVM* SnapShotVM = GetSnapshotVM(false))
+	if (EventQueue.Last() == InEventName)
 	{
-		SnapShotVM->CopyFrom(VM, false, false, false, true, true);
+		if(URigVM* SnapShotVM = GetSnapshotVM(false))
+		{
+			SnapShotVM->CopyFrom(VM, false, false, false, true, true);
+		}
+		DebugInfo.ResetState();
+		VM->SetBreakpointAction(ERigVMBreakpointAction::None);
 	}
-	DebugInfo.ResetState();
-	VM->SetBreakpointAction(ERigVMBreakpointAction::None);
 #endif
 	
 	if (LatestExecutedState != EControlRigState::Init && bAccumulateTime)
diff --git a/Engine/Plugins/Animation/ControlRig/Source/ControlRig/Public/ControlRig.h b/Engine/Plugins/Animation/ControlRig/Source/ControlRig/Public/ControlRig.h
--- a/Engine/Plugins/Animation/ControlRig/Source/ControlRig/Public/ControlRig.h
+++ b/Engine/Plugins/Animation/ControlRig/Source/ControlRig/Public/ControlRig.h
@@ -446,7 +446,7 @@ private:
 	
 	void HandleOnControlModified(UControlRig* Subject, FRigControlElement* Control, const FRigControlModifiedContext& Context);
 
-	void HandleExecutionReachedExit();
+	void HandleExecutionReachedExit(const FName& InEventName);
 	
 	TArray<FRigVMExternalVariable> GetExternalVariablesImpl(bool bFallbackToBlueprint) const;
 
diff --git a/Engine/Plugins/Animation/ControlRig/Source/ControlRigEditor/Private/Editor/ControlRigEditor.cpp b/Engine/Plugins/Animation/ControlRig/Source/ControlRigEditor/Private/Editor/ControlRigEditor.cpp
--- a/Engine/Plugins/Animation/ControlRig/Source/ControlRigEditor/Private/Editor/ControlRigEditor.cpp
+++ b/Engine/Plugins/Animation/ControlRig/Source/ControlRigEditor/Private/Editor/ControlRigEditor.cpp
@@ -2858,8 +2858,13 @@ void FControlRigEditor::HandleControlRigExecutedEvent(UControlRig* InControlRig,
 	UpdateGraphCompilerErrors();
 }
 
-void FControlRigEditor::HandleControlRigExecutionHalted(const int32 InstructionIndex, UObject* InNodeObject)
+void FControlRigEditor::HandleControlRigExecutionHalted(const int32 InstructionIndex, UObject* InNodeObject, const FName& InEntryName)
 {
+	if (HaltedAtNode == InNodeObject)
+	{
+		return;
+	}
+		
 	if (URigVMNode* InNode = Cast<URigVMNode>(InNodeObject))
 	{
 		if (HaltedAtNode != nullptr)
@@ -2888,6 +2893,13 @@ void FControlRigEditor::HandleControlRigExecutionHalted(const int32 InstructionI
 			}
 		}
 	}
+	else 
+	{
+		if (InEntryName == ControlRig->GetEventQueue().Last())
+		{
+			HaltedAtNode = nullptr;
+		}
+	}
 }
 
 void FControlRigEditor::CreateEditorModeManager()
diff --git a/Engine/Plugins/Animation/ControlRig/Source/ControlRigEditor/Private/Editor/ControlRigEditor.h b/Engine/Plugins/Animation/ControlRig/Source/ControlRigEditor/Private/Editor/ControlRigEditor.h
--- a/Engine/Plugins/Animation/ControlRig/Source/ControlRigEditor/Private/Editor/ControlRigEditor.h
+++ b/Engine/Plugins/Animation/ControlRig/Source/ControlRigEditor/Private/Editor/ControlRigEditor.h
@@ -240,7 +240,7 @@ protected:
 	void HandleModifiedEvent(ERigVMGraphNotifType InNotifType, URigVMGraph* InGraph, UObject* InSubject);
 	void HandleVMCompiledEvent(UBlueprint* InBlueprint, URigVM* InVM);
 	void HandleControlRigExecutedEvent(UControlRig* InControlRig, const EControlRigState InState, const FName& InEventName);
-	void HandleControlRigExecutionHalted(const int32 InstructionIndex, UObject* InNode);
+	void HandleControlRigExecutionHalted(const int32 InstructionIndex, UObject* InNode, const FName& InEntryName);
 
 	// FBaseToolKit overrides
 	void CreateEditorModeManager() override;
diff --git a/Engine/Plugins/Animation/ControlRig/Source/ControlRigEditor/Private/Editor/SControlRigStackView.cpp b/Engine/Plugins/Animation/ControlRig/Source/ControlRigEditor/Private/Editor/SControlRigStackView.cpp
--- a/Engine/Plugins/Animation/ControlRig/Source/ControlRigEditor/Private/Editor/SControlRigStackView.cpp
+++ b/Engine/Plugins/Animation/ControlRig/Source/ControlRigEditor/Private/Editor/SControlRigStackView.cpp
@@ -252,6 +252,7 @@ void SControlRigStackView::Construct( const FArguments& InArgs, TSharedRef<FCont
 	CommandList = MakeShared<FUICommandList>();
 	bSuspendModelNotifications = false;
 	bSuspendControllerSelection = false;
+	HaltedAtInstruction = INDEX_NONE;
 
 	BindCommands();
 
@@ -845,13 +846,25 @@ void SControlRigStackView::OnVMCompiled(UBlueprint* InCompiledBlueprint, URigVM*
 	}
 }
 
-void SControlRigStackView::HandleExecutionHalted(const int32 HaltedAtInstruction, UObject* InNode)
-{	
-	if (HaltedAtInstruction != INDEX_NONE && Operators.Num() > HaltedAtInstruction)
+void SControlRigStackView::HandleExecutionHalted(const int32 InHaltedAtInstruction, UObject* InNode, const FName& InEntryName)
+{
+	if (HaltedAtInstruction == InHaltedAtInstruction)
+	{
+		return;
+	}
+
+	if (InHaltedAtInstruction == INDEX_NONE && InEntryName == ControlRigEditor.Pin()->ControlRig->GetEventQueue().Last())
 	{
-		TreeView->SetSelection(Operators[HaltedAtInstruction]);
-		TreeView->SetScrollOffset(FMath::Max(HaltedAtInstruction-5, 0));
+		HaltedAtInstruction = InHaltedAtInstruction;
+		return;
 	}
+	
+	if (InHaltedAtInstruction != INDEX_NONE && Operators.Num() > InHaltedAtInstruction)
+	{
+		HaltedAtInstruction = InHaltedAtInstruction;
+		TreeView->SetSelection(Operators[InHaltedAtInstruction]);
+		TreeView->SetScrollOffset(FMath::Max(InHaltedAtInstruction-5, 0));
+	}	
 }
 
 void SControlRigStackView::OnFilterTextChanged(const FText& SearchText)
diff --git a/Engine/Plugins/Animation/ControlRig/Source/ControlRigEditor/Private/Editor/SControlRigStackView.h b/Engine/Plugins/Animation/ControlRig/Source/ControlRigEditor/Private/Editor/SControlRigStackView.h
--- a/Engine/Plugins/Animation/ControlRig/Source/ControlRigEditor/Private/Editor/SControlRigStackView.h
+++ b/Engine/Plugins/Animation/ControlRig/Source/ControlRigEditor/Private/Editor/SControlRigStackView.h
@@ -99,7 +99,7 @@ private:
 	void OnVMCompiled(UBlueprint* InCompiledBlueprint, URigVM* InCompiledVM);
 
 	//* Focus on the instruction when the execution is halted */
-	void HandleExecutionHalted(const int32 HaltedAtInstruction, UObject* InNode);
+	void HandleExecutionHalted(const int32 HaltedAtInstruction, UObject* InNode, const FName& InEntryName);
 
 	/** Search box widget */
 	TSharedPtr<SSearchBox> FilterBox;
@@ -128,6 +128,8 @@ private:
 
 	TArray<TSharedPtr<FRigStackEntry>> Operators;
 
+	int32 HaltedAtInstruction;
+
 	FDelegateHandle OnModelModified;
 	FDelegateHandle OnControlRigInitializedHandle;
 	FDelegateHandle OnVMCompiledHandle;
diff --git a/Engine/Source/Runtime/RigVM/Private/RigVMCore/RigVM.cpp b/Engine/Source/Runtime/RigVM/Private/RigVMCore/RigVM.cpp
--- a/Engine/Source/Runtime/RigVM/Private/RigVMCore/RigVM.cpp
+++ b/Engine/Source/Runtime/RigVM/Private/RigVMCore/RigVM.cpp
@@ -1312,7 +1312,7 @@ void URigVM::RebuildByteCodeOnLoad()
 }
 
 #if WITH_EDITOR
-bool URigVM::ShouldHaltAtInstruction(const uint16 InstructionIndex)
+bool URigVM::ShouldHaltAtInstruction(const FName& InEventName, const uint16 InstructionIndex)
 {
 	FRigVMByteCode& ByteCode = GetByteCode();
 
@@ -1344,7 +1344,7 @@ bool URigVM::ShouldHaltAtInstruction(const uint16 InstructionIndex)
 						{
 							DebugInfo->SetCurrentActiveBreakpointCallstack(TArray<UObject*>(FullCallstack->GetData(), FullCallstack->Find((UObject*)Breakpoint->Subject)+1));
 						}
-						ExecutionHalted().Broadcast(Context.InstructionIndex, Breakpoint->Subject);
+						ExecutionHalted().Broadcast(Context.InstructionIndex, Breakpoint->Subject, InEventName);
 					}
 					return true;
 				}
@@ -1468,7 +1468,7 @@ bool URigVM::ShouldHaltAtInstruction(const uint16 InstructionIndex)
 
 				HaltedAtBreakpoint = NewBreakpoint;
 				HaltedAtBreakpointHit = DebugInfo->GetBreakpointHits(HaltedAtBreakpoint);
-				ExecutionHalted().Broadcast(Context.InstructionIndex, NewBreakpointNode);
+				ExecutionHalted().Broadcast(Context.InstructionIndex, NewBreakpointNode, InEventName);
 		
 				return true;
 			}
@@ -1911,7 +1911,7 @@ bool URigVM::Execute(TArrayView<URigVMMemoryStorage*> Memory, TArrayView<void*>
 	while (Instructions.IsValidIndex(Context.InstructionIndex))
 	{
 #if WITH_EDITOR
-		if (DebugInfo && ShouldHaltAtInstruction(Context.InstructionIndex))
+		if (DebugInfo && ShouldHaltAtInstruction(InEntryName, Context.InstructionIndex))
 		{
 			return true;
 		}
@@ -2382,13 +2382,13 @@ bool URigVM::Execute(TArrayView<URigVMMemoryStorage*> Memory, TArrayView<void*>
 			}
 			case ERigVMOpCode::Exit:
 			{
-				ExecutionReachedExit().Broadcast();
+				ExecutionReachedExit().Broadcast(InEntryName);
 #if WITH_EDITOR					
 				if (HaltedAtBreakpoint != nullptr)
 				{
 					HaltedAtBreakpoint = nullptr;
 					DebugInfo->SetCurrentActiveBreakpoint(nullptr);
-					ExecutionHalted().Broadcast(INDEX_NONE, nullptr);
+					ExecutionHalted().Broadcast(INDEX_NONE, nullptr, InEntryName);
 				}
 #endif
 				return true;
@@ -2964,7 +2964,7 @@ bool URigVM::Execute(TArrayView<URigVMMemoryStorage*> Memory, TArrayView<void*>
 	{
 		DebugInfo->SetCurrentActiveBreakpoint(nullptr);
 		HaltedAtBreakpoint = nullptr;
-		ExecutionHalted().Broadcast(INDEX_NONE, nullptr);
+		ExecutionHalted().Broadcast(INDEX_NONE, nullptr, InEntryName);
 	}
 #endif
 
diff --git a/Engine/Source/Runtime/RigVM/Public/RigVMCore/RigVM.h b/Engine/Source/Runtime/RigVM/Public/RigVMCore/RigVM.h
--- a/Engine/Source/Runtime/RigVM/Public/RigVMCore/RigVM.h
+++ b/Engine/Source/Runtime/RigVM/Public/RigVMCore/RigVM.h
@@ -138,9 +138,9 @@ class RIGVM_API URigVM : public UObject
 public:
 
 	/** Bindable event for external objects to be notified when the VM reaches an Exit Operation */
-	DECLARE_EVENT(URigVM, FExecutionReachedExitEvent);
+	DECLARE_EVENT_OneParam(URigVM, FExecutionReachedExitEvent, const FName&);
 #if WITH_EDITOR
-	DECLARE_EVENT_TwoParams(URigVM, FExecutionHaltedEvent, int32, UObject*);
+	DECLARE_EVENT_ThreeParams(URigVM, FExecutionHaltedEvent, int32, UObject*, const FName&);
 #endif
 
 	URigVM();
@@ -773,7 +773,7 @@ private:
 	int32 HaltedAtBreakpointHit;
 	ERigVMBreakpointAction CurrentBreakpointAction;
 
-	bool ShouldHaltAtInstruction(const uint16 InstructionIndex);
+	bool ShouldHaltAtInstruction(const FName& InEventName, const uint16 InstructionIndex);
 #endif
 
 	UPROPERTY()
