Control Rig: Fix crash with array intersections

diff --git a/Engine/Source/Runtime/RigVM/Private/RigVMCore/RigVM.cpp b/Engine/Source/Runtime/RigVM/Private/RigVMCore/RigVM.cpp
--- a/Engine/Source/Runtime/RigVM/Private/RigVMCore/RigVM.cpp
+++ b/Engine/Source/Runtime/RigVM/Private/RigVMCore/RigVM.cpp
@@ -2774,7 +2774,18 @@ bool URigVM::Execute(TArrayView<URigVMMemoryStorage*> Memory, TArrayView<void*>
 
 				for(int32 Index = 0; Index < ArrayHelperA.Num(); Index++)
 				{
-					const uint32 HashValue = ElementPropertyA->GetValueTypeHash(ArrayHelperA.GetRawPtr(Index));
+					uint32 HashValue;
+					if (ElementPropertyA->PropertyFlags & CPF_HasGetValueTypeHash)
+					{
+						HashValue = ElementPropertyA->GetValueTypeHash(ArrayHelperA.GetRawPtr(Index));
+					}
+					else
+					{
+						FString Value;
+						ElementPropertyA->ExportTextItem(Value, ArrayHelperA.GetRawPtr(Index), nullptr, nullptr, PPF_None);
+						HashValue = TextKeyUtil::HashString(Value);
+					}
+					
 					if(!HashA.Contains(HashValue))
 					{
 						HashA.Add(HashValue, Index);
@@ -2782,7 +2793,17 @@ bool URigVM::Execute(TArrayView<URigVMMemoryStorage*> Memory, TArrayView<void*>
 				}
 				for(int32 Index = 0; Index < ArrayHelperB.Num(); Index++)
 				{
-					const uint32 HashValue = ElementPropertyB->GetValueTypeHash(ArrayHelperB.GetRawPtr(Index));
+					uint32 HashValue;
+					if (ElementPropertyB->PropertyFlags & CPF_HasGetValueTypeHash)
+					{
+						HashValue = ElementPropertyB->GetValueTypeHash(ArrayHelperB.GetRawPtr(Index));
+					}
+					else
+					{
+						FString Value;
+						ElementPropertyB->ExportTextItem(Value, ArrayHelperB.GetRawPtr(Index), nullptr, nullptr, PPF_None);
+						HashValue = TextKeyUtil::HashString(Value);
+					}
 					if(!HashB.Contains(HashValue))
 					{
 						HashB.Add(HashValue, Index);
