Control Rig: Improve profiler lookup to be based on hash of callstack rather than string based

diff --git a/Engine/Plugins/Animation/ControlRig/Source/ControlRigDeveloper/Private/Graph/ControlRigGraphNode.cpp b/Engine/Plugins/Animation/ControlRig/Source/ControlRigDeveloper/Private/Graph/ControlRigGraphNode.cpp
--- a/Engine/Plugins/Animation/ControlRig/Source/ControlRigDeveloper/Private/Graph/ControlRigGraphNode.cpp
+++ b/Engine/Plugins/Animation/ControlRig/Source/ControlRigDeveloper/Private/Graph/ControlRigGraphNode.cpp
@@ -389,7 +389,7 @@ FLinearColor UControlRigGraphNode::GetNodeProfilingColor() const
 #if WITH_EDITOR
 	if(bEnableProfiling)
 	{
-		if(UControlRigBlueprint* Blueprint = Cast<UControlRigBlueprint>(GetOuter()->GetOuter()))
+		if(UControlRigBlueprint* Blueprint = GetTypedOuter<UControlRigBlueprint>())
 		{
 			if(UControlRig* DebuggedControlRig = Cast<UControlRig>(Blueprint->GetObjectBeingDebugged()))
 			{
@@ -644,7 +644,7 @@ void UControlRigGraphNode::CreateOutputPins(URigVMPin* InParentPin)
 
 UClass* UControlRigGraphNode::GetControlRigGeneratedClass() const
 {
-	UControlRigBlueprint* Blueprint = Cast<UControlRigBlueprint>(GetOuter()->GetOuter());
+	UControlRigBlueprint* Blueprint = GetTypedOuter<UControlRigBlueprint>();
 	if(Blueprint)
 	{
 		if (Blueprint->GeneratedClass)
@@ -659,7 +659,7 @@ UClass* UControlRigGraphNode::GetControlRigGeneratedClass() const
 
 UClass* UControlRigGraphNode::GetControlRigSkeletonGeneratedClass() const
 {
-	UControlRigBlueprint* Blueprint = Cast<UControlRigBlueprint>(GetOuter()->GetOuter());
+	UControlRigBlueprint* Blueprint = GetTypedOuter<UControlRigBlueprint>();
 	if(Blueprint)
 	{
 		if (Blueprint->SkeletonGeneratedClass)
diff --git a/Engine/Source/Developer/RigVMDeveloper/Private/RigVMModel/RigVMNode.cpp b/Engine/Source/Developer/RigVMDeveloper/Private/RigVMModel/RigVMNode.cpp
--- a/Engine/Source/Developer/RigVMDeveloper/Private/RigVMModel/RigVMNode.cpp
+++ b/Engine/Source/Developer/RigVMDeveloper/Private/RigVMModel/RigVMNode.cpp
@@ -476,8 +476,7 @@ TArray<int32> URigVMNode::GetInstructionsForVMImpl(URigVM* InVM, const FRigVMAST
 	if(InProxy.IsValid())
 	{
 		const FRigVMASTProxy Proxy = InProxy.GetChild((UObject*)this);
-		const FString CallPath = Proxy.GetCallstack().GetCallPath();
-		return InVM->GetByteCode().GetAllInstructionIndicesForCallPath(CallPath, false, true);
+		return InVM->GetByteCode().GetAllInstructionIndicesForCallstack(Proxy.GetCallstack().GetStack());
 	}
 	else
 	{
diff --git a/Engine/Source/Runtime/RigVM/Private/RigVMCore/RigVMByteCode.cpp b/Engine/Source/Runtime/RigVM/Private/RigVMCore/RigVMByteCode.cpp
--- a/Engine/Source/Runtime/RigVM/Private/RigVMCore/RigVMByteCode.cpp
+++ b/Engine/Source/Runtime/RigVM/Private/RigVMCore/RigVMByteCode.cpp
@@ -663,6 +663,8 @@ void FRigVMByteCode::Reset()
 	CallPathPerInstruction.Reset();
 	CallPathToInstructions.Reset();
 	CallstackPerInstruction.Reset();
+	CallstackHashToInstructions.Reset();
+	CallstackHashPerInstruction.Reset();
 #endif
 }
 
@@ -679,6 +681,8 @@ void FRigVMByteCode::Empty()
 	CallPathPerInstruction.Empty();
 	CallPathToInstructions.Empty();
 	CallstackPerInstruction.Empty();
+	CallstackHashToInstructions.Empty();
+	CallstackHashPerInstruction.Empty();
 #endif
 }
 
@@ -1702,6 +1706,32 @@ TArray<int32> FRigVMByteCode::GetAllInstructionIndicesForCallPath(const FString&
 	return MatchedInstructions;
 }
 
+int32 FRigVMByteCode::GetFirstInstructionIndexForCallstack(const TArray<UObject*>& InCallstack) const
+{
+	const TArray<int32>& InstructionIndices = GetAllInstructionIndicesForCallstack(InCallstack);
+	if (InstructionIndices.Num() > 0)
+	{
+		return InstructionIndices[0];
+	}
+	return INDEX_NONE;
+}
+
+const TArray<int32>& FRigVMByteCode::GetAllInstructionIndicesForCallstack(const TArray<UObject*>& InCallstack) const
+{
+	if(InCallstack.IsEmpty())
+	{
+		return EmptyInstructionIndices;
+	}
+
+	const uint32 Hash = GetCallstackHash(InCallstack);
+	if(const TArray<int32>* Instructions = CallstackHashToInstructions.Find(Hash))
+	{
+		return *Instructions;
+	}
+	
+	return EmptyInstructionIndices;
+}
+
 void FRigVMByteCode::SetSubject(int32 InInstructionIndex, const FString& InCallPath, const TArray<UObject*>& InCallstack)
 {
 	UObject* Subject = InCallstack.Last();
@@ -1724,6 +1754,20 @@ void FRigVMByteCode::SetSubject(int32 InInstructionIndex, const FString& InCallP
 		CallstackPerInstruction.AddZeroed(1 + InInstructionIndex - CallstackPerInstruction.Num());
 	}
 	CallstackPerInstruction[InInstructionIndex] = InCallstack;
+
+	if (CallstackHashPerInstruction.Num() <= InInstructionIndex)
+	{
+		CallstackHashPerInstruction.AddZeroed(1 + InInstructionIndex - CallstackHashPerInstruction.Num());
+	}
+	CallstackHashPerInstruction[InInstructionIndex] = GetCallstackHash(InCallstack);
+
+	for(int32 CallstackLength = InCallstack.Num(); CallstackLength > 0; CallstackLength--)
+	{
+		UObject* const* DataPtr = &InCallstack[InCallstack.Num() - CallstackLength];
+		TArrayView<UObject* const> View(DataPtr, CallstackLength);
+		uint32 Hash = GetCallstackHash(View);
+		CallstackHashToInstructions.FindOrAdd(Hash).Add(InInstructionIndex);
+	}
 }
 
 const TArray<UObject*>* FRigVMByteCode::GetCallstackForInstruction(int32 InInstructionIndex) const
@@ -1735,4 +1779,34 @@ const TArray<UObject*>* FRigVMByteCode::GetCallstackForInstruction(int32 InInstr
 	return nullptr;
 }
 
+uint32 FRigVMByteCode::GetCallstackHashForInstruction(int32 InInstructionIndex) const
+{
+	if (CallstackHashPerInstruction.IsValidIndex(InInstructionIndex))
+	{
+		return CallstackHashPerInstruction[InInstructionIndex];
+	}
+	return 0;
+}
+
+uint32 FRigVMByteCode::GetCallstackHash(const TArray<UObject*>& InCallstack)
+{
+	UObject*const * DataPtr = nullptr;
+	if(InCallstack.Num() > 0)
+	{
+		DataPtr = &InCallstack[0];
+	}
+	TArrayView<UObject* const> View(DataPtr, InCallstack.Num());
+	return GetCallstackHash(View);
+}
+
+uint32 FRigVMByteCode::GetCallstackHash(const TArrayView<UObject* const>& InCallstack)
+{
+	uint32 Hash = GetTypeHash(InCallstack.Num());
+	for(const UObject* Object : InCallstack)
+	{
+		Hash = HashCombine(Hash, GetTypeHash(Object));
+	}
+	return Hash;
+}
+
 #endif
diff --git a/Engine/Source/Runtime/RigVM/Public/RigVMCore/RigVMByteCode.h b/Engine/Source/Runtime/RigVM/Public/RigVMCore/RigVMByteCode.h
--- a/Engine/Source/Runtime/RigVM/Public/RigVMCore/RigVMByteCode.h
+++ b/Engine/Source/Runtime/RigVM/Public/RigVMCore/RigVMByteCode.h
@@ -1006,9 +1006,22 @@ public:
 	// returns all found instruction indices for a given callpath
 	TArray<int32> GetAllInstructionIndicesForCallPath(const FString& InCallPath, bool bStartsWith = false, bool bEndsWith = false) const;
 
+	// returns the first hit instruction index for a given callpath (or INDEX_NONE)
+	int32 GetFirstInstructionIndexForCallstack(const TArray<UObject*>& InCallstack) const;
+
+	// returns all found instruction indices for a given callpath
+	const TArray<int32>& GetAllInstructionIndicesForCallstack(const TArray<UObject*>& InCallstack) const;
+
 	// returns the callstack which was used to inject a given instruction
 	const TArray<UObject*>* GetCallstackForInstruction(int32 InInstructionIndex) const;
 
+	// returns the callstack hash which was used to inject a given instruction
+	uint32 GetCallstackHashForInstruction(int32 InInstructionIndex) const;
+
+	// computes a hash for a given callstack
+	static uint32 GetCallstackHash(const TArray<UObject*>& InCallstack);
+	static uint32 GetCallstackHash(const TArrayView<UObject* const>& InCallstack);
+
 #endif
 
 private:
@@ -1040,6 +1053,8 @@ private:
 	TArray<FString> CallPathPerInstruction;
 	TMap<FString, TArray<int32>> CallPathToInstructions;
 	TArray<TArray<UObject*>> CallstackPerInstruction;
+	TMap<uint32, TArray<int32>> CallstackHashToInstructions;
+	TArray<uint32> CallstackHashPerInstruction;
 
 #endif
 
