RigVM: Allow sub structures as part of pins / property paths.

diff --git a/Engine/Source/Runtime/RigVM/Private/RigVMCore/RigVMPropertyPath.cpp b/Engine/Source/Runtime/RigVM/Private/RigVMCore/RigVMPropertyPath.cpp
--- a/Engine/Source/Runtime/RigVM/Private/RigVMCore/RigVMPropertyPath.cpp
+++ b/Engine/Source/Runtime/RigVM/Private/RigVMCore/RigVMPropertyPath.cpp
@@ -180,7 +180,9 @@ uint8* FRigVMPropertyPath::GetData_Internal(uint8* InPtr, const FProperty* InPro
 				// make sure to only allow jumps within the same data structure.
 				// we check this to avoid heap corruption.
 				const FStructProperty* StructProperty = CastFieldChecked<FStructProperty>(InProperty);
-				check(StructProperty->Struct == Segment.Property->GetOwnerStruct());
+
+				check(StructProperty->Struct->IsChildOf(Segment.Property->GetOwnerStruct());
+					
 				InProperty = Segment.Property;
 #endif
 
