Control Rig: Fix debug stepping inside loops

diff --git a/Engine/Source/Runtime/RigVM/Private/RigVMCore/RigVM.cpp b/Engine/Source/Runtime/RigVM/Private/RigVMCore/RigVM.cpp
--- a/Engine/Source/Runtime/RigVM/Private/RigVMCore/RigVM.cpp
+++ b/Engine/Source/Runtime/RigVM/Private/RigVMCore/RigVM.cpp
@@ -1330,11 +1330,6 @@ bool URigVM::ShouldHaltAtInstruction(const FName& InEventName, const uint16 Inst
 {
 	FRigVMByteCode& ByteCode = GetByteCode();
 
-	if (HaltedAtBreakpoint && InstructionIndex < HaltedAtBreakpoint->InstructionIndex)
-	{
-		return false;
-	}
-
 	TArray<TSharedPtr<FRigVMBreakpoint>> BreakpointsAtInstruction = DebugInfo->FindBreakpointsAtInstruction(InstructionIndex);
 	for (TSharedPtr<FRigVMBreakpoint> Breakpoint : BreakpointsAtInstruction)
 	{
@@ -1478,6 +1473,8 @@ bool URigVM::ShouldHaltAtInstruction(const FName& InEventName, const uint16 Inst
 
 				// Create new temporary breakpoint
 				TSharedPtr<FRigVMBreakpoint> NewBreakpoint = DebugInfo->AddBreakpoint(Context.InstructionIndex, NewBreakpointNode, 0, true);
+				DebugInfo->SetBreakpointHits(NewBreakpoint, GetInstructionVisitedCount(Context.InstructionIndex));
+				DebugInfo->SetBreakpointActivationOnHit(NewBreakpoint, GetInstructionVisitedCount(Context.InstructionIndex));
 				CurrentBreakpointAction = ERigVMBreakpointAction::None;					
 
 				HaltedAtBreakpoint = NewBreakpoint;
diff --git a/Engine/Source/Runtime/RigVM/Private/RigVMCore/RigVMDebugInfo.cpp b/Engine/Source/Runtime/RigVM/Private/RigVMCore/RigVMDebugInfo.cpp
--- a/Engine/Source/Runtime/RigVM/Private/RigVMCore/RigVMDebugInfo.cpp
+++ b/Engine/Source/Runtime/RigVM/Private/RigVMCore/RigVMDebugInfo.cpp
@@ -144,6 +144,18 @@ bool FRigVMDebugInfo::IsActive(const TSharedPtr<FRigVMBreakpoint> InBreakpoint)
 	return false;		
 }
 
+void FRigVMDebugInfo::SetBreakpointHits(const TSharedPtr<FRigVMBreakpoint> InBreakpoint, const uint16 InBreakpointHits)
+{
+	if (BreakpointHits.Contains(InBreakpoint))
+	{
+		BreakpointHits[InBreakpoint] = InBreakpointHits;
+	}
+	else
+	{
+		BreakpointHits.Add(InBreakpoint, InBreakpointHits);
+	}
+}
+
 void FRigVMDebugInfo::HitBreakpoint(const TSharedPtr<FRigVMBreakpoint> InBreakpoint)
 {
 	if (BreakpointHits.Contains(InBreakpoint))
@@ -156,6 +168,18 @@ void FRigVMDebugInfo::HitBreakpoint(const TSharedPtr<FRigVMBreakpoint> InBreakpo
 	}
 }
 
+void FRigVMDebugInfo::SetBreakpointActivationOnHit(const TSharedPtr<FRigVMBreakpoint> InBreakpoint, const uint16 InActivationOnHit)
+{
+	if (BreakpointActivationOnHit.Contains(InBreakpoint))
+	{
+		BreakpointActivationOnHit[InBreakpoint] = InActivationOnHit;
+	}
+	else
+	{
+		BreakpointActivationOnHit.Add(InBreakpoint, InActivationOnHit);
+	}
+}
+
 void FRigVMDebugInfo::IncrementBreakpointActivationOnHit(const TSharedPtr<FRigVMBreakpoint> InBreakpoint)
 {
 	if (BreakpointActivationOnHit.Contains(InBreakpoint))
diff --git a/Engine/Source/Runtime/RigVM/Public/RigVMCore/RigVMDebugInfo.h b/Engine/Source/Runtime/RigVM/Public/RigVMCore/RigVMDebugInfo.h
--- a/Engine/Source/Runtime/RigVM/Public/RigVMCore/RigVMDebugInfo.h
+++ b/Engine/Source/Runtime/RigVM/Public/RigVMCore/RigVMDebugInfo.h
@@ -38,10 +38,17 @@ struct RIGVM_API FRigVMBreakpoint
 		Depth = 0;
 	}
 
+	// Whether or not the breakpoint is active
 	bool bIsActive;
+
+	// Instruction where this breakpoint is set
 	uint16 InstructionIndex;
-	UObject* Subject;
-	uint16 Depth;
+
+	// Node related to the breakpoint
+	UObject* Subject; 
+
+	// The same instruction might be related to multiple breakpoints at different callstack depths
+	uint16 Depth; 
 
 	bool operator==(const FRigVMBreakpoint& Other) const
 	{
@@ -101,8 +108,12 @@ struct RIGVM_API FRigVMDebugInfo
 
 	bool IsActive(const TSharedPtr<FRigVMBreakpoint> InBreakpoint) const;
 
+	void SetBreakpointHits(const TSharedPtr<FRigVMBreakpoint> InBreakpoint, const uint16 InBreakpointHits);
+
 	void HitBreakpoint(const TSharedPtr<FRigVMBreakpoint> InBreakpoint);
 
+	void SetBreakpointActivationOnHit(const TSharedPtr<FRigVMBreakpoint> InBreakpoint, const uint16 InActivationOnHit);
+
 	void IncrementBreakpointActivationOnHit(const TSharedPtr<FRigVMBreakpoint> InBreakpoint);
 
 	uint16 GetBreakpointHits(const TSharedPtr<FRigVMBreakpoint> InBreakpoint) const;
