Control Rig: Fix SteppingInLoop unit test not passing

diff --git a/Engine/Source/Runtime/RigVM/Private/RigVMCore/RigVM.cpp b/Engine/Source/Runtime/RigVM/Private/RigVMCore/RigVM.cpp
--- a/Engine/Source/Runtime/RigVM/Private/RigVMCore/RigVM.cpp
+++ b/Engine/Source/Runtime/RigVM/Private/RigVMCore/RigVM.cpp
@@ -1330,11 +1330,6 @@ bool URigVM::ShouldHaltAtInstruction(const FName& InEventName, const uint16 Inst
 {
 	FRigVMByteCode& ByteCode = GetByteCode();
 
-	if (HaltedAtBreakpoint && InstructionIndex < HaltedAtBreakpoint->InstructionIndex)
-	{
-		return false;
-	}
-
 	TArray<TSharedPtr<FRigVMBreakpoint>> BreakpointsAtInstruction = DebugInfo->FindBreakpointsAtInstruction(InstructionIndex);
 	for (TSharedPtr<FRigVMBreakpoint> Breakpoint : BreakpointsAtInstruction)
 	{
@@ -1478,6 +1473,8 @@ bool URigVM::ShouldHaltAtInstruction(const FName& InEventName, const uint16 Inst
 
 				// Create new temporary breakpoint
 				TSharedPtr<FRigVMBreakpoint> NewBreakpoint = DebugInfo->AddBreakpoint(Context.InstructionIndex, NewBreakpointNode, 0, true);
+				DebugInfo->SetBreakpointHits(NewBreakpoint, GetInstructionVisitedCount(Context.InstructionIndex));
+				DebugInfo->SetBreakpointActivationOnHit(NewBreakpoint, GetInstructionVisitedCount(Context.InstructionIndex));
 				CurrentBreakpointAction = ERigVMBreakpointAction::None;					
 
 				HaltedAtBreakpoint = NewBreakpoint;
