Control Rig: Support for caching VM initialization and hashing improvements

diff --git a/Engine/Plugins/Animation/ControlRig/Source/ControlRig/Private/ControlRig.cpp b/Engine/Plugins/Animation/ControlRig/Source/ControlRig/Private/ControlRig.cpp
--- a/Engine/Plugins/Animation/ControlRig/Source/ControlRig/Private/ControlRig.cpp
+++ b/Engine/Plugins/Animation/ControlRig/Source/ControlRig/Private/ControlRig.cpp
@@ -16,6 +16,7 @@
 #include "Components/SkeletalMeshComponent.h"
 #include "Animation/AnimInstance.h"
 #include "Animation/AnimationPoseData.h"
+#include "Internationalization/TextKey.h"
 #if WITH_EDITOR
 #include "ControlRigModule.h"
 #include "Modules/ModuleManager.h"
@@ -818,7 +819,7 @@ void UControlRig::Execute(const EControlRigState InState, const FName& InEventNa
 				}
 			}
 			
-			if (InState == EControlRigState::Init)
+			if (InState == EControlRigState::Init && InEventName != FRigUnit_BeginExecution::EventName)
 			{
 				ExecuteUnits(Context, FRigUnit_BeginExecution::EventName);
 			}
@@ -1102,7 +1103,46 @@ void UControlRig::ExecuteUnits(FRigUnitContext& InOutContext, const FName& InEve
 #if UE_RIGVM_UCLASS_BASED_STORAGE_DISABLED
 			VM->Initialize(FRigVMMemoryContainerPtrArray(LocalMemory, 3), AdditionalArguments);
 #else
-			VM->Initialize(LocalMemory, AdditionalArguments);
+
+			if(IsInGameThread())
+			{
+				const uint32 SnapshotHash = GetHashForInitializeVMSnapShot();
+				UControlRig* CDO = Cast<UControlRig>(GetClass()->GetDefaultObject());
+
+				bool bIsValidSnapshot = false;
+				TObjectPtr<URigVM>* InitializedVMSnapshotPtr = CDO->InitializedVMSnapshots.Find(SnapshotHash);
+				if(InitializedVMSnapshotPtr && !InitializedVMSnapshotPtr->IsNull())
+				{
+					const URigVM* InitializedVMSnapshot = InitializedVMSnapshotPtr->Get();
+
+					if(VM->WorkMemoryStorageObject->GetClass() == InitializedVMSnapshot->WorkMemoryStorageObject->GetClass())
+					{
+						VM->WorkMemoryStorageObject->CopyFrom(InitializedVMSnapshot->WorkMemoryStorageObject);
+						VM->InvalidateCachedMemory();
+						VM->Initialize(LocalMemory, AdditionalArguments, false);
+						bIsValidSnapshot = true;
+					}
+					else
+					{
+						CDO->InitializedVMSnapshots.Remove(SnapshotHash);
+					}
+				}
+				
+				if(!bIsValidSnapshot)
+				{
+					VM->Initialize(LocalMemory, AdditionalArguments);
+
+					URigVM* InitializedVMSnapshot = NewObject<URigVM>(CDO, NAME_None, RF_Public);
+					InitializedVMSnapshot->WorkMemoryStorageObject = NewObject<URigVMMemoryStorage>(InitializedVMSnapshot, VM->GetWorkMemory()->GetClass());
+					InitializedVMSnapshot->WorkMemoryStorageObject->CopyFrom(VM->WorkMemoryStorageObject);
+
+					CDO->InitializedVMSnapshots.Add(SnapshotHash, InitializedVMSnapshot);
+				}
+			}
+			else
+			{
+				VM->Initialize(LocalMemory, AdditionalArguments);
+			}
 #endif
 		}
 		else
@@ -2864,6 +2904,22 @@ void UControlRig::PostInitInstance(UControlRig* InCDO)
 	}
 }
 
+uint32 UControlRig::GetHashForInitializeVMSnapShot()
+{
+	uint32 Hash = GetHierarchy()->GetNameHash();
+
+	const TArray<FRigVMExternalVariable> ExternalVariables = GetExternalVariablesImpl(false);
+
+	Hash = HashCombine(Hash, GetTypeHash(ExternalVariables.Num()));
+
+	for(const FRigVMExternalVariable& ExternalVariable : ExternalVariables)
+	{
+		Hash = HashCombine(ExternalVariable.GetTypeHash(), Hash);
+	}
+
+	return Hash;
+}
+
 void UControlRig::OnHierarchyTransformUndoRedo(URigHierarchy* InHierarchy, const FRigElementKey& InKey, ERigTransformType::Type InTransformType, const FTransform& InTransform, bool bIsUndo)
 {
 	if(InKey.Type == ERigElementType::Control)
@@ -2952,5 +3008,3 @@ UControlRig::FTransientControlPoseScope::~FTransientControlPoseScope()
 #endif
  
 #undef LOCTEXT_NAMESPACE
-
-
diff --git a/Engine/Plugins/Animation/ControlRig/Source/ControlRig/Private/Rigs/RigHierarchy.cpp b/Engine/Plugins/Animation/ControlRig/Source/ControlRig/Private/Rigs/RigHierarchy.cpp
--- a/Engine/Plugins/Animation/ControlRig/Source/ControlRig/Private/Rigs/RigHierarchy.cpp
+++ b/Engine/Plugins/Animation/ControlRig/Source/ControlRig/Private/Rigs/RigHierarchy.cpp
@@ -216,7 +216,7 @@ void URigHierarchy::Load(FArchive& Ar)
 		Element->Load(Ar, this, FRigBaseElement::StaticData);
 	}
 
-	TopologyVersion++;
+	IncrementTopologyVersion();
 
 	for(int32 ElementIndex = 0; ElementIndex < ElementCount; ElementIndex++)
 	{
@@ -224,7 +224,7 @@ void URigHierarchy::Load(FArchive& Ar)
 		Element->Load(Ar, this, FRigBaseElement::InterElementData);
 	}
 
-	TopologyVersion++;
+	IncrementTopologyVersion();
 
 	for(int32 ElementIndex = 0; ElementIndex < ElementCount; ElementIndex++)
 	{
@@ -412,6 +412,7 @@ void URigHierarchy::CopyHierarchy(URigHierarchy* InHierarchy)
 		//FRigBaseElement* Target = MakeElement(Key.Type);
 		
 		Target->Key = Key;
+		Target->NameString = Source->NameString;
 		Target->SubIndex = SubIndex;
 		Target->Index = Elements.Add(Target);
 
@@ -435,6 +436,19 @@ void URigHierarchy::CopyHierarchy(URigHierarchy* InHierarchy)
 	EnsureCacheValidity();
 }
 
+uint32 URigHierarchy::GetNameHash() const
+{
+	uint32 Hash = GetTypeHash(GetTopologyVersion());
+
+	for (int32 ElementIndex = 0; ElementIndex < Elements.Num(); ElementIndex++)
+	{
+		const FRigBaseElement* Element = Elements[ElementIndex];
+		Hash = HashCombine(Hash, GetTypeHash(Element->GetName()));
+	}
+
+	return Hash;
+}
+
 #if WITH_EDITOR
 void URigHierarchy::RegisterListeningHierarchy(URigHierarchy* InHierarchy)
 {
@@ -1579,43 +1593,65 @@ FRigElementKey URigHierarchy::GetWorldSpaceReferenceKey()
 
 TArray<FRigElementKey> URigHierarchy::GetAllKeys(bool bTraverse, ERigElementType InElementType) const
 {
-	TArray<FRigElementKey> Keys;
-	Keys.Reserve(Elements.Num());
+	return GetKeysByPredicate([InElementType](const FRigBaseElement& InElement)
+	{
+		return InElement.IsTypeOf(InElementType);
+	}, bTraverse);
+}
 
-	if(bTraverse)
+TArray<FRigElementKey> URigHierarchy::GetKeysByPredicate(
+	TFunctionRef<bool(const FRigBaseElement&)> InPredicateFunc,
+	bool bTraverse
+	) const
+{
+	auto ElementTraverser = [&](TFunctionRef<void(const FRigBaseElement&)> InProcessFunc)
 	{
-		TArray<bool> ElementVisited;
-		ElementVisited.AddZeroed(Elements.Num());
-		
-		for (int32 ElementIndex = 0; ElementIndex < Elements.Num(); ElementIndex++)
+		if(bTraverse)
 		{
-			FRigBaseElement* Element = Elements[ElementIndex];
-			Traverse(Element, true, [&ElementVisited, &Keys, InElementType](FRigBaseElement* InElement, bool& bContinue)
-            {
-				bContinue = !ElementVisited[InElement->GetIndex()];
-
-				if(bContinue)
+			// TBitArray reserves 4, we'll do 16 so we can remember at least 512 elements before
+			// we need to hit the heap.
+			TBitArray<TInlineAllocator<16>> ElementVisited(false, Elements.Num());
+			
+			for (int32 ElementIndex = 0; ElementIndex < Elements.Num(); ElementIndex++)
+			{
+				FRigBaseElement* Element = Elements[ElementIndex];
+				Traverse(Element, true, [&ElementVisited, InProcessFunc, InPredicateFunc](FRigBaseElement* InElement, bool& bContinue)
 				{
-					if(InElement->IsTypeOf(InElementType))
+					bContinue = !ElementVisited[InElement->GetIndex()];
+
+					if(bContinue)
 					{
-						Keys.Add(InElement->GetKey());
+						if(InPredicateFunc(*InElement))
+						{
+							InProcessFunc(*InElement);
+						}
+						ElementVisited[InElement->GetIndex()] = true;
 					}
-					ElementVisited[InElement->GetIndex()] = true;
-				}
-            });
+				});
+			}
 		}
-	}
-	else
-	{
-		for (int32 ElementIndex = 0; ElementIndex < Elements.Num(); ElementIndex++)
+		else
 		{
-			FRigBaseElement* Element = Elements[ElementIndex];
-			if(Element->IsTypeOf(InElementType))
+			for (int32 ElementIndex = 0; ElementIndex < Elements.Num(); ElementIndex++)
 			{
-				Keys.Add(Element->GetKey());
+				const FRigBaseElement* Element = Elements[ElementIndex];
+				if(InPredicateFunc(*Element))
+				{
+					InProcessFunc(*Element);
+				}
 			}
 		}
-	}
+	};
+	
+	// First count up how many elements we matched and only reserve that amount. There's very little overhead
+	// since we're just running over the same data, so it should still be hot when we do the second pass.
+	int32 NbElements = 0;
+	ElementTraverser([&NbElements](const FRigBaseElement&) { NbElements++; });
+	
+	TArray<FRigElementKey> Keys;
+	Keys.Reserve(NbElements);
+	ElementTraverser([&Keys](const FRigBaseElement& InElement) { Keys.Add(InElement.GetKey()); });
+
 	return Keys;
 }
 
@@ -1828,6 +1864,12 @@ URigHierarchyController* URigHierarchy::GetController(bool bCreateIfNeeded)
 	return nullptr;
 }
 
+void URigHierarchy::IncrementTopologyVersion()
+{
+	TopologyVersion++;
+	KeyCollectionCache.Reset();
+}
+
 FRigPose URigHierarchy::GetPose(
 	bool bInitial,
 	ERigElementType InElementType,
@@ -4625,7 +4667,6 @@ void URigHierarchy::ComputeParentConstraintIndices(
 		}
 	}
 }
-
 void URigHierarchy::IntegrateParentConstraintVector(
 	FVector& OutVector,
 	const FTransform& InTransform,
diff --git a/Engine/Plugins/Animation/ControlRig/Source/ControlRig/Private/Rigs/RigHierarchyController.cpp b/Engine/Plugins/Animation/ControlRig/Source/ControlRig/Private/Rigs/RigHierarchyController.cpp
--- a/Engine/Plugins/Animation/ControlRig/Source/ControlRig/Private/Rigs/RigHierarchyController.cpp
+++ b/Engine/Plugins/Animation/ControlRig/Source/ControlRig/Private/Rigs/RigHierarchyController.cpp
@@ -1566,12 +1566,13 @@ int32 URigHierarchyController::AddElement(FRigBaseElement* InElementToAdd, FRigB
 {
 	ensure(IsValid());
 
+	InElementToAdd->NameString = InElementToAdd->Key.Name.ToString();
 	InElementToAdd->SubIndex = Hierarchy->Num(InElementToAdd->Key.Type);
 	InElementToAdd->Index = Hierarchy->Elements.Add(InElementToAdd);
 	Hierarchy->ElementsPerType[URigHierarchy::RigElementTypeToFlatIndex(InElementToAdd->GetKey().Type)].Add(InElementToAdd);
 
 	Hierarchy->IndexLookup.Add(InElementToAdd->Key, InElementToAdd->Index);
-	Hierarchy->TopologyVersion++;
+	Hierarchy->IncrementTopologyVersion();
 
 	{
 		const TGuardValue<bool> Guard(bSuspendNotifications, true);
@@ -1731,7 +1732,7 @@ bool URigHierarchyController::RemoveElement(FRigBaseElement* InElement)
 		}
 	}
 	
-	Hierarchy->TopologyVersion++;
+	Hierarchy->IncrementTopologyVersion();
 
 	Notify(ERigHierarchyNotification::ElementRemoved, InElement);
 	if(Hierarchy->Num() == 0)
@@ -1827,6 +1828,7 @@ bool URigHierarchyController::RenameElement(FRigBaseElement* InElement, const FN
    
 		TGuardValue<TMap<FRigElementKey, int32>> MapGuard(Hierarchy->IndexLookup, TemporaryMap);
 		InElement->Key.Name = Hierarchy->GetSafeNewName(InName.ToString(), InElement->GetType());
+		InElement->NameString = InElement->Key.Name.ToString();
 	}
 	
 	const FRigElementKey NewKey = InElement->GetKey();
@@ -1861,7 +1863,7 @@ bool URigHierarchyController::RenameElement(FRigBaseElement* InElement, const FN
 	}
 	
 	Hierarchy->PreviousNameMap.FindOrAdd(NewKey) = OldKey;
-	Hierarchy->TopologyVersion++;
+	Hierarchy->IncrementTopologyVersion();
 	Notify(ERigHierarchyNotification::ElementRenamed, InElement);
 
 	return true;
@@ -1992,7 +1994,7 @@ bool URigHierarchyController::AddParent(FRigBaseElement* InChild, FRigBaseElemen
 		AddElementToDirty(Constraint.ParentElement, SingleParentElement);
 		SingleParentElement->ParentElement = Constraint.ParentElement;
 
-		Hierarchy->TopologyVersion++;
+		Hierarchy->IncrementTopologyVersion();
 
 		if(!bMaintainGlobalTransform)
 		{
@@ -2027,7 +2029,7 @@ bool URigHierarchyController::AddParent(FRigBaseElement* InChild, FRigBaseElemen
 			}
 		}
 
-		Hierarchy->TopologyVersion++;
+		Hierarchy->IncrementTopologyVersion();
 
 		if(InWeight > SMALL_NUMBER)
 		{
@@ -2143,7 +2145,7 @@ bool URigHierarchyController::RemoveParent(FRigBaseElement* InChild, FRigBaseEle
 			// remove the previous parent
 			SingleParentElement->ParentElement = nullptr;
 			RemoveElementToDirty(InParent, SingleParentElement); 
-			Hierarchy->TopologyVersion++;
+			Hierarchy->IncrementTopologyVersion();
 
 			if(!bMaintainGlobalTransform)
 			{
@@ -2216,7 +2218,7 @@ bool URigHierarchyController::RemoveParent(FRigBaseElement* InChild, FRigBaseEle
 				ControlElement->Shape.MarkDirty(ERigTransformType::InitialGlobal);
 			}
 
-			Hierarchy->TopologyVersion++;
+			Hierarchy->IncrementTopologyVersion();
 
 			if(!bMaintainGlobalTransform)
 			{
diff --git a/Engine/Plugins/Animation/ControlRig/Source/ControlRig/Private/Rigs/RigHierarchyDefines.cpp b/Engine/Plugins/Animation/ControlRig/Source/ControlRig/Private/Rigs/RigHierarchyDefines.cpp
--- a/Engine/Plugins/Animation/ControlRig/Source/ControlRig/Private/Rigs/RigHierarchyDefines.cpp
+++ b/Engine/Plugins/Animation/ControlRig/Source/ControlRig/Private/Rigs/RigHierarchyDefines.cpp
@@ -115,10 +115,15 @@ FRigElementKeyCollection FRigElementKeyCollection::MakeFromName(
 
 	check(InHierarchy);
 
-	FRigElementKeyCollection Collection(InHierarchy->GetAllKeys(true));
-	Collection = Collection.FilterByType(InElementTypes);
-	Collection = Collection.FilterByName(InPartialName);
-	return Collection;
+	constexpr bool bTraverse = true;
+
+	const FString PartialNameString = InPartialName.ToString();
+	
+	return InHierarchy->GetKeysByPredicate([PartialNameString, InElementTypes](const FRigBaseElement& InElement) -> bool
+	{
+		return InElement.IsTypeOf(static_cast<ERigElementType>(InElementTypes)) &&
+			   InElement.GetNameString().Contains(PartialNameString);
+	}, bTraverse);
 }
 
 FRigElementKeyCollection FRigElementKeyCollection::MakeFromChain(
diff --git a/Engine/Plugins/Animation/ControlRig/Source/ControlRig/Private/Rigs/RigHierarchyElements.cpp b/Engine/Plugins/Animation/ControlRig/Source/ControlRig/Private/Rigs/RigHierarchyElements.cpp
--- a/Engine/Plugins/Animation/ControlRig/Source/ControlRig/Private/Rigs/RigHierarchyElements.cpp
+++ b/Engine/Plugins/Animation/ControlRig/Source/ControlRig/Private/Rigs/RigHierarchyElements.cpp
@@ -82,6 +82,8 @@ void FRigBaseElement::Load(FArchive& Ar, URigHierarchy* Hierarchy, ESerializatio
 
 		ensure(LoadedKey.Type == Key.Type);
 		Key = LoadedKey;
+
+		NameString = Key.Name.ToString();
 	}
 }
 
diff --git a/Engine/Plugins/Animation/ControlRig/Source/ControlRig/Private/Units/Execution/RigUnit_Collection.cpp b/Engine/Plugins/Animation/ControlRig/Source/ControlRig/Private/Units/Execution/RigUnit_Collection.cpp
--- a/Engine/Plugins/Animation/ControlRig/Source/ControlRig/Private/Units/Execution/RigUnit_Collection.cpp
+++ b/Engine/Plugins/Animation/ControlRig/Source/ControlRig/Private/Units/Execution/RigUnit_Collection.cpp
@@ -8,29 +8,39 @@
 #include "Units/RigUnitTest.h"
 #endif
 
+#define FRigUnit_CollectionChain_Hash 1
+#define FRigUnit_CollectionNameSearch_Hash 2
+#define FRigUnit_CollectionChildren_Hash 3
+#define FRigUnit_CollectionReplaceItems_Hash 4
+#define FRigUnit_CollectionItems_Hash 5
+#define FRigUnit_CollectionUnion_Hash 6
+#define FRigUnit_CollectionIntersection_Hash 7
+#define FRigUnit_CollectionDifference_Hash 8
+#define FRigUnit_CollectionReverse_Hash 9
+#define FRigUnit_CollectionCount_Hash 10
+#define FRigUnit_CollectionItemAtIndex_Hash 11
+
 FRigUnit_CollectionChain_Execute()
 {
-	FRigUnit_CollectionChainArray::StaticExecute(RigVMExecuteContext, FirstItem, LastItem, Reverse, Collection.Keys, CachedCollection, CachedHierarchyHash, Context);
+	FRigUnit_CollectionChainArray::StaticExecute(RigVMExecuteContext, FirstItem, LastItem, Reverse, Collection.Keys, Context);
 }
 
 FRigUnit_CollectionChainArray_Execute()
 {
-	if (Context.State == EControlRigState::Init)
+	uint32 Hash = FRigUnit_CollectionChain_Hash + Context.Hierarchy->GetTopologyVersion() * 17;
+	Hash = HashCombine(Hash, GetTypeHash(FirstItem));
+	Hash = HashCombine(Hash, GetTypeHash(LastItem));
+	Hash = HashCombine(Hash, Reverse ? 1 : 0);
+
+	FRigElementKeyCollection Collection;
+	if(const FRigElementKeyCollection* Cache = Context.Hierarchy->FindCachedCollection(Hash))
 	{
-		CachedHierarchyHash = INDEX_NONE;
+		Collection = *Cache;
 	}
-
-	int32 CurrentHierarchyHash = Context.Hierarchy->GetTopologyVersion() * 17;
-	CurrentHierarchyHash += GetTypeHash(FirstItem);
-	CurrentHierarchyHash += GetTypeHash(LastItem);
-	CurrentHierarchyHash += Reverse ? 1 : 0;
-
-	if (CachedHierarchyHash != CurrentHierarchyHash || CachedCollection.IsEmpty())
+	else
 	{
-		CachedHierarchyHash = CurrentHierarchyHash;
-		CachedCollection = FRigElementKeyCollection::MakeFromChain(Context.Hierarchy, FirstItem, LastItem, Reverse);
-
-		if (CachedCollection.IsEmpty())
+		Collection = FRigElementKeyCollection::MakeFromChain(Context.Hierarchy, FirstItem, LastItem, Reverse);
+		if (Collection.IsEmpty())
 		{
 			if (Context.Hierarchy->GetIndex(FirstItem) == INDEX_NONE)
 			{
@@ -47,72 +57,73 @@ FRigUnit_CollectionChainArray_Execute()
 				}
 			}
 		}
+		Context.Hierarchy->AddCachedCollection(Hash, Collection);
 	}
 
-	Items = CachedCollection.Keys;
+	Items = Collection.Keys;
 }
 
 FRigUnit_CollectionNameSearch_Execute()
 {
-	FRigUnit_CollectionNameSearchArray::StaticExecute(RigVMExecuteContext, PartialName, TypeToSearch, Collection.Keys, CachedCollection, CachedHierarchyHash, Context);
+	FRigUnit_CollectionNameSearchArray::StaticExecute(RigVMExecuteContext, PartialName, TypeToSearch, Collection.Keys, Context);
 }
 
 FRigUnit_CollectionNameSearchArray_Execute()
 {
 	DECLARE_SCOPE_HIERARCHICAL_COUNTER_RIGUNIT()
 
-	if (Context.State == EControlRigState::Init)
+	uint32 Hash = FRigUnit_CollectionNameSearch_Hash + Context.Hierarchy->GetTopologyVersion() * 17;
+	Hash = HashCombine(Hash, GetTypeHash(PartialName));
+	Hash = HashCombine(Hash, (int32)TypeToSearch * 8);
+
+	FRigElementKeyCollection Collection;
+	if(const FRigElementKeyCollection* Cache = Context.Hierarchy->FindCachedCollection(Hash))
 	{
-		CachedHierarchyHash = INDEX_NONE;
+		Collection = *Cache;
 	}
-
-	int32 CurrentHierarchyHash = Context.Hierarchy->GetTopologyVersion() * 17;
-	CurrentHierarchyHash += GetTypeHash(PartialName);
-	CurrentHierarchyHash += (int32)TypeToSearch * 8;
-
-	if (CachedHierarchyHash != CurrentHierarchyHash || CachedCollection.IsEmpty())
+	else
 	{
-		CachedHierarchyHash = CurrentHierarchyHash;
-		CachedCollection = FRigElementKeyCollection::MakeFromName(Context.Hierarchy, PartialName, (uint8)TypeToSearch);
+		Collection = FRigElementKeyCollection::MakeFromName(Context.Hierarchy, PartialName, (uint8)TypeToSearch);
+		Context.Hierarchy->AddCachedCollection(Hash, Collection);
 	}
 
-	Items = CachedCollection.Keys;
+	Items = Collection.Keys;
 }
 
 FRigUnit_CollectionChildren_Execute()
 {
-	FRigUnit_CollectionChildrenArray::StaticExecute(RigVMExecuteContext, Parent, bIncludeParent, bRecursive, TypeToSearch, Collection.Keys, CachedCollection, CachedHierarchyHash, Context);
+	FRigUnit_CollectionChildrenArray::StaticExecute(RigVMExecuteContext, Parent, bIncludeParent, bRecursive, TypeToSearch, Collection.Keys, Context);
 }
 
 FRigUnit_CollectionChildrenArray_Execute()
 {
 	DECLARE_SCOPE_HIERARCHICAL_COUNTER_RIGUNIT()
 
-	if (Context.State == EControlRigState::Init)
+	uint32 Hash = FRigUnit_CollectionChildren_Hash + Context.Hierarchy->GetTopologyVersion() * 17;
+	Hash = HashCombine(Hash, GetTypeHash(Parent));
+	Hash = HashCombine(Hash, bRecursive ? 2 : 0);
+	Hash = HashCombine(Hash, bIncludeParent ? 1 : 0);
+	Hash = HashCombine(Hash, (int32)TypeToSearch * 8);
+
+	FRigElementKeyCollection Collection;
+	if(const FRigElementKeyCollection* Cache = Context.Hierarchy->FindCachedCollection(Hash))
 	{
-		CachedHierarchyHash = INDEX_NONE;
+		Collection = *Cache;
 	}
-
-	int32 CurrentHierarchyHash = Context.Hierarchy->GetTopologyVersion() * 17;
-	CurrentHierarchyHash += GetTypeHash(Parent);
-	CurrentHierarchyHash += bRecursive ? 2 : 0;
-	CurrentHierarchyHash += bIncludeParent ? 1 : 0;
-	CurrentHierarchyHash += (int32)TypeToSearch * 8;
-
-	if (CachedHierarchyHash != CurrentHierarchyHash || CachedCollection.IsEmpty())
+	else
 	{
-		CachedHierarchyHash = CurrentHierarchyHash;
-		CachedCollection = FRigElementKeyCollection::MakeFromChildren(Context.Hierarchy, Parent, bRecursive, bIncludeParent, (uint8)TypeToSearch);
-		if (CachedCollection.IsEmpty())
+		Collection = FRigElementKeyCollection::MakeFromChildren(Context.Hierarchy, Parent, bRecursive, bIncludeParent, (uint8)TypeToSearch);
+		if (Collection.IsEmpty())
 		{
 			if (Context.Hierarchy->GetIndex(Parent) == INDEX_NONE)
 			{
 				UE_CONTROLRIG_RIGUNIT_REPORT_WARNING(TEXT("Parent '%s' is not valid."), *Parent.ToString());
 			}
 		}
+		Context.Hierarchy->AddCachedCollection(Hash, Collection);
 	}
 
-	Items = CachedCollection.Keys;
+	Items = Collection.Keys;
 }
 
 #if WITH_EDITOR
@@ -156,29 +167,28 @@ IMPLEMENT_RIGUNIT_AUTOMATION_TEST(FRigUnit_CollectionChildren)
 
 FRigUnit_CollectionReplaceItems_Execute()
 {
-	FRigUnit_CollectionReplaceItemsArray::StaticExecute(RigVMExecuteContext, Items.Keys, Old, New, RemoveInvalidItems, bAllowDuplicates, Collection.Keys, CachedCollection, CachedHierarchyHash, Context);
+	FRigUnit_CollectionReplaceItemsArray::StaticExecute(RigVMExecuteContext, Items.Keys, Old, New, RemoveInvalidItems, bAllowDuplicates, Collection.Keys, Context);
 }
 
 FRigUnit_CollectionReplaceItemsArray_Execute()
 {
 	DECLARE_SCOPE_HIERARCHICAL_COUNTER_RIGUNIT()
 
-	if (Context.State == EControlRigState::Init)
+	uint32 Hash = FRigUnit_CollectionReplaceItems_Hash + Context.Hierarchy->GetTopologyVersion() * 17;
+	Hash = HashCombine(Hash, GetTypeHash(Items));
+	Hash = HashCombine(Hash, 12 * GetTypeHash(Old));
+	Hash = HashCombine(Hash, 13 * GetTypeHash(New));
+	Hash = HashCombine(Hash, RemoveInvalidItems ? 14 : 0);
+
+	FRigElementKeyCollection Collection;
+	if(const FRigElementKeyCollection* Cache = Context.Hierarchy->FindCachedCollection(Hash))
 	{
-		CachedHierarchyHash = INDEX_NONE;
+		Collection = *Cache;
 	}
-
-	int32 CurrentHierarchyHash = Context.Hierarchy->GetTopologyVersion() * 17;
-	CurrentHierarchyHash += GetTypeHash(Items);
-	CurrentHierarchyHash += 12 * GetTypeHash(Old);
-	CurrentHierarchyHash += 13 * GetTypeHash(New);
-	CurrentHierarchyHash += RemoveInvalidItems ? 14 : 0;
-
-	if (CachedHierarchyHash != CurrentHierarchyHash || CachedCollection.IsEmpty())
+	else
 	{
-		CachedHierarchyHash = CurrentHierarchyHash;
-		CachedCollection.Reset();
-
+		Collection.Reset();
+		
 		for (int32 Index = 0; Index < Items.Num(); Index++)
 		{
 			FRigElementKey Key = Items[Index];
@@ -188,21 +198,23 @@ FRigUnit_CollectionReplaceItemsArray_Execute()
 			{
 				if(bAllowDuplicates)
 				{
-					CachedCollection.Add(Key);
+					Collection.Add(Key);
 				}
 				else
 				{
-					CachedCollection.AddUnique(Key);
+					Collection.AddUnique(Key);
 				}
 			}
 			else if(!RemoveInvalidItems)
 			{
-				CachedCollection.Add(FRigElementKey());
+				Collection.Add(FRigElementKey());
 			}
 		}
-	}
 
-	Result = CachedCollection.Keys;
+		Context.Hierarchy->AddCachedCollection(Hash, Collection);
+	}
+	
+	Result = Collection.Keys;
 }
 
 FRigUnit_CollectionItems_Execute()
diff --git a/Engine/Plugins/Animation/ControlRig/Source/ControlRig/Private/Units/Execution/RigUnit_Collection.h b/Engine/Plugins/Animation/ControlRig/Source/ControlRig/Private/Units/Execution/RigUnit_Collection.h
--- a/Engine/Plugins/Animation/ControlRig/Source/ControlRig/Private/Units/Execution/RigUnit_Collection.h
+++ b/Engine/Plugins/Animation/ControlRig/Source/ControlRig/Private/Units/Execution/RigUnit_Collection.h
@@ -31,7 +31,6 @@ struct CONTROLRIG_API FRigUnit_CollectionChain : public FRigUnit_CollectionBase
 	{
 		FirstItem = LastItem = FRigElementKey(NAME_None, ERigElementType::Bone);
 		Reverse = false;
-		CachedHierarchyHash = INDEX_NONE;
 	}
 
 	RIGVM_METHOD()
@@ -48,12 +47,6 @@ struct CONTROLRIG_API FRigUnit_CollectionChain : public FRigUnit_CollectionBase
 
 	UPROPERTY(meta = (Output))
 	FRigElementKeyCollection Collection;
-
-	UPROPERTY()
-	FRigElementKeyCollection CachedCollection;
-
-	UPROPERTY()
-	int32 CachedHierarchyHash;
 };
 
 /**
@@ -69,7 +62,6 @@ struct CONTROLRIG_API FRigUnit_CollectionChainArray : public FRigUnit_Collection
 	{
 		FirstItem = LastItem = FRigElementKey(NAME_None, ERigElementType::Bone);
 		Reverse = false;
-		CachedHierarchyHash = INDEX_NONE;
 	}
 
 	RIGVM_METHOD()
@@ -86,12 +78,6 @@ struct CONTROLRIG_API FRigUnit_CollectionChainArray : public FRigUnit_Collection
 
 	UPROPERTY(meta = (Output))
 	TArray<FRigElementKey> Items;
-
-	UPROPERTY()
-	FRigElementKeyCollection CachedCollection;
-
-	UPROPERTY()
-	int32 CachedHierarchyHash;
 };
 
 /**
@@ -107,7 +93,6 @@ struct CONTROLRIG_API FRigUnit_CollectionNameSearch : public FRigUnit_Collection
 	{
 		PartialName = NAME_None;
 		TypeToSearch = ERigElementType::All;
-		CachedHierarchyHash = INDEX_NONE;
 	}
 
 	RIGVM_METHOD()
@@ -121,12 +106,6 @@ struct CONTROLRIG_API FRigUnit_CollectionNameSearch : public FRigUnit_Collection
 
 	UPROPERTY(meta = (Output))
 	FRigElementKeyCollection Collection;
-
-	UPROPERTY()
-	FRigElementKeyCollection CachedCollection;
-
-	UPROPERTY()
-	int32 CachedHierarchyHash;
 };
 
 /**
@@ -142,7 +121,6 @@ struct CONTROLRIG_API FRigUnit_CollectionNameSearchArray : public FRigUnit_Colle
 	{
 		PartialName = NAME_None;
 		TypeToSearch = ERigElementType::All;
-		CachedHierarchyHash = INDEX_NONE;
 	}
 
 	RIGVM_METHOD()
@@ -156,12 +134,6 @@ struct CONTROLRIG_API FRigUnit_CollectionNameSearchArray : public FRigUnit_Colle
 
 	UPROPERTY(meta = (Output))
 	TArray<FRigElementKey> Items;
-
-	UPROPERTY()
-	FRigElementKeyCollection CachedCollection;
-
-	UPROPERTY()
-	int32 CachedHierarchyHash;
 };
 
 /**
@@ -179,7 +151,6 @@ struct CONTROLRIG_API FRigUnit_CollectionChildren : public FRigUnit_CollectionBa
 		bIncludeParent = false;
 		bRecursive = false;
 		TypeToSearch = ERigElementType::All;
-		CachedHierarchyHash = INDEX_NONE;
 	}
 
 	RIGVM_METHOD()
@@ -199,12 +170,6 @@ struct CONTROLRIG_API FRigUnit_CollectionChildren : public FRigUnit_CollectionBa
 
 	UPROPERTY(meta = (Output))
 	FRigElementKeyCollection Collection;
-
-	UPROPERTY()
-	FRigElementKeyCollection CachedCollection;
-
-	UPROPERTY()
-	int32 CachedHierarchyHash;
 };
 
 /**
@@ -222,7 +187,6 @@ struct CONTROLRIG_API FRigUnit_CollectionChildrenArray : public FRigUnit_Collect
 		bIncludeParent = false;
 		bRecursive = false;
 		TypeToSearch = ERigElementType::All;
-		CachedHierarchyHash = INDEX_NONE;
 	}
 
 	RIGVM_METHOD()
@@ -242,12 +206,6 @@ struct CONTROLRIG_API FRigUnit_CollectionChildrenArray : public FRigUnit_Collect
 
 	UPROPERTY(meta = (Output))
 	TArray<FRigElementKey> Items;
-
-	UPROPERTY()
-	FRigElementKeyCollection CachedCollection;
-
-	UPROPERTY()
-	int32 CachedHierarchyHash;
 };
 
 /**
@@ -262,7 +220,6 @@ struct CONTROLRIG_API FRigUnit_CollectionReplaceItems : public FRigUnit_Collecti
 	{
 		Old = New = NAME_None;
 		RemoveInvalidItems = false;
-		CachedHierarchyHash = INDEX_NONE;
 		bAllowDuplicates = false;
 	}
 
@@ -286,12 +243,6 @@ struct CONTROLRIG_API FRigUnit_CollectionReplaceItems : public FRigUnit_Collecti
 
 	UPROPERTY(meta = (Output))
 	FRigElementKeyCollection Collection;
-
-	UPROPERTY()
-	FRigElementKeyCollection CachedCollection;
-
-	UPROPERTY()
-	int32 CachedHierarchyHash;
 };
 
 /**
@@ -307,7 +258,6 @@ struct CONTROLRIG_API FRigUnit_CollectionReplaceItemsArray : public FRigUnit_Col
 		Old = New = NAME_None;
 		RemoveInvalidItems = false;
 		bAllowDuplicates = false;
-		CachedHierarchyHash = INDEX_NONE;
 	}
 
 	RIGVM_METHOD()
@@ -330,12 +280,6 @@ struct CONTROLRIG_API FRigUnit_CollectionReplaceItemsArray : public FRigUnit_Col
 
 	UPROPERTY(meta = (Output))
 	TArray<FRigElementKey> Result;
-
-	UPROPERTY()
-	FRigElementKeyCollection CachedCollection;
-
-	UPROPERTY()
-	int32 CachedHierarchyHash;
 };
 
 /**
diff --git a/Engine/Plugins/Animation/ControlRig/Source/ControlRig/Public/ControlRig.h b/Engine/Plugins/Animation/ControlRig/Source/ControlRig/Public/ControlRig.h
--- a/Engine/Plugins/Animation/ControlRig/Source/ControlRig/Public/ControlRig.h
+++ b/Engine/Plugins/Animation/ControlRig/Source/ControlRig/Public/ControlRig.h
@@ -434,6 +434,14 @@ private:
 	UPROPERTY()
 	TObjectPtr<URigVM> VM;
 
+	// this is only used by the CDO
+	// and stores an initialized VM per hash.
+	UPROPERTY(transient)
+	TMap<uint32, TObjectPtr<URigVM>> InitializedVMSnapshots;
+
+	// computes the hash used to store / find VM snapshots
+	uint32 GetHashForInitializeVMSnapShot();
+	
 	UPROPERTY()
 	TObjectPtr<URigHierarchy> DynamicHierarchy;
 
diff --git a/Engine/Plugins/Animation/ControlRig/Source/ControlRig/Public/Rigs/RigHierarchy.h b/Engine/Plugins/Animation/ControlRig/Source/ControlRig/Public/Rigs/RigHierarchy.h
--- a/Engine/Plugins/Animation/ControlRig/Source/ControlRig/Public/Rigs/RigHierarchy.h
+++ b/Engine/Plugins/Animation/ControlRig/Source/ControlRig/Public/Rigs/RigHierarchy.h
@@ -126,6 +126,12 @@ public:
 	UFUNCTION(BlueprintCallable, Category = URigHierarchy)
 	void CopyHierarchy(URigHierarchy* InHierarchy);
 
+	/**
+	 * Returns a hash for the hierarchy representing all names
+	 * as well as the topology version.
+	 */
+	uint32 GetNameHash() const;
+
 #if WITH_EDITOR
 	/**
 	* Add dependent hierarchies that listens to changes made to this hierarchy
@@ -1881,6 +1887,15 @@ public:
 	 */
 	TArray<FRigElementKey> GetAllKeys(bool bTraverse = false, ERigElementType InElementType = ERigElementType::All) const;
 
+	/**
+	 * Returns element keys of this hierarchy, filtered by a predicate.
+	 * @param InPredicateFunc The predicate function to apply. Should return \c true if the element
+	 *    should be added to the result array.
+	 * @param bInTraverse If set to true the keys will be returned by depth first traversal
+	 * @return The keys of all elements
+	 */
+	TArray<FRigElementKey> GetKeysByPredicate(TFunctionRef<bool(const FRigBaseElement&)> InPredicateFunc, bool bInTraverse = false) const;
+
 	/**
 	 * Returns all element keys of this hierarchy
 	 * @param bTraverse If set to true the keys will be returned by depth first traversal
@@ -1978,6 +1993,11 @@ public:
 	 */
 	uint16 GetTopologyVersion() const { return TopologyVersion; }
 
+	/**
+	 * Increments the topology version
+	 */
+	void IncrementTopologyVersion();
+
 	/**
 	 * Returns the current / initial pose of the hierarchy
 	 * @param bInitial If set to true the initial pose will be returned
@@ -2793,6 +2813,16 @@ protected:
 		return ERigElementType::None;
 	}
 
+public:
+
+	const FRigElementKeyCollection* FindCachedCollection(uint32 InHash) const { return KeyCollectionCache.Find(InHash); }
+	FRigElementKeyCollection& FindOrAddCachedCollection(uint32 InHash) const { return KeyCollectionCache.FindOrAdd(InHash); };
+	void AddCachedCollection(uint32 InHash, const FRigElementKeyCollection& InCollection) const { KeyCollectionCache.Add(InHash, InCollection); }
+	
+private:
+	
+	mutable TMap<uint32, FRigElementKeyCollection> KeyCollectionCache;
+
 	FTransform GetWorldTransformForReference(const FRigUnitContext* InContext, const FRigElementKey& InKey, bool bInitial);
 	
 	FORCEINLINE static float GetWeightForLerp(const float WeightA, const float WeightB)
diff --git a/Engine/Plugins/Animation/ControlRig/Source/ControlRig/Public/Rigs/RigHierarchyElements.h b/Engine/Plugins/Animation/ControlRig/Source/ControlRig/Public/Rigs/RigHierarchyElements.h
--- a/Engine/Plugins/Animation/ControlRig/Source/ControlRig/Public/Rigs/RigHierarchyElements.h
+++ b/Engine/Plugins/Animation/ControlRig/Source/ControlRig/Public/Rigs/RigHierarchyElements.h
@@ -411,6 +411,7 @@ public:
 
 	FRigBaseElement()
     : Key()
+	, NameString()
     , Index(INDEX_NONE)
 	, SubIndex(INDEX_NONE)
 	, bSelected(false)
@@ -431,6 +432,9 @@ protected:
 	UPROPERTY(BlueprintReadOnly, VisibleAnywhere, Category = RigElement, meta = (AllowPrivateAccess = "true"))
 	FRigElementKey Key;
 
+	UPROPERTY(transient)
+	FString NameString;
+
 	UPROPERTY(BlueprintReadOnly, VisibleAnywhere, Category = RigElement, meta = (AllowPrivateAccess = "true"))
 	int32 Index;
 
@@ -453,6 +457,7 @@ public:
 	virtual void Load(FArchive& Ar, URigHierarchy* Hierarchy, ESerializationPhase SerializationPhase);
 
 	FORCEINLINE const FName& GetName() const { return Key.Name; }
+	FORCEINLINE const FString& GetNameString() const { return NameString; }
 	FORCEINLINE virtual const FName& GetDisplayName() const { return GetName(); }
 	FORCEINLINE ERigElementType GetType() const { return Key.Type; }
 	FORCEINLINE const FRigElementKey& GetKey() const { return Key; }
diff --git a/Engine/Source/Runtime/RigVM/Private/RigVMCore/RigVM.cpp b/Engine/Source/Runtime/RigVM/Private/RigVMCore/RigVM.cpp
--- a/Engine/Source/Runtime/RigVM/Private/RigVMCore/RigVM.cpp
+++ b/Engine/Source/Runtime/RigVM/Private/RigVMCore/RigVM.cpp
@@ -5,6 +5,8 @@
 #include "UObject/AnimObjectVersion.h"
 #include "UObject/UE5MainStreamObjectVersion.h"
 #include "HAL/PlatformTLS.h"
+#include "Async/ParallelFor.h"
+
 
 void FRigVMParameter::Serialize(FArchive& Ar)
 {
@@ -659,26 +661,14 @@ void URigVM::CopyFrom(URigVM* InVM, bool bDeferCopy, bool bReferenceLiteralMemor
 			else if(TargetMemory->GetClass() != SourceMemory->GetClass())
 			{
 				TargetMemory->Rename(nullptr, GetTransientPackage(), REN_ForceNoResetLoaders | REN_DoNotDirty | REN_DontCreateRedirectors | REN_NonTransactional);
-				TargetMemory->MarkPendingKill();
 				TargetMemory = NewObject<URigVMMemoryStorage>(Outer, SourceMemory->GetClass());
 			}
 
-			for(int32 PropertyIndex = 0; PropertyIndex < TargetMemory->Num(); PropertyIndex++)
-			{
-				URigVMMemoryStorage::CopyProperty(
-					TargetMemory,
-					PropertyIndex,
-					FRigVMPropertyPath::Empty,
-					SourceMemory,
-					PropertyIndex,
-					FRigVMPropertyPath::Empty
-				);
-			}
+			TargetMemory->CopyFrom(SourceMemory);
 		}
 		else if(TargetMemory != nullptr)
 		{
 			TargetMemory->Rename(nullptr, GetTransientPackage(), REN_ForceNoResetLoaders | REN_DoNotDirty | REN_DontCreateRedirectors | REN_NonTransactional);
-			TargetMemory->MarkPendingKill();
 			TargetMemory = nullptr;
 		}
 	};
@@ -855,14 +845,12 @@ void URigVM::ClearMemory()
 	if(WorkMemoryStorageObject)
 	{
 		WorkMemoryStorageObject->Rename(nullptr, GetTransientPackage(), REN_ForceNoResetLoaders | REN_DoNotDirty | REN_DontCreateRedirectors | REN_NonTransactional);
-		WorkMemoryStorageObject->MarkPendingKill();
 		WorkMemoryStorageObject = nullptr;
 	}
 
 	if(DebugMemoryStorageObject)
 	{
 		DebugMemoryStorageObject->Rename(nullptr, GetTransientPackage(), REN_ForceNoResetLoaders | REN_DoNotDirty | REN_DontCreateRedirectors | REN_NonTransactional);
-		DebugMemoryStorageObject->MarkPendingKill();
 		DebugMemoryStorageObject = nullptr;
 	}
 
@@ -1059,92 +1047,33 @@ void URigVM::CacheMemoryHandlesIfRequired(TArrayView<URigVMMemoryStorage*> InMem
 	
 #endif		
 
+#if UE_RIGVM_UCLASS_BASED_STORAGE_DISABLED
+	static_assert(false, "This code is incompatible with the old VM storage");
+#endif
+	
 	FRigVMByteCode& ByteCode = GetByteCode();
 
-	uint16 InstructionIndex = 0;
-	while (Instructions.IsValidIndex(InstructionIndex))
+	auto InstructionOpEval = [&](
+		int32 InstructionIndex,
+		int32 InHandleBaseIndex,
+		TFunctionRef<void(int32 InHandleIndex, const FRigVMOperand& InArg)> InOpFunc
+		) -> void
 	{
-		FirstHandleForInstruction.Add(CachedMemoryHandles.Num());
+		const ERigVMOpCode OpCode = Instructions[InstructionIndex].OpCode; 
 
-		switch (Instructions[InstructionIndex].OpCode)
+		if (OpCode >= ERigVMOpCode::Execute_0_Operands && OpCode <= ERigVMOpCode::Execute_64_Operands)
 		{
-			case ERigVMOpCode::Execute_0_Operands:
-			case ERigVMOpCode::Execute_1_Operands:
-			case ERigVMOpCode::Execute_2_Operands:
-			case ERigVMOpCode::Execute_3_Operands:
-			case ERigVMOpCode::Execute_4_Operands:
-			case ERigVMOpCode::Execute_5_Operands:
-			case ERigVMOpCode::Execute_6_Operands:
-			case ERigVMOpCode::Execute_7_Operands:
-			case ERigVMOpCode::Execute_8_Operands:
-			case ERigVMOpCode::Execute_9_Operands:
-			case ERigVMOpCode::Execute_10_Operands:
-			case ERigVMOpCode::Execute_11_Operands:
-			case ERigVMOpCode::Execute_12_Operands:
-			case ERigVMOpCode::Execute_13_Operands:
-			case ERigVMOpCode::Execute_14_Operands:
-			case ERigVMOpCode::Execute_15_Operands:
-			case ERigVMOpCode::Execute_16_Operands:
-			case ERigVMOpCode::Execute_17_Operands:
-			case ERigVMOpCode::Execute_18_Operands:
-			case ERigVMOpCode::Execute_19_Operands:
-			case ERigVMOpCode::Execute_20_Operands:
-			case ERigVMOpCode::Execute_21_Operands:
-			case ERigVMOpCode::Execute_22_Operands:
-			case ERigVMOpCode::Execute_23_Operands:
-			case ERigVMOpCode::Execute_24_Operands:
-			case ERigVMOpCode::Execute_25_Operands:
-			case ERigVMOpCode::Execute_26_Operands:
-			case ERigVMOpCode::Execute_27_Operands:
-			case ERigVMOpCode::Execute_28_Operands:
-			case ERigVMOpCode::Execute_29_Operands:
-			case ERigVMOpCode::Execute_30_Operands:
-			case ERigVMOpCode::Execute_31_Operands:
-			case ERigVMOpCode::Execute_32_Operands:
-			case ERigVMOpCode::Execute_33_Operands:
-			case ERigVMOpCode::Execute_34_Operands:
-			case ERigVMOpCode::Execute_35_Operands:
-			case ERigVMOpCode::Execute_36_Operands:
-			case ERigVMOpCode::Execute_37_Operands:
-			case ERigVMOpCode::Execute_38_Operands:
-			case ERigVMOpCode::Execute_39_Operands:
-			case ERigVMOpCode::Execute_40_Operands:
-			case ERigVMOpCode::Execute_41_Operands:
-			case ERigVMOpCode::Execute_42_Operands:
-			case ERigVMOpCode::Execute_43_Operands:
-			case ERigVMOpCode::Execute_44_Operands:
-			case ERigVMOpCode::Execute_45_Operands:
-			case ERigVMOpCode::Execute_46_Operands:
-			case ERigVMOpCode::Execute_47_Operands:
-			case ERigVMOpCode::Execute_48_Operands:
-			case ERigVMOpCode::Execute_49_Operands:
-			case ERigVMOpCode::Execute_50_Operands:
-			case ERigVMOpCode::Execute_51_Operands:
-			case ERigVMOpCode::Execute_52_Operands:
-			case ERigVMOpCode::Execute_53_Operands:
-			case ERigVMOpCode::Execute_54_Operands:
-			case ERigVMOpCode::Execute_55_Operands:
-			case ERigVMOpCode::Execute_56_Operands:
-			case ERigVMOpCode::Execute_57_Operands:
-			case ERigVMOpCode::Execute_58_Operands:
-			case ERigVMOpCode::Execute_59_Operands:
-			case ERigVMOpCode::Execute_60_Operands:
-			case ERigVMOpCode::Execute_61_Operands:
-			case ERigVMOpCode::Execute_62_Operands:
-			case ERigVMOpCode::Execute_63_Operands:
-			case ERigVMOpCode::Execute_64_Operands:
-			{
-				const FRigVMExecuteOp& Op = ByteCode.GetOpAt<FRigVMExecuteOp>(Instructions[InstructionIndex]);
-				FRigVMOperandArray Operands = ByteCode.GetOperandsForExecuteOp(Instructions[InstructionIndex]);
-
-				for (const FRigVMOperand& Arg : Operands)
-				{
-					CacheSingleMemoryHandle(Arg, true);
-				}
+			FRigVMOperandArray Operands = ByteCode.GetOperandsForExecuteOp(Instructions[InstructionIndex]);
 
-				InstructionIndex++;
-				break;
+			for (const FRigVMOperand& Arg : Operands)
+			{
+				InOpFunc(InHandleBaseIndex++, Arg);
 			}
+		}
+		else
+		{
+			switch (OpCode)
+			{
 			case ERigVMOpCode::Zero:
 			case ERigVMOpCode::BoolFalse:
 			case ERigVMOpCode::BoolTrue:
@@ -1152,71 +1081,64 @@ void URigVM::CacheMemoryHandlesIfRequired(TArrayView<URigVMMemoryStorage*> InMem
 			case ERigVMOpCode::Decrement:
 			case ERigVMOpCode::ArrayReset:
 			case ERigVMOpCode::ArrayReverse:
-			{
-				const FRigVMUnaryOp& Op = ByteCode.GetOpAt<FRigVMUnaryOp>(Instructions[InstructionIndex]);
-				CacheSingleMemoryHandle(Op.Arg);
-				InstructionIndex++;
-				break;
-			}
+				{
+					const FRigVMUnaryOp& Op = ByteCode.GetOpAt<FRigVMUnaryOp>(Instructions[InstructionIndex]);
+					InOpFunc(InHandleBaseIndex, Op.Arg);
+					break;
+				}
 			case ERigVMOpCode::Copy:
-			{
-				const FRigVMCopyOp& Op = ByteCode.GetOpAt<FRigVMCopyOp>(Instructions[InstructionIndex]);
-				CacheSingleMemoryHandle(Op.Source);
-				CacheSingleMemoryHandle(Op.Target);
+				{
+					const FRigVMCopyOp& Op = ByteCode.GetOpAt<FRigVMCopyOp>(Instructions[InstructionIndex]);
+					InOpFunc(InHandleBaseIndex + 0, Op.Source);
+					InOpFunc(InHandleBaseIndex + 1, Op.Target);
 
 #if UE_RIGVM_UCLASS_BASED_STORAGE_DISABLED
-				if (UScriptStruct* ScriptStruct = GetScriptStructForCopyOp(Op))
-				{
-					CachedMemoryHandles.Add(FRigVMMemoryHandle((uint8*)ScriptStruct));
-				}
+					if (UScriptStruct* ScriptStruct = GetScriptStructForCopyOp(Op))
+					{
+						CachedMemoryHandles.Add(FRigVMMemoryHandle((uint8*)ScriptStruct));
+					}
 #endif
 
-				InstructionIndex++;
-				break;
-			}
+					break;
+				}
 			case ERigVMOpCode::Equals:
 			case ERigVMOpCode::NotEquals:
-			{
-				const FRigVMComparisonOp& Op = ByteCode.GetOpAt<FRigVMComparisonOp>(Instructions[InstructionIndex]);
-				FRigVMOperand Arg = Op.A;
-				CacheSingleMemoryHandle(Arg);
-				Arg = Op.B;
-				CacheSingleMemoryHandle(Arg);
-				Arg = Op.Result;
-				CacheSingleMemoryHandle(Arg);
-				InstructionIndex++;
-				break;
-			}
+				{
+					const FRigVMComparisonOp& Op = ByteCode.GetOpAt<FRigVMComparisonOp>(Instructions[InstructionIndex]);
+					FRigVMOperand Arg = Op.A;
+					InOpFunc(InHandleBaseIndex + 0, Arg);
+					Arg = Op.B;
+					InOpFunc(InHandleBaseIndex + 1, Arg);
+					Arg = Op.Result;
+					InOpFunc(InHandleBaseIndex + 2, Arg);
+					break;
+				}
 			case ERigVMOpCode::JumpAbsolute:
 			case ERigVMOpCode::JumpForward:
 			case ERigVMOpCode::JumpBackward:
-			{
-				InstructionIndex++;
-				break;
-			}
+				{
+					break;
+				}
 			case ERigVMOpCode::JumpAbsoluteIf:
 			case ERigVMOpCode::JumpForwardIf:
 			case ERigVMOpCode::JumpBackwardIf:
-			{
-				const FRigVMJumpIfOp& Op = ByteCode.GetOpAt<FRigVMJumpIfOp>(Instructions[InstructionIndex]);
-				const FRigVMOperand& Arg = Op.Arg;
-				CacheSingleMemoryHandle(Arg);
-				InstructionIndex++;
-				break;
-			}
+				{
+					const FRigVMJumpIfOp& Op = ByteCode.GetOpAt<FRigVMJumpIfOp>(Instructions[InstructionIndex]);
+					const FRigVMOperand& Arg = Op.Arg;
+					InOpFunc(InHandleBaseIndex, Arg);
+					break;
+				}
 			case ERigVMOpCode::ChangeType:
-			{
-				const FRigVMChangeTypeOp& Op = ByteCode.GetOpAt<FRigVMChangeTypeOp>(Instructions[InstructionIndex]);
-				const FRigVMOperand& Arg = Op.Arg;
-				CacheSingleMemoryHandle(Arg);
-				InstructionIndex++;
-				break;
-			}
+				{
+					const FRigVMChangeTypeOp& Op = ByteCode.GetOpAt<FRigVMChangeTypeOp>(Instructions[InstructionIndex]);
+					const FRigVMOperand& Arg = Op.Arg;
+					InOpFunc(InHandleBaseIndex, Arg);
+					break;
+				}
 			case ERigVMOpCode::Exit:
-			{
-				InstructionIndex++;
-				break;
-			}
+				{
+					break;
+				}
 			case ERigVMOpCode::BeginBlock:
 			case ERigVMOpCode::ArrayGetNum:
 			case ERigVMOpCode::ArraySetNum:
@@ -1224,66 +1146,90 @@ void URigVM::CacheMemoryHandlesIfRequired(TArrayView<URigVMMemoryStorage*> InMem
 			case ERigVMOpCode::ArrayClone:
 			case ERigVMOpCode::ArrayRemove:
 			case ERigVMOpCode::ArrayUnion:
-			{
-				const FRigVMBinaryOp& Op = ByteCode.GetOpAt<FRigVMBinaryOp>(Instructions[InstructionIndex]);
-				CacheSingleMemoryHandle(Op.ArgA);
-				CacheSingleMemoryHandle(Op.ArgB);
-				InstructionIndex++;
-				break;
-			}
+				{
+					const FRigVMBinaryOp& Op = ByteCode.GetOpAt<FRigVMBinaryOp>(Instructions[InstructionIndex]);
+					InOpFunc(InHandleBaseIndex + 0, Op.ArgA);
+					InOpFunc(InHandleBaseIndex + 1, Op.ArgB);
+					break;
+				}
 			case ERigVMOpCode::ArrayAdd:
 			case ERigVMOpCode::ArrayGetAtIndex:
 			case ERigVMOpCode::ArraySetAtIndex:
 			case ERigVMOpCode::ArrayInsert:
 			case ERigVMOpCode::ArrayDifference:
 			case ERigVMOpCode::ArrayIntersection:
-			{
-				const FRigVMTernaryOp& Op = ByteCode.GetOpAt<FRigVMTernaryOp>(Instructions[InstructionIndex]);
-				CacheSingleMemoryHandle(Op.ArgA);
-				CacheSingleMemoryHandle(Op.ArgB);
-				CacheSingleMemoryHandle(Op.ArgC);
-				InstructionIndex++;
-				break;
-			}
+				{
+					const FRigVMTernaryOp& Op = ByteCode.GetOpAt<FRigVMTernaryOp>(Instructions[InstructionIndex]);
+					InOpFunc(InHandleBaseIndex + 0, Op.ArgA);
+					InOpFunc(InHandleBaseIndex + 1, Op.ArgB);
+					InOpFunc(InHandleBaseIndex + 2, Op.ArgC);
+					break;
+				}
 			case ERigVMOpCode::ArrayFind:
-			{
-				const FRigVMQuaternaryOp& Op = ByteCode.GetOpAt<FRigVMQuaternaryOp>(Instructions[InstructionIndex]);
-				CacheSingleMemoryHandle(Op.ArgA);
-				CacheSingleMemoryHandle(Op.ArgB);
-				CacheSingleMemoryHandle(Op.ArgC);
-				CacheSingleMemoryHandle(Op.ArgD);
-				InstructionIndex++;
-				break;
-			}
+				{
+					const FRigVMQuaternaryOp& Op = ByteCode.GetOpAt<FRigVMQuaternaryOp>(Instructions[InstructionIndex]);
+					InOpFunc(InHandleBaseIndex + 0, Op.ArgA);
+					InOpFunc(InHandleBaseIndex + 1, Op.ArgB);
+					InOpFunc(InHandleBaseIndex + 2, Op.ArgC);
+					InOpFunc(InHandleBaseIndex + 3, Op.ArgD);
+					break;
+				}
 			case ERigVMOpCode::ArrayIterator:
-			{
-				const FRigVMSenaryOp& Op = ByteCode.GetOpAt<FRigVMSenaryOp>(Instructions[InstructionIndex]);
-				CacheSingleMemoryHandle(Op.ArgA);
-				CacheSingleMemoryHandle(Op.ArgB);
-				CacheSingleMemoryHandle(Op.ArgC);
-				CacheSingleMemoryHandle(Op.ArgD);
-				CacheSingleMemoryHandle(Op.ArgE);
-				CacheSingleMemoryHandle(Op.ArgF);
-				InstructionIndex++;
-				break;
-			}
+				{
+					const FRigVMSenaryOp& Op = ByteCode.GetOpAt<FRigVMSenaryOp>(Instructions[InstructionIndex]);
+					InOpFunc(InHandleBaseIndex + 0, Op.ArgA);
+					InOpFunc(InHandleBaseIndex + 1, Op.ArgB);
+					InOpFunc(InHandleBaseIndex + 2, Op.ArgC);
+					InOpFunc(InHandleBaseIndex + 3, Op.ArgD);
+					InOpFunc(InHandleBaseIndex + 4, Op.ArgE);
+					InOpFunc(InHandleBaseIndex + 5, Op.ArgF);
+					break;
+				}
 			case ERigVMOpCode::EndBlock:
-			{
-				InstructionIndex++;
-				break;
-			}
+				{
+					break;
+				}
 			case ERigVMOpCode::Invalid:
-			{
-				ensure(false);
-				break;
+			default:
+				{
+					checkNoEntry();
+					break;
+				}
 			}
 		}
+	};
+
+	// Make sure we have enough room to prevent repeated allocations.
+	FirstHandleForInstruction.Reset(Instructions.Num() + 1);
+
+	// Count how many handles we need and set up the indirection offsets for the handles.
+	int32 HandleCount = 0;
+	FirstHandleForInstruction.Add(0);
+	for (int32 InstructionIndex = 0; InstructionIndex < Instructions.Num(); InstructionIndex++)
+	{
+		InstructionOpEval(InstructionIndex, INDEX_NONE, [&HandleCount](int32, const FRigVMOperand& ){ HandleCount++; });
+		FirstHandleForInstruction.Add(HandleCount);
 	}
+	
+	// Allocate all the space and zero it out to ensure all pages required for it are paged in
+	// immediately.
+	CachedMemoryHandles.SetNumUninitialized(HandleCount);
 
-	if (FirstHandleForInstruction.Num() < Instructions.Num())
+	// Prefetch the memory types to ensure they exist.
+	for (int32 MemoryType = 0; MemoryType < int32(ERigVMMemoryType::Invalid); MemoryType++)
 	{
-		FirstHandleForInstruction.Add(CachedMemoryHandles.Num());
+		GetMemoryByType(ERigVMMemoryType(MemoryType), /* bCreateIfNeeded = */true);
 	}
+	
+
+	// Now cache the handles as needed.
+	ParallelFor(Instructions.Num(),
+		[&](int32 InstructionIndex)
+		{
+			InstructionOpEval(InstructionIndex, FirstHandleForInstruction[InstructionIndex], [&](int32 InHandleIndex, const FRigVMOperand& InOp){ CacheSingleMemoryHandle(InHandleIndex, InOp); });
+		}
+	);
+
 }
 
 void URigVM::RebuildByteCodeOnLoad()
@@ -1493,7 +1439,7 @@ bool URigVM::ShouldHaltAtInstruction(const FName& InEventName, const uint16 Inst
 #if UE_RIGVM_UCLASS_BASED_STORAGE_DISABLED
 bool URigVM::Initialize(FRigVMMemoryContainerPtrArray Memory, FRigVMFixedArray<void*> AdditionalArguments)
 #else
-bool URigVM::Initialize(TArrayView<URigVMMemoryStorage*> Memory, TArrayView<void*> AdditionalArguments)
+bool URigVM::Initialize(TArrayView<URigVMMemoryStorage*> Memory, TArrayView<void*> AdditionalArguments, bool bInitializeMemory)
 #endif
 {
 	if (ExecutingThreadId != INDEX_NONE)
@@ -1527,6 +1473,12 @@ bool URigVM::Initialize(TArrayView<URigVMMemoryStorage*> Memory, TArrayView<void
 #endif
 
 	CacheMemoryHandlesIfRequired(Memory);
+
+	if(!bInitializeMemory)
+	{
+		return true;
+	}
+	
 	FRigVMByteCode& ByteCode = GetByteCode();
 	TArray<FRigVMFunctionPtr>& Functions = GetFunctions();
 
@@ -3430,7 +3382,7 @@ void URigVM::ClearDebugMemory()
 #endif
 }
 
-void URigVM::CacheSingleMemoryHandle(const FRigVMOperand& InArg, bool bForExecute)
+void URigVM::CacheSingleMemoryHandle(int32 InHandleIndex, const FRigVMOperand& InArg, bool bForExecute)
 {
 #if UE_RIGVM_UCLASS_BASED_STORAGE_DISABLED
 	
@@ -3466,7 +3418,7 @@ void URigVM::CacheSingleMemoryHandle(const FRigVMOperand& InArg, bool bForExecut
 	
 #else
 
-	URigVMMemoryStorage* Memory = GetMemoryByType(InArg.GetMemoryType());
+	URigVMMemoryStorage* Memory = GetMemoryByType(InArg.GetMemoryType(), false);
 
 	if (InArg.GetMemoryType() == ERigVMMemoryType::External)
 	{
@@ -3482,8 +3434,7 @@ void URigVM::CacheSingleMemoryHandle(const FRigVMOperand& InArg, bool bForExecut
 		FRigVMExternalVariable& ExternalVariable = ExternalVariables[InArg.GetRegisterIndex()];
 		check(ExternalVariable.IsValid(false));
 
-		const FRigVMMemoryHandle Handle(ExternalVariable.Memory, ExternalVariable.Property, PropertyPath);
-		CachedMemoryHandles.Add(Handle);
+		CachedMemoryHandles[InHandleIndex] = {ExternalVariable.Memory, ExternalVariable.Property, PropertyPath};
 		return;
 	}
 
@@ -3499,8 +3450,7 @@ void URigVM::CacheSingleMemoryHandle(const FRigVMOperand& InArg, bool bForExecut
 	// is in - so a VM under GetTransientPackage() can be created - but not run.
 	uint8* Data = Memory->GetData<uint8>(InArg.GetRegisterIndex());
 	const FProperty* Property = Memory->GetProperties()[InArg.GetRegisterIndex()];
-	const FRigVMMemoryHandle Handle(Data, Property, PropertyPath);
-	CachedMemoryHandles.Add(Handle);
+	CachedMemoryHandles[InHandleIndex] = {Data, Property, PropertyPath};
 
 #endif
 }
diff --git a/Engine/Source/Runtime/RigVM/Private/RigVMCore/RigVMMemoryStorage.cpp b/Engine/Source/Runtime/RigVM/Private/RigVMCore/RigVMMemoryStorage.cpp
--- a/Engine/Source/Runtime/RigVM/Private/RigVMCore/RigVMMemoryStorage.cpp
+++ b/Engine/Source/Runtime/RigVM/Private/RigVMCore/RigVMMemoryStorage.cpp
@@ -311,7 +311,20 @@ URigVMMemoryStorageGeneratorClass* URigVMMemoryStorageGeneratorClass::CreateStor
 	URigVMMemoryStorageGeneratorClass* OldClass = FindObject<URigVMMemoryStorageGeneratorClass>(Package, *ClassName);
 	if(OldClass)
 	{
-		OldClass->Rename(nullptr, GetTransientPackage(), REN_ForceNoResetLoaders | REN_DoNotDirty | REN_DontCreateRedirectors | REN_NonTransactional);
+		FString DiscardedMemoryClassName;
+		static const TCHAR DiscardedMemoryClassTemplate[] = TEXT("DiscardedMemoryClassTemplate_%d");
+		static int32 DiscardedMemoryClassIndex = 0;
+		do
+		{
+			DiscardedMemoryClassName = FString::Printf(DiscardedMemoryClassTemplate, DiscardedMemoryClassIndex++);
+			if(StaticFindObjectFast(nullptr, GetTransientPackage(), *DiscardedMemoryClassName) == nullptr)
+			{
+				break;
+			}
+		}
+		while (DiscardedMemoryClassIndex < INT_MAX);
+
+		OldClass->Rename(*DiscardedMemoryClassName, GetTransientPackage(), REN_ForceNoResetLoaders | REN_DoNotDirty | REN_DontCreateRedirectors | REN_NonTransactional);
 	}
 
 	// create the new class
@@ -865,6 +878,24 @@ bool URigVMMemoryStorage::CopyProperty(
 
 #endif
 
+void URigVMMemoryStorage::CopyFrom(URigVMMemoryStorage* InSourceMemory)
+{
+	check(InSourceMemory);
+	check(GetClass() == InSourceMemory->GetClass());
+
+	for(int32 PropertyIndex = 0; PropertyIndex < Num(); PropertyIndex++)
+	{
+		URigVMMemoryStorage::CopyProperty(
+			this,
+			PropertyIndex,
+			FRigVMPropertyPath::Empty,
+			InSourceMemory,
+			PropertyIndex,
+			FRigVMPropertyPath::Empty
+		);
+	}
+}
+
 const TArray<const FProperty*>& URigVMMemoryStorage::GetProperties() const
 {
 	if(const URigVMMemoryStorageGeneratorClass* Class = Cast<URigVMMemoryStorageGeneratorClass>(GetClass()))
diff --git a/Engine/Source/Runtime/RigVM/Public/RigVMCore/RigVM.h b/Engine/Source/Runtime/RigVM/Public/RigVMCore/RigVM.h
--- a/Engine/Source/Runtime/RigVM/Public/RigVMCore/RigVM.h
+++ b/Engine/Source/Runtime/RigVM/Public/RigVMCore/RigVM.h
@@ -171,7 +171,7 @@ public:
 #if UE_RIGVM_UCLASS_BASED_STORAGE_DISABLED
 	bool Initialize(FRigVMMemoryContainerPtrArray Memory, FRigVMFixedArray<void*> AdditionalArguments);
 #else
-	bool Initialize(TArrayView<URigVMMemoryStorage*> Memory, TArrayView<void*> AdditionalArguments);
+	bool Initialize(TArrayView<URigVMMemoryStorage*> Memory, TArrayView<void*> AdditionalArguments, bool bInitializeMemory = true);
 #endif
 
 	// Executes the VM.
@@ -747,11 +747,12 @@ public:
 	uint32 GetNumExecutions() const { return NumExecutions; }
 	const FRigVMExecuteContext& GetContext() const { return Context; }
 
+	void InvalidateCachedMemory();
+
 private:
 
 	void ResolveFunctionsIfRequired();
 	void RefreshInstructionsIfRequired();
-	void InvalidateCachedMemory();
 #if UE_RIGVM_UCLASS_BASED_STORAGE_DISABLED
 	void CacheMemoryHandlesIfRequired(FRigVMMemoryContainerPtrArray InMemory);
 #else
@@ -834,7 +835,7 @@ private:
 	// debug watch register memory needs to be cleared for each execution
 	void ClearDebugMemory();
 	
-	void CacheSingleMemoryHandle(const FRigVMOperand& InArg, bool bForExecute = false);
+	void CacheSingleMemoryHandle(int32 InHandleIndex, const FRigVMOperand& InArg, bool bForExecute = false);
 
 	FORCEINLINE_DEBUGGABLE void CopyOperandForDebuggingIfNeeded(const FRigVMOperand& InArg, const FRigVMMemoryHandle& InHandle)
 	{
diff --git a/Engine/Source/Runtime/RigVM/Public/RigVMCore/RigVMExternalVariable.h b/Engine/Source/Runtime/RigVM/Public/RigVMCore/RigVMExternalVariable.h
--- a/Engine/Source/Runtime/RigVM/Public/RigVMCore/RigVMExternalVariable.h
+++ b/Engine/Source/Runtime/RigVM/Public/RigVMCore/RigVMExternalVariable.h
@@ -5,6 +5,7 @@
 #include "CoreMinimal.h"
 #include "RigVMTraits.h"
 #include "RigVMMemory.h"
+#include "UObject/UnrealType.h"
 
 /**
  * The external variable can be used to map external / unowned
@@ -84,6 +85,62 @@ struct RIGVM_API FRigVMExternalVariable
 		}
 	}
 
+	static uint32 GetPropertyTypeHash(const FProperty *InProperty, const uint8 *InMemory)
+	{
+		if (!ensure(InProperty))
+		{
+			return 0;
+		}
+
+		if (InProperty->PropertyFlags & CPF_HasGetValueTypeHash)
+		{
+			return InProperty->GetValueTypeHash(InMemory);
+		}
+
+		if (const FBoolProperty* BoolProperty = CastField<FBoolProperty>(InProperty))
+		{
+			return static_cast<uint32>(BoolProperty->GetPropertyValue(InMemory));
+		}
+		else if (const FArrayProperty *ArrayProperty = CastField<FArrayProperty>(InProperty))
+		{
+			FScriptArrayHelper ArrayHelper(ArrayProperty, InMemory);
+			int32 Hash = ::GetTypeHash(ArrayHelper.Num());
+			for (int32 Index = 0; Index < ArrayHelper.Num(); Index++)
+			{
+				Hash = HashCombine(Hash, GetPropertyTypeHash(ArrayProperty->Inner, ArrayHelper.GetRawPtr(Index)));
+			}
+			return Hash;
+		}
+		else if (const FStructProperty *StructProperty = CastField<FStructProperty>(InProperty))
+		{
+			const UScriptStruct* StructType = StructProperty->Struct;
+			if (StructType->GetCppStructOps() && StructType->GetCppStructOps()->HasGetTypeHash())
+			{
+				return StructType->GetCppStructOps()->GetStructTypeHash(InMemory);
+			}
+			else
+			{
+				uint32 Hash = 0;
+				for (TFieldIterator<FProperty> It(StructType); It; ++It)
+				{
+					const FProperty* SubProperty = *It;
+					Hash = HashCombine(Hash, GetPropertyTypeHash(*It, InMemory + SubProperty->GetOffset_ForInternal()));
+				}
+				return Hash;
+			}
+		}
+
+		// If we get here, we're missing support for a property type that doesn't do its own hashing. 
+		checkNoEntry();
+		return 0;
+	}
+
+	uint32 GetTypeHash() const
+	{
+		return GetPropertyTypeHash(Property, Memory);
+	}
+	
+
 	FORCEINLINE static FRigVMExternalVariable Make(const FProperty* InProperty, void* InContainer, const FName& InOptionalName = NAME_None)
 	{
 		check(InProperty);
diff --git a/Engine/Source/Runtime/RigVM/Public/RigVMCore/RigVMMemoryStorage.h b/Engine/Source/Runtime/RigVM/Public/RigVMCore/RigVMMemoryStorage.h
--- a/Engine/Source/Runtime/RigVM/Public/RigVMCore/RigVMMemoryStorage.h
+++ b/Engine/Source/Runtime/RigVM/Public/RigVMCore/RigVMMemoryStorage.h
@@ -599,6 +599,8 @@ public:
 	
 #endif
 
+	void CopyFrom(URigVMMemoryStorage* InSourceMemory);
+
 private:
 	
 	static const TArray<const FProperty*> EmptyProperties;
