Control Rig: Add hash for Map and Set property types.

#rnx

diff --git a/Engine/Source/Runtime/RigVM/Public/RigVMCore/RigVMExternalVariable.h b/Engine/Source/Runtime/RigVM/Public/RigVMCore/RigVMExternalVariable.h
--- a/Engine/Source/Runtime/RigVM/Public/RigVMCore/RigVMExternalVariable.h
+++ b/Engine/Source/Runtime/RigVM/Public/RigVMCore/RigVMExternalVariable.h
@@ -111,6 +111,27 @@ struct RIGVM_API FRigVMExternalVariable
 			}
 			return Hash;
 		}
+		else if (const FMapProperty *MapProperty = CastField<FMapProperty>(InProperty))
+		{
+			FScriptMapHelper MapHelper(MapProperty, InMemory);
+			int32 Hash = ::GetTypeHash(MapHelper.Num());
+			for (int32 Index = 0; Index < MapHelper.Num(); Index++)
+			{
+				Hash = HashCombine(Hash, GetPropertyTypeHash(MapProperty->KeyProp, MapHelper.GetKeyPtr(Index)));
+				Hash = HashCombine(Hash, GetPropertyTypeHash(MapProperty->ValueProp, MapHelper.GetValuePtr(Index)));
+			}
+			return Hash;
+		}
+		else if (const FSetProperty *SetProperty = CastField<FSetProperty>(InProperty))
+		{
+			FScriptSetHelper SetHelper(SetProperty, InMemory);
+			int32 Hash = ::GetTypeHash(SetHelper.Num());
+			for (int32 Index = 0; Index < SetHelper.Num(); Index++)
+			{
+				Hash = HashCombine(Hash, GetPropertyTypeHash(SetProperty->ElementProp, SetHelper.GetElementPtr(Index)));
+			}
+			return Hash;
+		}
 		else if (const FStructProperty *StructProperty = CastField<FStructProperty>(InProperty))
 		{
 			const UScriptStruct* StructType = StructProperty->Struct;
