Control Rig: make sure to clear object flags when clearing VM memory to avoid unwanted objects getting cooked

diff --git a/Engine/Source/Runtime/RigVM/Private/RigVMCore/RigVM.cpp b/Engine/Source/Runtime/RigVM/Private/RigVMCore/RigVM.cpp
--- a/Engine/Source/Runtime/RigVM/Private/RigVMCore/RigVM.cpp
+++ b/Engine/Source/Runtime/RigVM/Private/RigVMCore/RigVM.cpp
@@ -6,7 +6,7 @@
 #include "UObject/UE5MainStreamObjectVersion.h"
 #include "HAL/PlatformTLS.h"
 #include "Async/ParallelFor.h"
-
+#include "UObject/UObjectHash.h"
 
 void FRigVMParameter::Serialize(FArchive& Ar)
 {
@@ -770,6 +770,9 @@ URigVMMemoryStorage* URigVM::GetMemoryByType(ERigVMMemoryType InMemoryType, bool
 					}
 					else
 					{
+						// since literal memory object can be shared across packages, it needs to have the RF_Public flag
+						// for example, a control rig instance in a level sequence pacakge can references
+						// the literal memory object in the control rig package
 						LiteralMemoryStorageObject = NewObject<URigVMMemoryStorage>(this, FName(), RF_Public);
 					}
 				}
@@ -791,11 +794,11 @@ URigVMMemoryStorage* URigVM::GetMemoryByType(ERigVMMemoryType InMemoryType, bool
 				{
 					if(UClass* Class = URigVMMemoryStorageGeneratorClass::GetStorageClass(this, InMemoryType))
 					{
-						WorkMemoryStorageObject = NewObject<URigVMMemoryStorage>(this, Class, FName(), RF_Public);
+						WorkMemoryStorageObject = NewObject<URigVMMemoryStorage>(this, Class);
 					}
 					else
 					{
-						WorkMemoryStorageObject = NewObject<URigVMMemoryStorage>(this, FName(), RF_Public);
+						WorkMemoryStorageObject = NewObject<URigVMMemoryStorage>(this);
 					}
 				}
 			}
@@ -818,12 +821,12 @@ URigVMMemoryStorage* URigVM::GetMemoryByType(ERigVMMemoryType InMemoryType, bool
 #if WITH_EDITOR
 					if(UClass* Class = URigVMMemoryStorageGeneratorClass::GetStorageClass(this, InMemoryType))
 					{
-						DebugMemoryStorageObject = NewObject<URigVMMemoryStorage>(this, Class, FName(), RF_Public);
+						DebugMemoryStorageObject = NewObject<URigVMMemoryStorage>(this, Class);
 					}
 					else
 #endif
 					{
-						DebugMemoryStorageObject = NewObject<URigVMMemoryStorage>(this, FName(), RF_Public);
+						DebugMemoryStorageObject = NewObject<URigVMMemoryStorage>(this);
 					}
 				}
 			}
@@ -840,6 +843,31 @@ URigVMMemoryStorage* URigVM::GetMemoryByType(ERigVMMemoryType InMemoryType, bool
 
 void URigVM::ClearMemory()
 {
+	// At one point our memory objects were saved with RF_Public,
+	// so to truly clear them, we have to also clear the flags
+	// RF_Public will make them stay around as zombie unreferenced objects, and get included in SavePackage and cooking.
+	// Clear their flags so they are not included by editor or cook SavePackage calls.
+
+	// we now make sure that only the literal memory object on the CDO is marked as RF_Public
+	// and work memory objects are no longer marked as RF_Public
+	
+	TArray<UObject*> SubObjects;
+	GetObjectsWithOuter(this, SubObjects);
+	for (UObject* SubObject : SubObjects)
+	{
+		if (URigVMMemoryStorage* MemoryObject = Cast<URigVMMemoryStorage>(SubObject))
+		{
+			// we don't care about memory type here because
+			// 
+			// if "this" is not CDO, its subobjects will not include the literal memory and
+			// thus only clears the flag for work mem
+			// 
+			// if "this" is CDO, its subobjects will include the literal memory and this allows
+			// us to actually clear the literal memory
+			MemoryObject->ClearFlags(RF_Public);
+		}
+	}
+
 	LiteralMemoryStorageObject = nullptr;
 
 	if(WorkMemoryStorageObject)
