Control Rig: Fixed bad run node counts

diff --git a/Engine/Source/Runtime/RigVM/Private/RigVMCore/RigVMByteCode.cpp b/Engine/Source/Runtime/RigVM/Private/RigVMCore/RigVMByteCode.cpp
--- a/Engine/Source/Runtime/RigVM/Private/RigVMCore/RigVMByteCode.cpp
+++ b/Engine/Source/Runtime/RigVM/Private/RigVMCore/RigVMByteCode.cpp
@@ -1774,14 +1774,14 @@ void FRigVMByteCode::SetSubject(int32 InInstructionIndex, const FString& InCallP
 		SubjectPerInstruction.AddZeroed(1 + InInstructionIndex - SubjectPerInstruction.Num());
 	}
 	SubjectPerInstruction[InInstructionIndex] = Subject;
-	SubjectToInstructions.FindOrAdd(Subject).Add(InInstructionIndex);
+	SubjectToInstructions.FindOrAdd(Subject).AddUnique(InInstructionIndex);
 
 	if (CallPathPerInstruction.Num() <= InInstructionIndex)
 	{
 		CallPathPerInstruction.AddZeroed(1 + InInstructionIndex - CallPathPerInstruction.Num());
 	}
 	CallPathPerInstruction[InInstructionIndex] = InCallPath;
-	CallPathToInstructions.FindOrAdd(InCallPath).Add(InInstructionIndex);
+	CallPathToInstructions.FindOrAdd(InCallPath).AddUnique(InInstructionIndex);
 
 	if (CallstackPerInstruction.Num() <= InInstructionIndex)
 	{
@@ -1800,7 +1800,7 @@ void FRigVMByteCode::SetSubject(int32 InInstructionIndex, const FString& InCallP
 		UObject* const* DataPtr = &InCallstack[InCallstack.Num() - CallstackLength];
 		TArrayView<UObject* const> View(DataPtr, CallstackLength);
 		uint32 Hash = GetCallstackHash(View);
-		CallstackHashToInstructions.FindOrAdd(Hash).Add(InInstructionIndex);
+		CallstackHashToInstructions.FindOrAdd(Hash).AddUnique(InInstructionIndex);
 	}
 }
 
