RigVM: Deal with invalid properties from previous versions

diff --git a/Engine/Source/Runtime/RigVM/Private/RigVMCore/RigVMMemoryDeprecated.cpp b/Engine/Source/Runtime/RigVM/Private/RigVMCore/RigVMMemoryDeprecated.cpp
--- a/Engine/Source/Runtime/RigVM/Private/RigVMCore/RigVMMemoryDeprecated.cpp
+++ b/Engine/Source/Runtime/RigVM/Private/RigVMCore/RigVMMemoryDeprecated.cpp
@@ -149,6 +149,13 @@ void FRigVMRegisterOffset::Load(FArchive& Ar)
 			// so segments also need to be recalculated
 			int32 InitialOffset = ArrayIndex * ParentScriptStruct->GetStructureSize();
 			FRigVMRegisterOffset TempOffset(ParentScriptStruct, SegmentPath, InitialOffset, ElementSize);
+			if(TempOffset.Type == ERigVMRegisterType::Invalid)
+			{
+				Type = ERigVMRegisterType::Invalid;
+				CachedSegmentPath = SegmentPath;
+				return;
+			}
+			
 			if (TempOffset.GetSegments().Num() == Segments.Num())
 			{
 				Segments = TempOffset.GetSegments();
@@ -185,7 +192,7 @@ FRigVMRegisterOffset::FRigVMRegisterOffset(UScriptStruct* InScriptStruct, const
 
 	struct FRigVMRegisterOffsetBuilder
 	{
-		static void WalkStruct(UStruct* InStruct, const FString& InPath, FRigVMRegisterOffset& Offset)
+		static bool WalkStruct(UStruct* InStruct, const FString& InPath, FRigVMRegisterOffset& Offset)
 		{
 			FString Left, Right;
 			if (!InPath.Split(TEXT("."), &Left, &Right))
@@ -194,7 +201,10 @@ FRigVMRegisterOffset::FRigVMRegisterOffset(UScriptStruct* InScriptStruct, const
 			}
 
 			FProperty* Property = InStruct->FindPropertyByName(*Left);
-			check(Property)
+			if(Property == nullptr)
+			{
+				return false;
+			}
 
 			int32 SegmentIndex = Property->GetOffset_ReplaceWith_ContainerPtrToValuePtr();
 			if (Offset.Segments.Num() > 0)
@@ -217,11 +227,11 @@ FRigVMRegisterOffset::FRigVMRegisterOffset(UScriptStruct* InScriptStruct, const
 			{
 				if (FStructProperty* StructProperty = CastField<FStructProperty>(Property))
 				{
-					WalkStruct(StructProperty->Struct, Right, Offset);
+					return WalkStruct(StructProperty->Struct, Right, Offset);
 				}
 				else if (FArrayProperty* ArrayProperty = CastField<FArrayProperty>(Property))
 				{
-					WalkArray(ArrayProperty, Right, Offset);
+					return WalkArray(ArrayProperty, Right, Offset);
 				}
 			}
 			else
@@ -248,9 +258,11 @@ FRigVMRegisterOffset::FRigVMRegisterOffset(UScriptStruct* InScriptStruct, const
 					Offset.Type = ERigVMRegisterType::Plain;
 				}
 			}
+
+			return true;
 		}
 
-		static void WalkArray(FArrayProperty* InArrayProperty, const FString& InPath, FRigVMRegisterOffset& Offset)
+		static bool WalkArray(FArrayProperty* InArrayProperty, const FString& InPath, FRigVMRegisterOffset& Offset)
 		{
 			FString Left, Right;
 			if (!InPath.Split(TEXT("."), &Left, &Right))
@@ -281,11 +293,11 @@ FRigVMRegisterOffset::FRigVMRegisterOffset(UScriptStruct* InScriptStruct, const
 			{
 				if (FStructProperty* StructProperty = CastField<FStructProperty>(InArrayProperty->Inner))
 				{
-					WalkStruct(StructProperty->Struct, Right, Offset);
+					return WalkStruct(StructProperty->Struct, Right, Offset);
 				}
 				else if (FArrayProperty* ArrayProperty = CastField<FArrayProperty>(InArrayProperty->Inner))
 				{
-					WalkArray(ArrayProperty, Right, Offset);
+					return WalkArray(ArrayProperty, Right, Offset);
 				}
 			}
 			else
@@ -310,6 +322,7 @@ FRigVMRegisterOffset::FRigVMRegisterOffset(UScriptStruct* InScriptStruct, const
 				}
 			}
 
+			return true;
 		}
 	};
 
@@ -323,7 +336,12 @@ FRigVMRegisterOffset::FRigVMRegisterOffset(UScriptStruct* InScriptStruct, const
 		FString SegmentPath = InSegmentPath;
 		SegmentPath = SegmentPath.Replace(TEXT("["), TEXT("."));
 		SegmentPath = SegmentPath.Replace(TEXT("]"), TEXT("."));
-		FRigVMRegisterOffsetBuilder::WalkStruct(InScriptStruct, SegmentPath, *this);
+		if(!FRigVMRegisterOffsetBuilder::WalkStruct(InScriptStruct, SegmentPath, *this))
+		{
+			Type = ERigVMRegisterType::Invalid;
+			return;
+		}
+		
 		if (Type == ERigVMRegisterType::Plain)
 		{
 			if (CPPType == TEXT("FName"))
@@ -770,6 +788,16 @@ void FRigVMMemoryContainer::Load(FArchive& Ar)
 
 	bEncounteredErrorDuringLoad = false;
 
+	for(const FRigVMRegisterOffset& RegisterOffset : RegisterOffsets)
+	{
+		if(RegisterOffset.GetType() == ERigVMRegisterType::Invalid)
+		{
+			FString PackagePath = Ar.GetArchiveName();
+			UE_LOG(LogRigVM, Error, TEXT("RegisterOffset '%s' cannot be found. Asset '%s' no longer functional."), *RegisterOffset.GetCachedSegmentPath(), *PackagePath);
+			bEncounteredErrorDuringLoad = true;
+		}
+	}
+
 	ScriptStructs.Reset();
 	TArray<FString> ScriptStructPaths;
 	Ar << ScriptStructPaths;
diff --git a/Engine/Source/Runtime/RigVM/Public/RigVMCore/RigVMExternalVariable.h b/Engine/Source/Runtime/RigVM/Public/RigVMCore/RigVMExternalVariable.h
--- a/Engine/Source/Runtime/RigVM/Public/RigVMCore/RigVMExternalVariable.h
+++ b/Engine/Source/Runtime/RigVM/Public/RigVMCore/RigVMExternalVariable.h
@@ -80,7 +80,9 @@ struct RIGVM_API FRigVMExternalVariable
 		}
 		else if (const FObjectProperty* ObjectProperty = CastField<FObjectProperty>(InProperty))
 		{
-			OutTypeName = *ObjectProperty->PropertyClass->GetName();
+			OutTypeName = *FString::Printf(TEXT("TObjectPtr<%s%s>"),
+				ObjectProperty->PropertyClass->GetPrefixCPP(),
+				*ObjectProperty->PropertyClass->GetName());
 			OutTypeObject = ObjectProperty->PropertyClass;
 		}
 	}
