Control Rig: Copy byte code works when array nodes are used

diff --git a/Engine/Source/Runtime/RigVM/Private/RigVMCore/RigVM.cpp b/Engine/Source/Runtime/RigVM/Private/RigVMCore/RigVM.cpp
--- a/Engine/Source/Runtime/RigVM/Private/RigVMCore/RigVM.cpp
+++ b/Engine/Source/Runtime/RigVM/Private/RigVMCore/RigVM.cpp
@@ -3277,6 +3277,120 @@ TArray<FString> URigVM::DumpByteCodeAsTextArray(const TArray<int32>& InInstructi
 				ResultLine = TEXT("End Block");
 				break;
 			}
+			case ERigVMOpCode::ArrayReset:
+			{
+				const FRigVMUnaryOp& Op = ByteCode.GetOpAt<FRigVMUnaryOp>(Instructions[InstructionIndex]);
+				ResultLine = FString::Printf(TEXT("Reset array %s"), *GetOperandLabel(Op.Arg, OperandFormatFunction));
+				break;
+			}
+			case ERigVMOpCode::ArrayGetNum:
+			{
+				const FRigVMBinaryOp& Op = ByteCode.GetOpAt<FRigVMBinaryOp>(Instructions[InstructionIndex]);
+				ResultLine = FString::Printf(TEXT("Get size of array %s and assign to %s"), *GetOperandLabel(Op.ArgA, OperandFormatFunction), *GetOperandLabel(Op.ArgB, OperandFormatFunction));
+				break;
+			} 
+			case ERigVMOpCode::ArraySetNum:
+			{
+				const FRigVMBinaryOp& Op = ByteCode.GetOpAt<FRigVMBinaryOp>(Instructions[InstructionIndex]);
+				ResultLine = FString::Printf(TEXT("Set size of array %s to %s"), *GetOperandLabel(Op.ArgA, OperandFormatFunction), *GetOperandLabel(Op.ArgB, OperandFormatFunction));
+				break;
+			}
+			case ERigVMOpCode::ArrayGetAtIndex:
+			{
+				const FRigVMTernaryOp& Op = ByteCode.GetOpAt<FRigVMTernaryOp>(Instructions[InstructionIndex]);
+				ResultLine = FString::Printf(TEXT("Get item of array %s at index %s and assign to %s"), *GetOperandLabel(Op.ArgA, OperandFormatFunction), *GetOperandLabel(Op.ArgB, OperandFormatFunction), *GetOperandLabel(Op.ArgC, OperandFormatFunction));
+				break;
+			}
+			case ERigVMOpCode::ArraySetAtIndex:
+			{
+				const FRigVMTernaryOp& Op = ByteCode.GetOpAt<FRigVMTernaryOp>(Instructions[InstructionIndex]);
+				ResultLine = FString::Printf(TEXT("Set item of array %s at index %s to %s"), *GetOperandLabel(Op.ArgA, OperandFormatFunction), *GetOperandLabel(Op.ArgB, OperandFormatFunction), *GetOperandLabel(Op.ArgC, OperandFormatFunction));
+				break;
+			}
+			case ERigVMOpCode::ArrayAdd:
+			{
+				const FRigVMTernaryOp& Op = ByteCode.GetOpAt<FRigVMTernaryOp>(Instructions[InstructionIndex]);
+				ResultLine = FString::Printf(TEXT("Add element %s to array %s and return index %s"),
+					*GetOperandLabel(Op.ArgB, OperandFormatFunction),
+					*GetOperandLabel(Op.ArgA, OperandFormatFunction),
+					*GetOperandLabel(Op.ArgC, OperandFormatFunction));
+				break;
+			}
+			case ERigVMOpCode::ArrayInsert:
+			{
+				const FRigVMTernaryOp& Op = ByteCode.GetOpAt<FRigVMTernaryOp>(Instructions[InstructionIndex]);
+				ResultLine = FString::Printf(TEXT("Insert element %s to array %s at index %s"), *GetOperandLabel(Op.ArgC, OperandFormatFunction), *GetOperandLabel(Op.ArgA, OperandFormatFunction), *GetOperandLabel(Op.ArgB, OperandFormatFunction));
+				break;
+			}
+			case ERigVMOpCode::ArrayRemove:
+			{
+				const FRigVMBinaryOp& Op = ByteCode.GetOpAt<FRigVMBinaryOp>(Instructions[InstructionIndex]);
+				ResultLine = FString::Printf(TEXT("Remove element at index %s from array %s"), *GetOperandLabel(Op.ArgB, OperandFormatFunction), *GetOperandLabel(Op.ArgA, OperandFormatFunction));
+				break;
+			}
+			case ERigVMOpCode::ArrayFind:
+			{
+				const FRigVMQuaternaryOp& Op = ByteCode.GetOpAt<FRigVMQuaternaryOp>(Instructions[InstructionIndex]);
+				ResultLine = FString::Printf(TEXT("Find element %s in array %s and returns index %s and if element was found %s"),
+					*GetOperandLabel(Op.ArgB, OperandFormatFunction),
+					*GetOperandLabel(Op.ArgA, OperandFormatFunction),
+					*GetOperandLabel(Op.ArgC, OperandFormatFunction),
+					*GetOperandLabel(Op.ArgD, OperandFormatFunction));
+				break;
+			}
+			case ERigVMOpCode::ArrayAppend:
+			{
+				const FRigVMBinaryOp& Op = ByteCode.GetOpAt<FRigVMBinaryOp>(Instructions[InstructionIndex]);
+				ResultLine = FString::Printf(TEXT("Append array %s to array %s"), *GetOperandLabel(Op.ArgB, OperandFormatFunction), *GetOperandLabel(Op.ArgA, OperandFormatFunction));
+				break;
+			}
+			case ERigVMOpCode::ArrayClone:
+			{
+				const FRigVMBinaryOp& Op = ByteCode.GetOpAt<FRigVMBinaryOp>(Instructions[InstructionIndex]);
+				ResultLine = FString::Printf(TEXT("Clone array %s to array %s"), *GetOperandLabel(Op.ArgA, OperandFormatFunction), *GetOperandLabel(Op.ArgB, OperandFormatFunction));
+				break;
+			}
+			case ERigVMOpCode::ArrayIterator:
+			{
+				const FRigVMSenaryOp& Op = ByteCode.GetOpAt<FRigVMSenaryOp>(Instructions[InstructionIndex]);
+				ResultLine = FString::Printf(TEXT("Iterate over array %s, with current element in %s, current index in %s, array count in %s and current ratio in %s"),
+					*GetOperandLabel(Op.ArgA, OperandFormatFunction),
+					*GetOperandLabel(Op.ArgB, OperandFormatFunction),
+					*GetOperandLabel(Op.ArgC, OperandFormatFunction),
+					*GetOperandLabel(Op.ArgD, OperandFormatFunction),
+					*GetOperandLabel(Op.ArgE, OperandFormatFunction));
+				break;
+			}
+			case ERigVMOpCode::ArrayUnion:
+			{
+				const FRigVMBinaryOp& Op = ByteCode.GetOpAt<FRigVMBinaryOp>(Instructions[InstructionIndex]);
+				ResultLine = FString::Printf(TEXT("Merge array %s and array %s"), *GetOperandLabel(Op.ArgA, OperandFormatFunction), *GetOperandLabel(Op.ArgB, OperandFormatFunction));
+				break;
+			}
+			case ERigVMOpCode::ArrayDifference:
+			{
+				const FRigVMTernaryOp& Op = ByteCode.GetOpAt<FRigVMTernaryOp>(Instructions[InstructionIndex]);
+				ResultLine = FString::Printf(TEXT("Create array %s from differences of array %s and array %s"),
+					*GetOperandLabel(Op.ArgC, OperandFormatFunction),
+					*GetOperandLabel(Op.ArgA, OperandFormatFunction),
+					*GetOperandLabel(Op.ArgB, OperandFormatFunction));
+				break;
+			}
+			case ERigVMOpCode::ArrayIntersection:
+			{
+				const FRigVMTernaryOp& Op = ByteCode.GetOpAt<FRigVMTernaryOp>(Instructions[InstructionIndex]);
+				ResultLine = FString::Printf(TEXT("Create array %s from intersection of array %s and array %s"),
+					*GetOperandLabel(Op.ArgC, OperandFormatFunction), 
+					*GetOperandLabel(Op.ArgA, OperandFormatFunction),
+					*GetOperandLabel(Op.ArgB, OperandFormatFunction));
+				break;
+			}
+			case ERigVMOpCode::ArrayReverse:
+			{
+				const FRigVMUnaryOp& Op = ByteCode.GetOpAt<FRigVMUnaryOp>(Instructions[InstructionIndex]);
+				ResultLine = FString::Printf(TEXT("Reverse array %s"), *GetOperandLabel(Op.Arg, OperandFormatFunction));
+				break;
+			}
 			default:
 			{
 				ensure(false);
