Control Rig: Fixed transform value not being printed correctly in watch string when it equals the default value

diff --git a/Engine/Plugins/Animation/ControlRig/Source/ControlRigEditor/Private/Graph/SControlRigGraphNode.cpp b/Engine/Plugins/Animation/ControlRig/Source/ControlRigEditor/Private/Graph/SControlRigGraphNode.cpp
--- a/Engine/Plugins/Animation/ControlRig/Source/ControlRigEditor/Private/Graph/SControlRigGraphNode.cpp
+++ b/Engine/Plugins/Animation/ControlRig/Source/ControlRigEditor/Private/Graph/SControlRigGraphNode.cpp
@@ -964,7 +964,8 @@ void SControlRigGraphNode::GetNodeInfoPopups(FNodeInfoContext* Context, TArray<F
 								DefaultValues = Memory.GetRegisterValueAsString(*WatchOperand, ModelPin->GetCPPType(), ModelPin->GetCPPTypeObject());
 #else
 								URigVMMemoryStorage* Memory = ActiveObject->GetVM()->GetDebugMemory();
-								const FString DebugValue = Memory->GetDataAsStringSafe(WatchOperand->GetRegisterIndex());
+								// We mark PPF_ExternalEditor so that default values are also printed
+								const FString DebugValue = Memory->GetDataAsStringSafe(WatchOperand->GetRegisterIndex(), PPF_ExternalEditor);
 								if(!DebugValue.IsEmpty())
 								{
 									DefaultValues = URigVMPin::SplitDefaultValue(DebugValue);
diff --git a/Engine/Source/Runtime/RigVM/Private/RigVMCore/RigVMMemoryStorage.cpp b/Engine/Source/Runtime/RigVM/Private/RigVMCore/RigVMMemoryStorage.cpp
--- a/Engine/Source/Runtime/RigVM/Private/RigVMCore/RigVMMemoryStorage.cpp
+++ b/Engine/Source/Runtime/RigVM/Private/RigVMCore/RigVMMemoryStorage.cpp
@@ -641,36 +641,36 @@ void URigVMMemoryStorageGeneratorClass::RefreshPropertyPaths()
 const TArray<const FProperty*> URigVMMemoryStorage::EmptyProperties;
 const TArray<FRigVMPropertyPath> URigVMMemoryStorage::EmptyPropertyPaths;
 
-FString URigVMMemoryStorage::GetDataAsString(int32 InPropertyIndex)
+FString URigVMMemoryStorage::GetDataAsString(int32 InPropertyIndex, int32 PortFlags)
 {
 	check(IsValidIndex(InPropertyIndex));
 	const uint8* Data = GetData<uint8>(InPropertyIndex);
 
 	FString Value;
-	GetProperties()[InPropertyIndex]->ExportTextItem(Value, Data, nullptr, nullptr, PPF_None);
+	GetProperties()[InPropertyIndex]->ExportTextItem(Value, Data, nullptr, nullptr, PortFlags);
 	return Value;
 }
 
-FString URigVMMemoryStorage::GetDataAsString(const FRigVMOperand& InOperand)
+FString URigVMMemoryStorage::GetDataAsString(const FRigVMOperand& InOperand, int32 PortFlags)
 {
 	const int32 PropertyIndex = InOperand.GetRegisterIndex();
 	check(IsValidIndex(PropertyIndex));
-	return GetDataAsString(PropertyIndex);
+	return GetDataAsString(PropertyIndex, PortFlags);
 }
 
-FString URigVMMemoryStorage::GetDataAsStringSafe(int32 InPropertyIndex)
+FString URigVMMemoryStorage::GetDataAsStringSafe(int32 InPropertyIndex, int32 PortFlags)
 {
 	if(!IsValidIndex(InPropertyIndex))
 	{
 		return FString();
 	}
-	return GetDataAsString(InPropertyIndex);
+	return GetDataAsString(InPropertyIndex, PortFlags);
 }
 
-FString URigVMMemoryStorage::GetDataAsStringSafe(const FRigVMOperand& InOperand)
+FString URigVMMemoryStorage::GetDataAsStringSafe(const FRigVMOperand& InOperand, int32 PortFlags)
 {
 	const int32 PropertyIndex = InOperand.GetRegisterIndex();
-	return GetDataAsStringSafe(PropertyIndex);
+	return GetDataAsStringSafe(PropertyIndex, PortFlags);
 }
 
 class FRigVMMemoryStorageImportErrorContext : public FOutputDevice
diff --git a/Engine/Source/Runtime/RigVM/Public/RigVMCore/RigVMMemoryStorage.h b/Engine/Source/Runtime/RigVM/Public/RigVMCore/RigVMMemoryStorage.h
--- a/Engine/Source/Runtime/RigVM/Public/RigVMCore/RigVMMemoryStorage.h
+++ b/Engine/Source/Runtime/RigVM/Public/RigVMCore/RigVMMemoryStorage.h
@@ -483,30 +483,30 @@ public:
 	}
 
 	// Returns the exported text for a given property index
-	FString GetDataAsString(int32 InPropertyIndex);
+	FString GetDataAsString(int32 InPropertyIndex, int32 PortFlags = PPF_None);
 
 	// Returns the exported text for given property name 
-	FORCEINLINE FString GetDataAsStringByName(const FName& InName)
+	FORCEINLINE FString GetDataAsStringByName(const FName& InName, int32 PortFlags = PPF_None)
 	{
 		const int32 PropertyIndex = GetPropertyIndexByName(InName);
-		return GetDataAsString(PropertyIndex);
+		return GetDataAsString(PropertyIndex, PortFlags);
 	}
 
 	// Returns the exported text for a given operand
-	FString GetDataAsString(const FRigVMOperand& InOperand);
+	FString GetDataAsString(const FRigVMOperand& InOperand, int32 PortFlags = PPF_None);
 
 	// Returns the exported text for a given property index
-	FString GetDataAsStringSafe(int32 InPropertyIndex);
+	FString GetDataAsStringSafe(int32 InPropertyIndex, int32 PortFlags = PPF_None);
 
 	// Returns the exported text for given property name 
-	FORCEINLINE FString GetDataAsStringByNameSafe(const FName& InName)
+	FORCEINLINE FString GetDataAsStringByNameSafe(const FName& InName, int32 PortFlags = PPF_None)
 	{
 		const int32 PropertyIndex = GetPropertyIndexByName(InName);
-		return GetDataAsStringSafe(PropertyIndex);
+		return GetDataAsStringSafe(PropertyIndex, PortFlags);
 	}
 
 	// Returns the exported text for a given operand
-	FString GetDataAsStringSafe(const FRigVMOperand& InOperand);
+	FString GetDataAsStringSafe(const FRigVMOperand& InOperand, int32 PortFlags = PPF_None);
 
 	// Sets the content of a property by index given an exported string. Returns true if succeeded
 	bool SetDataFromString(int32 InPropertyIndex, const FString& InValue);
