Control Rig: Traverse sub pins of select nodes correctly when compiling

Fixed the case of float vs double for arrays - this was never working before apparently.

Also improved the warning message from the VM's memory backend:
Failed to copy uValuesWings (TArray<double>) to Rig_ArrayReset_Array__IO (TArray<float>) in package ...

diff --git a/Engine/Source/Developer/RigVMDeveloper/Private/RigVMCompiler/RigVMCompiler.cpp b/Engine/Source/Developer/RigVMDeveloper/Private/RigVMCompiler/RigVMCompiler.cpp
--- a/Engine/Source/Developer/RigVMDeveloper/Private/RigVMCompiler/RigVMCompiler.cpp
+++ b/Engine/Source/Developer/RigVMDeveloper/Private/RigVMCompiler/RigVMCompiler.cpp
@@ -1077,16 +1077,37 @@ void URigVMCompiler::TraverseAssign(const FRigVMAssignExprAST* InExpr, FRigVMCom
 						return;
 					}
 
+					FString SegmentPath = Pin->GetSegmentPath(false);
+
+					// for select nodes we create a register for each case (since the cases are fixed in size)
+					// thus we do not need to setup a registeroffset for the array element.
 					if (URigVMSelectNode* SelectNode = Cast<URigVMSelectNode>(RootPin->GetNode()))
 					{
-						if (Pin->GetParentPin() == RootPin && RootPin->GetName() == URigVMSelectNode::ValueName)
+						if(RootPin->GetName() == URigVMSelectNode::ValueName)
 						{
-							return;
+							if (Pin->GetParentPin() == RootPin)
+							{
+								return;
+							}
+
+							// if the pin is a sub pin of a case of the select (for example: Values.0.Translation)
+							// we'll need to re-adjust the root pin to the case pin (for example: Values.0)
+							TArray<FString> SegmentPathPaths;
+							if(ensure(URigVMPin::SplitPinPath(SegmentPath, SegmentPathPaths)))
+							{
+								RootPin = RootPin->FindSubPin(SegmentPathPaths[0]);
+
+								SegmentPathPaths.RemoveAt(0);
+								ensure(SegmentPathPaths.Num() > 0);
+								SegmentPath = URigVMPin::JoinPinPath(SegmentPathPaths);
+							}
+							else
+							{
+								return;
+							}
 						}
 					}
 
-					FString SegmentPath = Pin->GetSegmentPath();
-
 #if UE_RIGVM_UCLASS_BASED_STORAGE_DISABLED
 
 					int32 ArrayIndex = INDEX_NONE;
diff --git a/Engine/Source/Runtime/RigVM/Private/RigVMCore/RigVMMemoryStorage.cpp b/Engine/Source/Runtime/RigVM/Private/RigVMCore/RigVMMemoryStorage.cpp
--- a/Engine/Source/Runtime/RigVM/Private/RigVMCore/RigVMMemoryStorage.cpp
+++ b/Engine/Source/Runtime/RigVM/Private/RigVMCore/RigVMMemoryStorage.cpp
@@ -769,9 +769,9 @@ bool URigVMMemoryStorage::CopyProperty(
 				FScriptArrayHelper TargetArray(TargetArrayProperty, InTargetPtr);
 				FScriptArrayHelper SourceArray(SourceArrayProperty, InSourcePtr);
 				
-				if(TargetArrayProperty->IsA<FFloatProperty>())
+				if(TargetArrayProperty->Inner->IsA<FFloatProperty>())
 				{
-					if(SourceArrayProperty->IsA<FDoubleProperty>())
+					if(SourceArrayProperty->Inner->IsA<FDoubleProperty>())
 					{
 						TargetArray.Resize(SourceArray.Num());
 						for(int32 Index=0;Index<TargetArray.Num();Index++)
@@ -783,9 +783,9 @@ bool URigVMMemoryStorage::CopyProperty(
 						return true;
 					}
 				}
-				else if(TargetArrayProperty->IsA<FDoubleProperty>())
+				else if(TargetArrayProperty->Inner->IsA<FDoubleProperty>())
 				{
-					if(SourceArrayProperty->IsA<FFloatProperty>())
+					if(SourceArrayProperty->Inner->IsA<FFloatProperty>())
 					{
 						TargetArray.Resize(SourceArray.Num());
 						for(int32 Index=0;Index<TargetArray.Num();Index++)
@@ -804,9 +804,20 @@ bool URigVMMemoryStorage::CopyProperty(
 		// between two properties which are not compatible
 		if(!InTargetProperty->SameType(InSourceProperty))
 		{
+			FString TargetType, SourceType, TargetExtendedType, SourceExtendedType;
+			TargetType = InTargetProperty->GetCPPType(&TargetExtendedType);
+			SourceType = InSourceProperty->GetCPPType(&SourceExtendedType);
+			TargetType += TargetExtendedType;
+			SourceType += SourceExtendedType;
+
 			UPackage* Package = InTargetProperty->GetOutermost();
 			check(Package);
-			UE_LOG(LogRigVM, Warning, TEXT("Failed to copy %s to %s (%s)"), *InTargetProperty->GetName(), *InSourceProperty->GetName(), *Package->GetName());
+			UE_LOG(LogRigVM, Warning, TEXT("Failed to copy %s (%s) to %s (%s) in package %s"),
+				*InSourceProperty->GetName(),
+				*SourceType,
+				*InTargetProperty->GetName(), 
+				*TargetType,
+				*Package->GetName());
 			return false;
 		}
 	}
