Control Rig: Fix copy property between EnumProperty and ByteProperty

diff --git a/Engine/Source/Runtime/RigVM/Private/RigVMCore/RigVMMemoryStorage.cpp b/Engine/Source/Runtime/RigVM/Private/RigVMCore/RigVMMemoryStorage.cpp
--- a/Engine/Source/Runtime/RigVM/Private/RigVMCore/RigVMMemoryStorage.cpp
+++ b/Engine/Source/Runtime/RigVM/Private/RigVMCore/RigVMMemoryStorage.cpp
@@ -752,6 +752,28 @@ bool URigVMMemoryStorage::CopyProperty(
 				return true;
 			}
 		}
+		else if (const FByteProperty* TargetByteProperty = CastField<FByteProperty>(InTargetProperty))
+		{
+			if (const FEnumProperty* SourceEnumProperty = CastField<FEnumProperty>(InSourceProperty))
+			{
+				if (TargetByteProperty->Enum == SourceEnumProperty->GetEnum())
+				{
+					InTargetProperty->CopyCompleteValue(InTargetPtr, InSourcePtr);
+					return true;
+				}
+			}
+		}
+		else if (const FEnumProperty* TargetEnumProperty = CastField<FEnumProperty>(InTargetProperty))
+		{
+			if (const FByteProperty* SourceByteProperty = CastField<FByteProperty>(InSourceProperty))
+			{
+				if (TargetEnumProperty->GetEnum() == SourceByteProperty->Enum)
+				{
+					InTargetProperty->CopyCompleteValue(InTargetPtr, InSourcePtr);
+					return true;
+				}
+			}
+		}
 		else if(const FArrayProperty* TargetArrayProperty = CastField<FArrayProperty>(InTargetProperty))
 		{
 			if(const FArrayProperty* SourceArrayProperty = CastField<FArrayProperty>(InSourceProperty))
@@ -787,6 +809,28 @@ bool URigVMMemoryStorage::CopyProperty(
 						return true;
 					}
 				}
+				else if(FByteProperty* TargetArrayInnerByteProperty = CastField<FByteProperty>(TargetArrayProperty->Inner))
+				{
+					if(FEnumProperty* SourceArrayInnerEnumProperty = CastField<FEnumProperty>(SourceArrayProperty->Inner))
+					{
+						if (TargetArrayInnerByteProperty->Enum == SourceArrayInnerEnumProperty->GetEnum())
+						{
+							InTargetProperty->CopyCompleteValue(InTargetPtr, InSourcePtr);
+							return true;
+						}
+					}
+				}
+				else if(FEnumProperty* TargetArrayInnerEnumProperty = CastField<FEnumProperty>(TargetArrayProperty->Inner))
+				{
+					if(FByteProperty* SourceArrayInnerByteProperty = CastField<FByteProperty>(SourceArrayProperty->Inner))
+					{
+						if (TargetArrayInnerEnumProperty->GetEnum() == SourceArrayInnerByteProperty->Enum)
+						{
+							InTargetProperty->CopyCompleteValue(InTargetPtr, InSourcePtr);
+							return true;
+						}
+					}
+				}
 			}
 		}
 
