Control Rig: Template collision message to be display and not warning

diff --git a/Engine/Source/Runtime/RigVM/Private/RigVMCore/RigVMTemplate.cpp b/Engine/Source/Runtime/RigVM/Private/RigVMCore/RigVMTemplate.cpp
--- a/Engine/Source/Runtime/RigVM/Private/RigVMCore/RigVMTemplate.cpp
+++ b/Engine/Source/Runtime/RigVM/Private/RigVMCore/RigVMTemplate.cpp
@@ -452,7 +452,7 @@ bool FRigVMTemplate::Merge(const FRigVMTemplate& InOther)
 		if(MatchingArguments == Arguments.Num())
 		{
 			// find the previously defined permutation.
-			UE_LOG(LogRigVM, Warning, TEXT("RigVMFunction '%s' cannot be merged into the '%s' template. It collides with '%s'."),
+			UE_LOG(LogRigVM, Display, TEXT("RigVMFunction '%s' cannot be merged into the '%s' template. It collides with '%s'."),
 				InOther.GetPermutation(0)->Name,
 				*GetNotation().ToString(),
 				GetPermutation(PermutationIndex)->Name);
