RigVM: Add low level mechanism for a node upgrade path

diff --git a/Engine/Source/Runtime/RigVM/Private/RigVMCore/RigVMStruct.cpp b/Engine/Source/Runtime/RigVM/Private/RigVMCore/RigVMStruct.cpp
--- a/Engine/Source/Runtime/RigVM/Private/RigVMCore/RigVMStruct.cpp
+++ b/Engine/Source/Runtime/RigVM/Private/RigVMCore/RigVMStruct.cpp
@@ -3,6 +3,7 @@
 #include "RigVMCore/RigVMStruct.h"
 #include "RigVMCore/RigVMRegistry.h"
 #include "UObject/StructOnScope.h"
+#include "RigVMModule.h"
 
 ////////////////////////////////////////////////////////////////////////////////////////////////////////////////////////
 
@@ -84,6 +85,29 @@ FName FRigVMUnitNodeCreatedContext::FindFirstVariableOfType(UObject* InCPPTypeOb
 	return NAME_None;
 }
 
+FRigVMStructUpgradeInfo::FRigVMStructUpgradeInfo()
+	: NodePath()
+	, OldStruct(nullptr)
+	, NewStruct(nullptr)
+{
+}
+
+bool FRigVMStructUpgradeInfo::IsValid() const
+{
+	return NodePath.IsEmpty() && (OldStruct != nullptr) && (NewStruct != nullptr);
+}
+
+#if WITH_EDITOR
+
+void FRigVMStructUpgradeInfo::SetDefaultValues(const FRigVMStruct* InNewStructMemory)
+{
+	check(NewStruct);
+	check(InNewStructMemory);
+	DefaultValues = InNewStructMemory->GetDefaultValues(NewStruct);
+}
+
+#endif
+
 ////////////////////////////////////////////////////////////////////////////////////////////////////////////////////////
 
 const FName FRigVMStruct::DeprecatedMetaName("Deprecated");
@@ -375,18 +399,18 @@ ERigVMPinDirection FRigVMStruct::GetPinDirectionFromProperty(FProperty* InProper
 	return ERigVMPinDirection::Hidden;
 }
 
-FString FRigVMStruct::ExportToFullyQualifiedText(FProperty* InMemberProperty, const uint8* InMemberMemoryPtr)
+FString FRigVMStruct::ExportToFullyQualifiedText(const FProperty* InMemberProperty, const uint8* InMemberMemoryPtr, bool bUseQuotes)
 {
 	check(InMemberProperty);
 	check(InMemberMemoryPtr);
 
 	FString DefaultValue;
 
-	if (FStructProperty* StructProperty = CastField<FStructProperty>(InMemberProperty))
+	if (const FStructProperty* StructProperty = CastField<FStructProperty>(InMemberProperty))
 	{
 		DefaultValue = ExportToFullyQualifiedText(StructProperty->Struct, InMemberMemoryPtr);
 	}
-	else if (FArrayProperty* ArrayProperty = CastField<FArrayProperty>(InMemberProperty))
+	else if (const FArrayProperty* ArrayProperty = CastField<FArrayProperty>(InMemberProperty))
 	{
 		FScriptArrayHelper ScriptArrayHelper(ArrayProperty, InMemberMemoryPtr);
 
@@ -413,13 +437,16 @@ FString FRigVMStruct::ExportToFullyQualifiedText(FProperty* InMemberProperty, co
 		if (CastField<FNameProperty>(InMemberProperty) != nullptr ||
 			CastField<FStrProperty>(InMemberProperty) != nullptr)
 		{
-			if (DefaultValue.IsEmpty())
+			if(bUseQuotes)
 			{
-				DefaultValue = TEXT("\"\"");
-			}
-			else
-			{
-				DefaultValue = FString::Printf(TEXT("\"%s\""), *DefaultValue);
+				if (DefaultValue.IsEmpty())
+				{
+					DefaultValue = TEXT("\"\"");
+				}
+				else
+				{
+					DefaultValue = FString::Printf(TEXT("\"%s\""), *DefaultValue);
+				}
 			}
 		}
 	}
@@ -454,4 +481,115 @@ FString FRigVMStruct::ExportToFullyQualifiedText(UScriptStruct* InStruct, const
 	return FString::Printf(TEXT("(%s)"), *FString::Join(FieldValues, TEXT(",")));
 }
 
+FString FRigVMStruct::ExportToFullyQualifiedText(UScriptStruct* InScriptStruct, const FName& InPropertyName, const uint8* InStructMemoryPointer) const
+{
+	check(InScriptStruct);
+	if(InStructMemoryPointer == nullptr)
+	{
+		InStructMemoryPointer = (const uint8*)this;
+	}
+	
+	const FProperty* Property = InScriptStruct->FindPropertyByName(InPropertyName);
+	if(Property == nullptr)
+	{
+		return FString();
+	}
+
+	const uint8* StructMemberMemoryPtr = Property->ContainerPtrToValuePtr<uint8>(InStructMemoryPointer);
+	return ExportToFullyQualifiedText(Property, StructMemberMemoryPtr);
+}
+
+TMap<FName, FString> FRigVMStruct::GetDefaultValues(UScriptStruct* InScriptStruct) const
+{
+	check(InScriptStruct);
+	
+	FRigVMStructUpgradeInfo Info;
+	Info.OldStruct = InScriptStruct;
+
+	TMap<FName, FString> DefaultValues;
+	for (TFieldIterator<FProperty> It(InScriptStruct); It; ++It)
+	{
+		if(It->HasAnyPropertyFlags(CPF_Transient))
+		{
+			continue;
+		}
+		
+		const FName PropertyName = It->GetFName();
+		const uint8* StructMemberMemoryPtr = It->ContainerPtrToValuePtr<uint8>(this);
+		const FString DefaultValue = ExportToFullyQualifiedText(*It, StructMemberMemoryPtr, false);
+		DefaultValues.Add(PropertyName, DefaultValue);
+	}
+
+	return DefaultValues;
+}
+
+class FRigVMStructApplyUpgradeInfoErrorContext : public FOutputDevice
+{
+public:
+
+	int32 NumErrors;
+	UScriptStruct* OldStruct;
+	UScriptStruct* NewStruct;
+	FName PropertyName;
+
+	FRigVMStructApplyUpgradeInfoErrorContext(UScriptStruct* InOldStruct, UScriptStruct* InNewStruct, const FName& InPropertyName)
+		: FOutputDevice()
+		, NumErrors(0)
+		, OldStruct(InOldStruct)
+		, NewStruct(InNewStruct)
+		, PropertyName(InPropertyName)
+	{
+	}
+
+	virtual void Serialize(const TCHAR* V, ELogVerbosity::Type Verbosity, const class FName& Category) override
+	{
+		if(Verbosity == ELogVerbosity::Error)
+		{
+			UE_LOG(LogRigVM, Warning, TEXT("Error when applying upgrade data from %s to %s, property %s: %s"),
+				*OldStruct->GetStructCPPName(),
+				*NewStruct->GetStructCPPName(),
+				*PropertyName.ToString(),
+				V);
+			NumErrors++;
+		}
+		else if(Verbosity == ELogVerbosity::Warning)
+		{
+			UE_LOG(LogRigVM, Warning, TEXT("Warning when applying upgrade data from %s to %s, property %s: %s"),
+				*OldStruct->GetStructCPPName(),
+				*NewStruct->GetStructCPPName(),
+				*PropertyName.ToString(),
+				V);
+		}
+	}
+};
+
+bool FRigVMStruct::ApplyUpgradeInfo(const FRigVMStructUpgradeInfo& InUpgradeInfo)
+{
+	check(InUpgradeInfo.IsValid());
+
+	for(const TPair<FName, FString>& Pair : InUpgradeInfo.DefaultValues)
+	{
+		const FName& PropertyName = Pair.Key;
+		const FString& DefaultValue = Pair.Value;
+		
+		const FProperty* Property = InUpgradeInfo.NewStruct->FindPropertyByName(PropertyName);
+		if(Property == nullptr)
+		{
+			return false;
+		}
+
+		uint8* MemberMemory = Property->ContainerPtrToValuePtr<uint8>(this);
+
+		FRigVMStructApplyUpgradeInfoErrorContext ErrorPipe(InUpgradeInfo.OldStruct, InUpgradeInfo.NewStruct, PropertyName);
+        Property->ImportText_Direct(*DefaultValue, MemberMemory, nullptr, PPF_None, &ErrorPipe);
+
+		if(ErrorPipe.NumErrors > 0)
+		{
+			return false;
+		}
+	}
+
+	return true;
+}
+
 #endif
\ No newline at end of file
diff --git a/Engine/Source/Runtime/RigVM/Public/RigVMCore/RigVMStruct.h b/Engine/Source/Runtime/RigVM/Public/RigVMCore/RigVMStruct.h
--- a/Engine/Source/Runtime/RigVM/Public/RigVMCore/RigVMStruct.h
+++ b/Engine/Source/Runtime/RigVM/Public/RigVMCore/RigVMStruct.h
@@ -14,6 +14,8 @@ DECLARE_DELEGATE_RetVal(TArray<FRigVMExternalVariable>, FRigVMGetExternalVariabl
 DECLARE_DELEGATE_RetVal_TwoParams(FName, FRigVMCreateExternalVariableDelegate, FRigVMExternalVariable, FString)
 DECLARE_DELEGATE_RetVal_TwoParams(bool, FRigVMBindPinToExternalVariableDelegate, FString, FString)
 
+struct FRigVMStruct;
+
 /** Context as of why the node was created */
 enum class ERigVMNodeCreatedReason : uint8
 {
@@ -133,6 +135,49 @@ private:
 	friend struct FScope;
 };
 
+/** Structure used to upgrade to a new implementation of a node */
+struct RIGVM_API FRigVMStructUpgradeInfo
+{
+	FRigVMStructUpgradeInfo();
+
+	// returns true if this upgrade info can be applied
+	bool IsValid() const;
+
+#if WITH_EDITOR
+	
+	// sets the default values from the new struct.
+	void SetDefaultValues(const FRigVMStruct* InNewStructMemory);
+
+#endif
+
+	// The complete node path including models / collapse node.
+	// The path may look like "RigGraph|CollapseNode1|Add"
+	FString NodePath;
+	
+	// The old struct this upgrade info originates from
+	UScriptStruct* OldStruct;
+
+	// The new struct this upgrade info is targeting
+	UScriptStruct* NewStruct;
+
+	// Remapping info for re-mapping pins
+	// Entries can be root pins or sub pins
+	TMap<FString, FString> PinNameMap;
+
+	// Remapping info for re-linking inputs 
+	// (takes precedence over pin name map)
+	// Entries can be root pins or sub pins
+	TMap<FString, FString> InputLinkMap;
+
+	// Remapping info for re-linking outputs 
+	// (takes precedence over pin name map)
+	// Entries can be root pins or sub pins
+	TMap<FString, FString> OutputLinkMap;
+
+	// New sets of default values
+	TMap<FName, FString> DefaultValues;
+};
+
 /**
  * The base class for all RigVM enabled structs.
  */
@@ -155,7 +200,6 @@ public:
 	FORCEINLINE virtual void OnUnitNodeCreated(FRigVMUnitNodeCreatedContext& InContext) const {}
 
 #if WITH_EDITOR
-
 	static bool ValidateStruct(UScriptStruct* InStruct, FString* OutErrorMessage);
 	static bool CheckPinType(UScriptStruct* InStruct, const FName& PinName, const FString& ExpectedType, FString* OutErrorMessage = nullptr);
 	static bool CheckPinDirection(UScriptStruct* InStruct, const FName& PinName, const FName& InDirectionMetaName);
@@ -163,8 +207,10 @@ public:
 	static bool CheckPinExists(UScriptStruct* InStruct, const FName& PinName, const FString& ExpectedType = FString(), FString* OutErrorMessage = nullptr);
 	static bool CheckMetadata(UScriptStruct* InStruct, const FName& PinName, const FName& InMetadataKey, FString* OutErrorMessage = nullptr);
 	static bool CheckFunctionExists(UScriptStruct* InStruct, const FName& FunctionName, FString* OutErrorMessage = nullptr);
-	static FString ExportToFullyQualifiedText(FProperty* InMemberProperty, const uint8* InMemberMemoryPtr);
+	static FString ExportToFullyQualifiedText(const FProperty* InMemberProperty, const uint8* InMemberMemoryPtr, bool bUseQuotes = true);
 	static FString ExportToFullyQualifiedText(UScriptStruct* InStruct, const uint8* InStructMemoryPtr);
+	FString ExportToFullyQualifiedText(UScriptStruct* InScriptStruct, const FName& InPropertyName, const uint8* InStructMemoryPointer = nullptr) const;
+	virtual FRigVMStructUpgradeInfo GetUpgradeInfo() const { return FRigVMStructUpgradeInfo(); }
 #endif
 
 	static const FName DeprecatedMetaName;
@@ -204,5 +250,12 @@ public:
 protected:
 
 	static float GetRatioFromIndex(int32 InIndex, int32 InCount);
+#if WITH_EDITOR
+	TMap<FName, FString> GetDefaultValues(UScriptStruct* InScriptStruct) const;
+	bool ApplyUpgradeInfo(const FRigVMStructUpgradeInfo& InUpgradeInfo);
+#endif
 
+	friend struct FRigVMStructUpgradeInfo;
+	friend class FRigVMGraphStructUpgradeInfoTest;
+	friend class URigVMController;
 };
