Control Rig: Remove access to external variables from RigVMExecuteContext

diff --git a/Engine/Source/Runtime/RigVM/Private/RigVMCore/RigVM.cpp b/Engine/Source/Runtime/RigVM/Private/RigVMCore/RigVM.cpp
--- a/Engine/Source/Runtime/RigVM/Private/RigVMCore/RigVM.cpp
+++ b/Engine/Source/Runtime/RigVM/Private/RigVMCore/RigVM.cpp
@@ -1335,7 +1335,6 @@ bool URigVM::Initialize(TArrayView<URigVMMemoryStorage*> Memory, TArrayView<void
 	Context.Reset();
 	Context.SliceOffsets.AddZeroed(Instructions.Num());
 	Context.OpaqueArguments = AdditionalArguments;
-	Context.ExternalVariables = ExternalVariables;
 
 	TGuardValue<URigVM*> VMInContext(Context.VM, this);
 	
@@ -1545,7 +1544,6 @@ bool URigVM::Execute(TArrayView<URigVMMemoryStorage*> Memory, TArrayView<void*>
 	Context.Reset();
 	Context.SliceOffsets.AddZeroed(Instructions.Num());
 	Context.OpaqueArguments = AdditionalArguments;
-	Context.ExternalVariables = ExternalVariables;
 
 	TGuardValue<URigVM*> VMInContext(Context.VM, this);
 
diff --git a/Engine/Source/Runtime/RigVM/Public/RigVMCore/RigVMExecuteContext.h b/Engine/Source/Runtime/RigVM/Public/RigVMCore/RigVMExecuteContext.h
--- a/Engine/Source/Runtime/RigVM/Public/RigVMCore/RigVMExecuteContext.h
+++ b/Engine/Source/Runtime/RigVM/Public/RigVMCore/RigVMExecuteContext.h
@@ -214,7 +214,6 @@ struct RIGVM_API FRigVMExecuteContext
 		SliceOffsets.Reset();
 		InstructionIndex = 0;
 		VM = nullptr;
-		ExternalVariables.Reset();
 	}
 
 	FORCEINLINE void CopyFrom(const FRigVMExecuteContext& Other)
@@ -224,7 +223,6 @@ struct RIGVM_API FRigVMExecuteContext
 		InstructionIndex = Other.InstructionIndex;
 		VM = Other.VM;
 		RuntimeSettings = Other.RuntimeSettings;
-		ExternalVariables = Other.ExternalVariables;
 		OpaqueArguments = Other.OpaqueArguments;
 		Slices = Other.Slices;
 		SliceOffsets = Other.SliceOffsets;
@@ -265,30 +263,6 @@ struct RIGVM_API FRigVMExecuteContext
 		return GetSlice().IsComplete();
 	}
 
-	FORCEINLINE const FRigVMExternalVariable* FindExternalVariable(const FName& InExternalVariableName) const
-	{
-		for (const FRigVMExternalVariable& ExternalVariable : ExternalVariables)
-		{
-			if (ExternalVariable.Name == InExternalVariableName)
-			{
-				return &ExternalVariable;
-			}
-		}
-		return nullptr;
-	}
-
-	FORCEINLINE FRigVMExternalVariable* FindExternalVariable(const FName& InExternalVariableName)
-	{
-		for (FRigVMExternalVariable& ExternalVariable : ExternalVariables)
-		{
-			if (ExternalVariable.Name == InExternalVariableName)
-			{
-				return &ExternalVariable;
-			}
-		}
-		return nullptr;
-	}
-
 	FORCEINLINE void Log(EMessageSeverity::Type InSeverity, const FString& InMessage) const
 	{
 		if(RuntimeSettings.LogFunction)
@@ -352,7 +326,6 @@ struct RIGVM_API FRigVMExecuteContext
 	URigVM* VM;
 	FRigVMRuntimeSettings RuntimeSettings;
 	TArrayView<void*> OpaqueArguments;
-	TArray<FRigVMExternalVariable> ExternalVariables;
 	TArray<FRigVMSlice> Slices;
 	TArray<uint16> SliceOffsets;
 	double LastExecutionMicroSeconds;
