Control Rig: Remove access to private arguments of RigVMExecuteContext from RigVMUnits

diff --git a/Engine/Plugins/Animation/ControlRig/Source/ControlRig/Private/ControlRig.cpp b/Engine/Plugins/Animation/ControlRig/Source/ControlRig/Private/ControlRig.cpp
--- a/Engine/Plugins/Animation/ControlRig/Source/ControlRig/Private/ControlRig.cpp
+++ b/Engine/Plugins/Animation/ControlRig/Source/ControlRig/Private/ControlRig.cpp
@@ -577,14 +577,14 @@ void UControlRig::Execute(const EControlRigState InState, const FName& InEventNa
 			{
 				HierarchyController->LogFunction = [this](EMessageSeverity::Type InSeverity, const FString& Message)
 				{
-					const FRigVMExecuteContext& Context = GetVM()->GetContext();
+					const FRigVMExtendedExecuteContext& Context = GetVM()->GetContext();
 					if(ControlRigLog)
 					{
-						ControlRigLog->Report(InSeverity, Context.FunctionName, Context.InstructionIndex, Message);
+						ControlRigLog->Report(InSeverity, Context.PublicData.FunctionName, Context.PublicData.InstructionIndex, Message);
 					}
 					else
 					{
-						LogOnce(InSeverity, Context.InstructionIndex, Message);
+						LogOnce(InSeverity, Context.PublicData.InstructionIndex, Message);
 					}
 				};
 			}
@@ -1236,7 +1236,7 @@ void UControlRig::ExecuteUnits(FRigUnitContext& InOutContext, const FName& InEve
 			Hierarchy->WrittenTransformsPerInstructionPerSlice.Reset();
 			Hierarchy->ReadTransformsPerInstructionPerSlice.AddZeroed(VM->GetByteCode().GetNumInstructions());
 			Hierarchy->WrittenTransformsPerInstructionPerSlice.AddZeroed(VM->GetByteCode().GetNumInstructions());
-			TGuardValue<const FRigVMExecuteContext*> HierarchyContextGuard(Hierarchy->ExecuteContext, &VM->GetContext());
+			TGuardValue<const FRigVMExtendedExecuteContext*> HierarchyContextGuard(Hierarchy->ExecuteContext, &VM->GetContext());
 #endif
 
 			VM->Execute(LocalMemory, AdditionalArguments, InEventName);
diff --git a/Engine/Plugins/Animation/ControlRig/Source/ControlRig/Private/Rigs/FKControlRig.cpp b/Engine/Plugins/Animation/ControlRig/Source/ControlRig/Private/Rigs/FKControlRig.cpp
--- a/Engine/Plugins/Animation/ControlRig/Source/ControlRig/Private/Rigs/FKControlRig.cpp
+++ b/Engine/Plugins/Animation/ControlRig/Source/ControlRig/Private/Rigs/FKControlRig.cpp
@@ -50,8 +50,6 @@ void UFKControlRig::ExecuteUnits(FRigUnitContext& InOutContext, const FName& InE
 
 	if (InEventName == FRigUnit_BeginExecution::EventName)
 	{
-		FRigVMExecuteContext VMContext;
-
 		GetHierarchy()->ForEach<FRigBoneElement>([&](FRigBoneElement* BoneElement) -> bool
         {
 			const FName ControlName = GetControlName(BoneElement->GetName());
@@ -120,7 +118,6 @@ void UFKControlRig::ExecuteUnits(FRigUnitContext& InOutContext, const FName& InE
 	}
 	else if (InEventName == FRigUnit_InverseExecution::EventName)
 	{
-		FRigVMExecuteContext VMContext;
 		const bool bNotify = true;
 		const FRigControlModifiedContext Context = FRigControlModifiedContext();
 		const bool bSetupUndo = false;
diff --git a/Engine/Plugins/Animation/ControlRig/Source/ControlRig/Private/Rigs/RigHierarchy.cpp b/Engine/Plugins/Animation/ControlRig/Source/ControlRig/Private/Rigs/RigHierarchy.cpp
--- a/Engine/Plugins/Animation/ControlRig/Source/ControlRig/Private/Rigs/RigHierarchy.cpp
+++ b/Engine/Plugins/Animation/ControlRig/Source/ControlRig/Private/Rigs/RigHierarchy.cpp
@@ -2179,7 +2179,7 @@ FTransform URigHierarchy::GetTransform(FRigTransformElement* InTransformElement,
 
 	if(bRecordTransformsPerInstruction && ExecuteContext)
 	{
-		TArray<TArray<int32>>& ReadTransformsPerSlice = ReadTransformsPerInstructionPerSlice[ExecuteContext->InstructionIndex];
+		TArray<TArray<int32>>& ReadTransformsPerSlice = ReadTransformsPerInstructionPerSlice[ExecuteContext->PublicData.GetInstructionIndex()];
 		while(ReadTransformsPerSlice.Num() < ExecuteContext->GetSlice().TotalNum())
 		{
 			ReadTransformsPerSlice.Add(TArray<int32>());
@@ -2284,7 +2284,7 @@ void URigHierarchy::SetTransform(FRigTransformElement* InTransformElement, const
 
 	if(bRecordTransformsPerInstruction && ExecuteContext)
 	{
-		TArray<TArray<int32>>& WrittenTransformsPerSlice = WrittenTransformsPerInstructionPerSlice[ExecuteContext->InstructionIndex];
+		TArray<TArray<int32>>& WrittenTransformsPerSlice = WrittenTransformsPerInstructionPerSlice[ExecuteContext->PublicData.GetInstructionIndex()];
 		while(WrittenTransformsPerSlice.Num() < ExecuteContext->GetSlice().TotalNum())
 		{
 			WrittenTransformsPerSlice.Add(TArray<int32>());
@@ -2412,7 +2412,7 @@ FTransform URigHierarchy::GetControlOffsetTransform(FRigControlElement* InContro
 
 	if(bRecordTransformsPerInstruction && ExecuteContext)
 	{
-		TArray<TArray<int32>>& ReadTransformsPerSlice = ReadTransformsPerInstructionPerSlice[ExecuteContext->InstructionIndex];
+		TArray<TArray<int32>>& ReadTransformsPerSlice = ReadTransformsPerInstructionPerSlice[ExecuteContext->PublicData.GetInstructionIndex()];
 		while(ReadTransformsPerSlice.Num() < ExecuteContext->GetSlice().TotalNum())
 		{
 			ReadTransformsPerSlice.Add(TArray<int32>());
@@ -2495,7 +2495,7 @@ void URigHierarchy::SetControlOffsetTransform(FRigControlElement* InControlEleme
 
 	if(bRecordTransformsPerInstruction && ExecuteContext)
 	{
-		TArray<TArray<int32>>& WrittenTransformsPerSlice = WrittenTransformsPerInstructionPerSlice[ExecuteContext->InstructionIndex];
+		TArray<TArray<int32>>& WrittenTransformsPerSlice = WrittenTransformsPerInstructionPerSlice[ExecuteContext->PublicData.GetInstructionIndex()];
 		while(WrittenTransformsPerSlice.Num() < ExecuteContext->GetSlice().TotalNum())
 		{
 			WrittenTransformsPerSlice.Add(TArray<int32>());
@@ -4038,7 +4038,7 @@ void URigHierarchy::EnsureCacheValidityImpl()
 
 #if WITH_EDITOR
 
-URigHierarchy::TElementDependencyMap URigHierarchy::GetDependenciesForVM(URigVM* InVM, FName InEventName)
+URigHierarchy::TElementDependencyMap URigHierarchy::GetDependenciesForVM(const URigVM* InVM, FName InEventName) const
 {
 	LLM_SCOPE_BYNAME(TEXT("Animation/ControlRig"));
 	check(InVM);
diff --git a/Engine/Plugins/Animation/ControlRig/Source/ControlRig/Private/Units/Hierarchy/RigUnit_SendEvent.cpp b/Engine/Plugins/Animation/ControlRig/Source/ControlRig/Private/Units/Hierarchy/RigUnit_SendEvent.cpp
--- a/Engine/Plugins/Animation/ControlRig/Source/ControlRig/Private/Units/Hierarchy/RigUnit_SendEvent.cpp
+++ b/Engine/Plugins/Animation/ControlRig/Source/ControlRig/Private/Units/Hierarchy/RigUnit_SendEvent.cpp
@@ -31,7 +31,7 @@ FRigUnit_SendEvent_Execute()
 				FRigEventContext EventContext;
 				EventContext.Key = Item;
 				EventContext.Event = Event;
-				EventContext.SourceEventName = ExecuteContext.EventName;
+				EventContext.SourceEventName = ExecuteContext.GetEventName();
 				EventContext.LocalTime = Context.AbsoluteTime + OffsetInSeconds;
 				Hierarchy->SendEvent(EventContext, false /* async */); //needs to be false for sequencer keying to work
 				break;
diff --git a/Engine/Plugins/Animation/ControlRig/Source/ControlRig/Private/Units/Hierarchy/RigUnit_SetTransform.cpp b/Engine/Plugins/Animation/ControlRig/Source/ControlRig/Private/Units/Hierarchy/RigUnit_SetTransform.cpp
--- a/Engine/Plugins/Animation/ControlRig/Source/ControlRig/Private/Units/Hierarchy/RigUnit_SetTransform.cpp
+++ b/Engine/Plugins/Animation/ControlRig/Source/ControlRig/Private/Units/Hierarchy/RigUnit_SetTransform.cpp
@@ -47,7 +47,7 @@ FRigUnit_SetTransform_Execute()
 						FTransform TransformMutable = Value;
 						FRigUnit_SetControlOffset::StaticExecute(RigVMExecuteContext, CachedIndex.GetKey().Name, TransformMutable, Space, CachedIndex, ExecuteContext, Context);
 						
-						if (ExecuteContext.EventName == FRigUnit_PrepareForExecution::EventName)
+						if (ExecuteContext.GetEventName() == FRigUnit_PrepareForExecution::EventName)
 						{
 							Hierarchy->SetLocalTransformByIndex(CachedIndex, FTransform::Identity, true, bPropagateToChildren);
 							Hierarchy->SetLocalTransformByIndex(CachedIndex, FTransform::Identity, false, bPropagateToChildren);
@@ -85,7 +85,7 @@ FRigUnit_SetTransform_Execute()
 						{
 							Hierarchy->SetGlobalTransformByIndex(CachedIndex, WeightedTransform, bInitial, bPropagateToChildren);
 
-							if (bInitial && ExecuteContext.EventName == FRigUnit_PrepareForExecution::EventName)
+							if (bInitial && ExecuteContext.GetEventName() == FRigUnit_PrepareForExecution::EventName)
 							{
 								Hierarchy->SetGlobalTransformByIndex(CachedIndex, WeightedTransform, false, bPropagateToChildren);
 							}
@@ -95,7 +95,7 @@ FRigUnit_SetTransform_Execute()
 						{
 							Hierarchy->SetLocalTransformByIndex(CachedIndex, WeightedTransform, bInitial, bPropagateToChildren);
 
-							if (bInitial && ExecuteContext.EventName == FRigUnit_PrepareForExecution::EventName)
+							if (bInitial && ExecuteContext.GetEventName() == FRigUnit_PrepareForExecution::EventName)
 							{
 								Hierarchy->SetLocalTransformByIndex(CachedIndex, WeightedTransform, false, bPropagateToChildren);
 							}
diff --git a/Engine/Plugins/Animation/ControlRig/Source/ControlRig/Public/Rigs/RigHierarchy.h b/Engine/Plugins/Animation/ControlRig/Source/ControlRig/Public/Rigs/RigHierarchy.h
--- a/Engine/Plugins/Animation/ControlRig/Source/ControlRig/Public/Rigs/RigHierarchy.h
+++ b/Engine/Plugins/Animation/ControlRig/Source/ControlRig/Public/Rigs/RigHierarchy.h
@@ -2973,14 +2973,14 @@ private:
 	void EnsureCacheValidityImpl();
 
 #if WITH_EDITOR
-	const FRigVMExecuteContext* ExecuteContext;
+	const FRigVMExtendedExecuteContext* ExecuteContext;
 	mutable bool bRecordTransformsPerInstruction;
 	mutable TArray<TArray<TArray<int32>>> ReadTransformsPerInstructionPerSlice;
 	mutable TArray<TArray<TArray<int32>>> WrittenTransformsPerInstructionPerSlice;
 
 public:
 
-	TElementDependencyMap GetDependenciesForVM(URigVM* InVM, FName InEventName = NAME_None);
+	TElementDependencyMap GetDependenciesForVM(const URigVM* InVM, FName InEventName = NAME_None) const;
 
 private:
 	
diff --git a/Engine/Plugins/Animation/ControlRig/Source/ControlRig/Public/Units/RigUnitContext.h b/Engine/Plugins/Animation/ControlRig/Source/ControlRig/Public/Units/RigUnitContext.h
--- a/Engine/Plugins/Animation/ControlRig/Source/ControlRig/Public/Units/RigUnitContext.h
+++ b/Engine/Plugins/Animation/ControlRig/Source/ControlRig/Public/Units/RigUnitContext.h
@@ -163,7 +163,7 @@ struct FRigUnitContext
 #define UE_CONTROLRIG_RIGUNIT_REPORT(Severity, Format, ...) \
 if(Context.Log != nullptr) \
 { \
-	Context.Log->Report(EMessageSeverity::Severity, RigVMExecuteContext.FunctionName, RigVMExecuteContext.InstructionIndex, FString::Printf((Format), ##__VA_ARGS__)); \
+	Context.Log->Report(EMessageSeverity::Severity, RigVMExecuteContext.GetFunctionName(), RigVMExecuteContext.GetInstructionIndex(), FString::Printf((Format), ##__VA_ARGS__)); \
 }
 #define UE_CONTROLRIG_RIGUNIT_LOG_MESSAGE(Format, ...) UE_CONTROLRIG_RIGUNIT_REPORT(Info, (Format), ##__VA_ARGS__)
 #define UE_CONTROLRIG_RIGUNIT_REPORT_WARNING(Format, ...) UE_CONTROLRIG_RIGUNIT_REPORT(Warning, (Format), ##__VA_ARGS__)
diff --git a/Engine/Source/Programs/UnrealHeaderTool/Private/CodeGenerator.cpp b/Engine/Source/Programs/UnrealHeaderTool/Private/CodeGenerator.cpp
--- a/Engine/Source/Programs/UnrealHeaderTool/Private/CodeGenerator.cpp
+++ b/Engine/Source/Programs/UnrealHeaderTool/Private/CodeGenerator.cpp
@@ -3381,7 +3381,8 @@ void FNativeClassHeaderGenerator::ExportGeneratedStructBodyMacros(FOutputDevice&
 	const FString& SingletonName = ScriptStructDef.GetSingletonName(true);
 	const FString& ChoppedSingletonName = ScriptStructDef.GetSingletonNameChopped(true);
 
-	const FString RigVMParameterPrefix = TEXT("FRigVMExecuteContext& RigVMExecuteContext");
+	const FString RigVMExecuteContextDeclaration = TEXT("FRigVMExtendedExecuteContext& RigVMExecuteContext");
+	const FString RigVMExecuteContextPublicDeclaration = TEXT("FRigVMExecuteContext& RigVMExecuteContext");
 	TArray<FString> RigVMVirtualFuncProlog, RigVMVirtualFuncEpilog, RigVMStubProlog;
 
 	// for RigVM methods we need to generated a macro used for implementing the static method
@@ -3448,7 +3449,7 @@ void FNativeClassHeaderGenerator::ExportGeneratedStructBodyMacros(FOutputDevice&
 		for (const FRigVMMethodInfo& MethodInfo : StructRigVMInfo.Methods)
 		{
 			FString ParameterSuffix = MethodInfo.Parameters.Declarations(true, TEXT(", \\\r\n\t\t"));
-			FString RigVMParameterPrefix2 = RigVMParameterPrefix + FString((StructMembers.IsEmpty() && ParameterSuffix.IsEmpty()) ? TEXT("") : TEXT(", \\\r\n\t\t"));
+			FString RigVMParameterPrefix2 = RigVMExecuteContextPublicDeclaration + FString((StructMembers.IsEmpty() && ParameterSuffix.IsEmpty()) ? TEXT("") : TEXT(", \\\r\n\t\t"));
 			OutGeneratedHeaderText.Logf(TEXT("#define %s_%s() \\\r\n"), *StructNameCPP, *MethodInfo.Name);
 			OutGeneratedHeaderText.Logf(TEXT("\t%s %s::Static%s( \\\r\n\t\t%s%s%s \\\r\n\t)\r\n"), *MethodInfo.ReturnType, *StructNameCPP, *MethodInfo.Name, *RigVMParameterPrefix2, *StructMembers, *ParameterSuffix);
 		}
@@ -3476,11 +3477,11 @@ void FNativeClassHeaderGenerator::ExportGeneratedStructBodyMacros(FOutputDevice&
 				FString StructMembersForStub = StructRigVMInfo.Members.Names(false, TEXT(",\r\n\t\t\t"), false);
 				FString ParameterSuffix = MethodInfo.Parameters.Declarations(true, TEXT(",\r\n\t\t"));
 				FString ParameterNamesSuffix = MethodInfo.Parameters.Names(true, TEXT(",\r\n\t\t\t"));
-				FString RigVMParameterPrefix2 = RigVMParameterPrefix + FString((StructMembers.IsEmpty() && ParameterSuffix.IsEmpty()) ? TEXT("") : TEXT(",\r\n\t\t"));
-				FString RigVMParameterPrefix4 = FString(TEXT("RigVMExecuteContext")) + FString((StructMembersForStub.IsEmpty() && ParameterSuffix.IsEmpty()) ? TEXT("") : TEXT(",\r\n\t\t\t"));
+				FString RigVMParameterPrefix2 = RigVMExecuteContextPublicDeclaration + FString((StructMembers.IsEmpty() && ParameterSuffix.IsEmpty()) ? TEXT("") : TEXT(",\r\n\t\t"));
+				FString RigVMParameterPrefix4 = FString(TEXT("RigVMExecuteContext.PublicData")) + FString((StructMembersForStub.IsEmpty() && ParameterSuffix.IsEmpty()) ? TEXT("") : TEXT(",\r\n\t\t\t"));
 
 				RigVMMethodsDeclarations += FString::Printf(TEXT("\tstatic %s Static%s(\r\n\t\t%s%s%s\r\n\t);\r\n"), *MethodInfo.ReturnType, *MethodInfo.Name, *RigVMParameterPrefix2, *StructMembers, *ParameterSuffix);
-				RigVMMethodsDeclarations += FString::Printf(TEXT("\tFORCEINLINE_DEBUGGABLE static %s RigVM%s(\r\n\t\t%s,\r\n\t\tFRigVMMemoryHandleArray RigVMMemoryHandles\r\n\t)\r\n"), *MethodInfo.ReturnType, *MethodInfo.Name, *RigVMParameterPrefix);
+				RigVMMethodsDeclarations += FString::Printf(TEXT("\tFORCEINLINE_DEBUGGABLE static %s RigVM%s(\r\n\t\t%s,\r\n\t\tFRigVMMemoryHandleArray RigVMMemoryHandles\r\n\t)\r\n"), *MethodInfo.ReturnType, *MethodInfo.Name, *RigVMExecuteContextDeclaration);
 				RigVMMethodsDeclarations += FString::Printf(TEXT("\t{\r\n"));
 
 				// implement inline stub method body
@@ -3684,7 +3685,6 @@ void FNativeClassHeaderGenerator::ExportGeneratedStructBodyMacros(FOutputDevice&
 
 			FString ParameterDeclaration = MethodInfo.Parameters.Declarations(false, TEXT(",\r\n\t\t"));
 			FString ParameterSuffix = MethodInfo.Parameters.Names(true, TEXT(",\r\n\t\t"));
-			FString RigVMParameterPrefix2 = RigVMParameterPrefix + FString((StructMembersForVirtualFunc.IsEmpty() && ParameterSuffix.IsEmpty()) ? TEXT("") : TEXT(",\r\n\t\t"));
 			FString RigVMParameterPrefix3 = FString(TEXT("RigVMExecuteContext")) + FString((StructMembersForVirtualFunc.IsEmpty() && ParameterSuffix.IsEmpty()) ? TEXT("") : TEXT(",\r\n\t\t"));
 
 			// implement the virtual function body.
diff --git a/Engine/Source/Runtime/RigVM/Private/RigVMCore/RigVM.cpp b/Engine/Source/Runtime/RigVM/Private/RigVMCore/RigVM.cpp
--- a/Engine/Source/Runtime/RigVM/Private/RigVMCore/RigVM.cpp
+++ b/Engine/Source/Runtime/RigVM/Private/RigVMCore/RigVM.cpp
@@ -1152,12 +1152,12 @@ bool URigVM::ShouldHaltAtInstruction(const FName& InEventName, const uint16 Inst
 						DebugInfo->SetCurrentActiveBreakpoint(Breakpoint);
 						
 						// We want to keep the callstack up to the node that produced the halt
-						const TArray<UObject*>* FullCallstack = ByteCode.GetCallstackForInstruction(Context.InstructionIndex);
+						const TArray<UObject*>* FullCallstack = ByteCode.GetCallstackForInstruction(Context.PublicData.InstructionIndex);
 						if (FullCallstack)
 						{
 							DebugInfo->SetCurrentActiveBreakpointCallstack(TArray<UObject*>(FullCallstack->GetData(), FullCallstack->Find((UObject*)Breakpoint->Subject)+1));
 						}
-						ExecutionHalted().Broadcast(Context.InstructionIndex, Breakpoint->Subject, InEventName);
+						ExecutionHalted().Broadcast(Context.PublicData.InstructionIndex, Breakpoint->Subject, InEventName);
 					}
 					return true;
 				}
@@ -1184,7 +1184,7 @@ bool URigVM::ShouldHaltAtInstruction(const FName& InEventName, const uint16 Inst
 					if (!DebugInfo->GetCurrentActiveBreakpoint())
 					{
 						DebugInfo->SetCurrentActiveBreakpoint(Breakpoint);
-						const TArray<UObject*>* FullCallstack = ByteCode.GetCallstackForInstruction(Context.InstructionIndex);
+						const TArray<UObject*>* FullCallstack = ByteCode.GetCallstackForInstruction(Context.PublicData.InstructionIndex);
 						
 						// We want to keep the callstack up to the node that produced the halt
 						if (FullCallstack)
@@ -1211,7 +1211,7 @@ bool URigVM::ShouldHaltAtInstruction(const FName& InEventName, const uint16 Inst
 	// If we are stepping, and the last active breakpoint was set, check if this is the new temporary breakpoint
 	if (CurrentBreakpointAction != ERigVMBreakpointAction::None && DebugInfo->GetCurrentActiveBreakpoint())
 	{
-		const TArray<UObject*>* CurrentCallstack = ByteCode.GetCallstackForInstruction(Context.InstructionIndex);
+		const TArray<UObject*>* CurrentCallstack = ByteCode.GetCallstackForInstruction(Context.PublicData.InstructionIndex);
 		if (CurrentCallstack && !CurrentCallstack->IsEmpty())
 		{
 			UObject* NewBreakpointNode = nullptr;
@@ -1276,14 +1276,14 @@ bool URigVM::ShouldHaltAtInstruction(const FName& InEventName, const uint16 Inst
 				}
 
 				// Create new temporary breakpoint
-				TSharedPtr<FRigVMBreakpoint> NewBreakpoint = DebugInfo->AddBreakpoint(Context.InstructionIndex, NewBreakpointNode, 0, true);
-				DebugInfo->SetBreakpointHits(NewBreakpoint, GetInstructionVisitedCount(Context.InstructionIndex));
-				DebugInfo->SetBreakpointActivationOnHit(NewBreakpoint, GetInstructionVisitedCount(Context.InstructionIndex));
+				TSharedPtr<FRigVMBreakpoint> NewBreakpoint = DebugInfo->AddBreakpoint(Context.PublicData.InstructionIndex, NewBreakpointNode, 0, true);
+				DebugInfo->SetBreakpointHits(NewBreakpoint, GetInstructionVisitedCount(Context.PublicData.InstructionIndex));
+				DebugInfo->SetBreakpointActivationOnHit(NewBreakpoint, GetInstructionVisitedCount(Context.PublicData.InstructionIndex));
 				CurrentBreakpointAction = ERigVMBreakpointAction::None;					
 
 				HaltedAtBreakpoint = NewBreakpoint;
 				HaltedAtBreakpointHit = DebugInfo->GetBreakpointHits(HaltedAtBreakpoint);
-				ExecutionHalted().Broadcast(Context.InstructionIndex, NewBreakpointNode, InEventName);
+				ExecutionHalted().Broadcast(Context.PublicData.InstructionIndex, NewBreakpointNode, InEventName);
 		
 				return true;
 			}
@@ -1338,9 +1338,9 @@ bool URigVM::Initialize(TArrayView<URigVMMemoryStorage*> Memory, TArrayView<void
 
 	TGuardValue<URigVM*> VMInContext(Context.VM, this);
 	
-	while (Instructions.IsValidIndex(Context.InstructionIndex))
+	while (Instructions.IsValidIndex(Context.PublicData.InstructionIndex))
 	{
-		const FRigVMInstruction& Instruction = Instructions[Context.InstructionIndex];
+		const FRigVMInstruction& Instruction = Instructions[Context.PublicData.InstructionIndex];
 
 		switch (Instruction.OpCode)
 		{
@@ -1411,10 +1411,10 @@ bool URigVM::Initialize(TArrayView<URigVMMemoryStorage*> Memory, TArrayView<void
 			case ERigVMOpCode::Execute_64_Operands:
 			{
 				const FRigVMExecuteOp& Op = ByteCode.GetOpAt<FRigVMExecuteOp>(Instruction);
-				int32 OperandCount = FirstHandleForInstruction[Context.InstructionIndex + 1] - FirstHandleForInstruction[Context.InstructionIndex];
-				FRigVMMemoryHandleArray OpHandles(&CachedMemoryHandles[FirstHandleForInstruction[Context.InstructionIndex]], OperandCount);
+				int32 OperandCount = FirstHandleForInstruction[Context.PublicData.InstructionIndex + 1] - FirstHandleForInstruction[Context.PublicData.InstructionIndex];
+				FRigVMMemoryHandleArray OpHandles(&CachedMemoryHandles[FirstHandleForInstruction[Context.PublicData.InstructionIndex]], OperandCount);
 #if WITH_EDITOR
-				Context.FunctionName = FunctionNames[Op.FunctionIndex];
+				Context.PublicData.FunctionName = FunctionNames[Op.FunctionIndex];
 #endif
 
 				// find out the largest slice count
@@ -1442,8 +1442,8 @@ bool URigVM::Initialize(TArrayView<URigVMMemoryStorage*> Memory, TArrayView<void
 			{
 				const FRigVMCopyOp& Op = ByteCode.GetOpAt<FRigVMCopyOp>(Instruction);
 
-				FRigVMMemoryHandle& SourceHandle = CachedMemoryHandles[FirstHandleForInstruction[Context.InstructionIndex]];
-				FRigVMMemoryHandle& TargetHandle = CachedMemoryHandles[FirstHandleForInstruction[Context.InstructionIndex] + 1];
+				FRigVMMemoryHandle& SourceHandle = CachedMemoryHandles[FirstHandleForInstruction[Context.PublicData.InstructionIndex]];
+				FRigVMMemoryHandle& TargetHandle = CachedMemoryHandles[FirstHandleForInstruction[Context.PublicData.InstructionIndex] + 1];
 				URigVMMemoryStorage::CopyProperty(TargetHandle, SourceHandle);
 				break;
 			}
@@ -1486,7 +1486,7 @@ bool URigVM::Initialize(TArrayView<URigVMMemoryStorage*> Memory, TArrayView<void
 				return false;
 			}
 		}
-		Context.InstructionIndex++;
+		Context.PublicData.InstructionIndex++;
 	}
 	
 	return true;
@@ -1530,7 +1530,7 @@ bool URigVM::Execute(TArrayView<URigVMMemoryStorage*> Memory, TArrayView<void*>
 		InstructionVisitOrder.Reset();
 		InstructionVisitedDuringLastRun.SetNumZeroed(Instructions.Num());
 		InstructionCyclesDuringLastRun.Reset();
-		if(Context.RuntimeSettings.bEnableProfiling)
+		if(Context.PublicData.RuntimeSettings.bEnableProfiling)
 		{
 			InstructionCyclesDuringLastRun.SetNumUninitialized(Instructions.Num());
 			for(int32 DurationIndex=0;DurationIndex<InstructionCyclesDuringLastRun.Num();DurationIndex++)
@@ -1556,7 +1556,7 @@ bool URigVM::Execute(TArrayView<URigVMMemoryStorage*> Memory, TArrayView<void*>
 		{
 			return false;
 		}
-		Context.InstructionIndex = (uint16)ByteCode.GetEntry(EntryIndex).InstructionIndex;
+		Context.PublicData.InstructionIndex = (uint16)ByteCode.GetEntry(EntryIndex).InstructionIndex;
 	}
 
 #if WITH_EDITOR
@@ -1571,27 +1571,27 @@ bool URigVM::Execute(TArrayView<URigVMMemoryStorage*> Memory, TArrayView<void*>
 #if WITH_EDITOR
 	uint64 StartCycles = 0;
 	uint64 OverallCycles = 0;
-	if(Context.RuntimeSettings.bEnableProfiling)
+	if(Context.PublicData.RuntimeSettings.bEnableProfiling)
 	{
 		StartCycles = FPlatformTime::Cycles64();
 	}
 #endif
 
-	while (Instructions.IsValidIndex(Context.InstructionIndex))
+	while (Instructions.IsValidIndex(Context.PublicData.InstructionIndex))
 	{
 #if WITH_EDITOR
-		if (DebugInfo && ShouldHaltAtInstruction(InEntryName, Context.InstructionIndex))
+		if (DebugInfo && ShouldHaltAtInstruction(InEntryName, Context.PublicData.InstructionIndex))
 		{
 			return true;
 		}
 
-		const int32 CurrentInstructionIndex = Context.InstructionIndex;
-		InstructionVisitedDuringLastRun[Context.InstructionIndex]++;
-		InstructionVisitOrder.Add(Context.InstructionIndex);
+		const int32 CurrentInstructionIndex = Context.PublicData.InstructionIndex;
+		InstructionVisitedDuringLastRun[Context.PublicData.InstructionIndex]++;
+		InstructionVisitOrder.Add(Context.PublicData.InstructionIndex);
 	
 #endif
 
-		const FRigVMInstruction& Instruction = Instructions[Context.InstructionIndex];
+		const FRigVMInstruction& Instruction = Instructions[Context.PublicData.InstructionIndex];
 
 		switch (Instruction.OpCode)
 		{
@@ -1662,10 +1662,10 @@ bool URigVM::Execute(TArrayView<URigVMMemoryStorage*> Memory, TArrayView<void*>
 			case ERigVMOpCode::Execute_64_Operands:
 			{
 				const FRigVMExecuteOp& Op = ByteCode.GetOpAt<FRigVMExecuteOp>(Instruction);
-				const int32 OperandCount = FirstHandleForInstruction[Context.InstructionIndex + 1] - FirstHandleForInstruction[Context.InstructionIndex];
-				FRigVMMemoryHandleArray Handles(&CachedMemoryHandles[FirstHandleForInstruction[Context.InstructionIndex]], OperandCount);
+				const int32 OperandCount = FirstHandleForInstruction[Context.PublicData.InstructionIndex + 1] - FirstHandleForInstruction[Context.PublicData.InstructionIndex];
+				FRigVMMemoryHandleArray Handles(&CachedMemoryHandles[FirstHandleForInstruction[Context.PublicData.InstructionIndex]], OperandCount);
 #if WITH_EDITOR
-				Context.FunctionName = FunctionNames[Op.FunctionIndex];
+				Context.PublicData.FunctionName = FunctionNames[Op.FunctionIndex];
 #endif
 				(*Functions[Op.FunctionIndex])(Context, Handles);
 
@@ -1681,41 +1681,41 @@ bool URigVM::Execute(TArrayView<URigVMMemoryStorage*> Memory, TArrayView<void*>
 				}
 #endif
 
-				Context.InstructionIndex++;
+				Context.PublicData.InstructionIndex++;
 				break;
 			}
 			case ERigVMOpCode::Zero:
 			{
-				*((int32*)CachedMemoryHandles[FirstHandleForInstruction[Context.InstructionIndex]].GetData()) = 0;
+				*((int32*)CachedMemoryHandles[FirstHandleForInstruction[Context.PublicData.InstructionIndex]].GetData()) = 0;
 #if WITH_EDITOR
 				if(DebugMemoryStorageObject->Num() > 0)
 				{
 					const FRigVMUnaryOp& Op = ByteCode.GetOpAt<FRigVMUnaryOp>(Instruction);
-					CopyOperandForDebuggingIfNeeded(Op.Arg, CachedMemoryHandles[FirstHandleForInstruction[Context.InstructionIndex]]);
+					CopyOperandForDebuggingIfNeeded(Op.Arg, CachedMemoryHandles[FirstHandleForInstruction[Context.PublicData.InstructionIndex]]);
 				}
 #endif
 
-				Context.InstructionIndex++;
+				Context.PublicData.InstructionIndex++;
 				break;
 			}
 			case ERigVMOpCode::BoolFalse:
 			{
-				*((bool*)CachedMemoryHandles[FirstHandleForInstruction[Context.InstructionIndex]].GetData()) = false;
-				Context.InstructionIndex++;
+				*((bool*)CachedMemoryHandles[FirstHandleForInstruction[Context.PublicData.InstructionIndex]].GetData()) = false;
+				Context.PublicData.InstructionIndex++;
 				break;
 			}
 			case ERigVMOpCode::BoolTrue:
 			{
-				*((bool*)CachedMemoryHandles[FirstHandleForInstruction[Context.InstructionIndex]].GetData()) = true;
-				Context.InstructionIndex++;
+				*((bool*)CachedMemoryHandles[FirstHandleForInstruction[Context.PublicData.InstructionIndex]].GetData()) = true;
+				Context.PublicData.InstructionIndex++;
 				break;
 			}
 			case ERigVMOpCode::Copy:
 			{
 				const FRigVMCopyOp& Op = ByteCode.GetOpAt<FRigVMCopyOp>(Instruction);
 
-				FRigVMMemoryHandle& SourceHandle = CachedMemoryHandles[FirstHandleForInstruction[Context.InstructionIndex]];
-				FRigVMMemoryHandle& TargetHandle = CachedMemoryHandles[FirstHandleForInstruction[Context.InstructionIndex] + 1];
+				FRigVMMemoryHandle& SourceHandle = CachedMemoryHandles[FirstHandleForInstruction[Context.PublicData.InstructionIndex]];
+				FRigVMMemoryHandle& TargetHandle = CachedMemoryHandles[FirstHandleForInstruction[Context.PublicData.InstructionIndex] + 1];
 				URigVMMemoryStorage::CopyProperty(TargetHandle, SourceHandle);
 					
 #if WITH_EDITOR
@@ -1725,33 +1725,33 @@ bool URigVM::Execute(TArrayView<URigVMMemoryStorage*> Memory, TArrayView<void*>
 				}
 #endif
 					
-				Context.InstructionIndex++;
+				Context.PublicData.InstructionIndex++;
 				break;
 			}
 			case ERigVMOpCode::Increment:
 			{
-				(*((int32*)CachedMemoryHandles[FirstHandleForInstruction[Context.InstructionIndex]].GetData()))++;
+				(*((int32*)CachedMemoryHandles[FirstHandleForInstruction[Context.PublicData.InstructionIndex]].GetData()))++;
 #if WITH_EDITOR
 				if(DebugMemoryStorageObject->Num() > 0)
 				{
 					const FRigVMUnaryOp& Op = ByteCode.GetOpAt<FRigVMUnaryOp>(Instruction);
-					CopyOperandForDebuggingIfNeeded(Op.Arg, CachedMemoryHandles[FirstHandleForInstruction[Context.InstructionIndex]]);
+					CopyOperandForDebuggingIfNeeded(Op.Arg, CachedMemoryHandles[FirstHandleForInstruction[Context.PublicData.InstructionIndex]]);
 				}
 #endif
-				Context.InstructionIndex++;
+				Context.PublicData.InstructionIndex++;
 				break;
 			}
 			case ERigVMOpCode::Decrement:
 			{
-				(*((int32*)CachedMemoryHandles[FirstHandleForInstruction[Context.InstructionIndex]].GetData()))--;
+				(*((int32*)CachedMemoryHandles[FirstHandleForInstruction[Context.PublicData.InstructionIndex]].GetData()))--;
 #if WITH_EDITOR
 				if(DebugMemoryStorageObject->Num() > 0)
 				{
 					const FRigVMUnaryOp& Op = ByteCode.GetOpAt<FRigVMUnaryOp>(Instruction);
-					CopyOperandForDebuggingIfNeeded(Op.Arg, CachedMemoryHandles[FirstHandleForInstruction[Context.InstructionIndex]]);
+					CopyOperandForDebuggingIfNeeded(Op.Arg, CachedMemoryHandles[FirstHandleForInstruction[Context.PublicData.InstructionIndex]]);
 				}
 #endif
-				Context.InstructionIndex++;
+				Context.PublicData.InstructionIndex++;
 				break;
 			}
 			case ERigVMOpCode::Equals:
@@ -1759,71 +1759,71 @@ bool URigVM::Execute(TArrayView<URigVMMemoryStorage*> Memory, TArrayView<void*>
 			{
 				const FRigVMComparisonOp& Op = ByteCode.GetOpAt<FRigVMComparisonOp>(Instruction);
 
-				FRigVMMemoryHandle& HandleA = CachedMemoryHandles[FirstHandleForInstruction[Context.InstructionIndex]];
-				FRigVMMemoryHandle& HandleB = CachedMemoryHandles[FirstHandleForInstruction[Context.InstructionIndex] + 1];
+				FRigVMMemoryHandle& HandleA = CachedMemoryHandles[FirstHandleForInstruction[Context.PublicData.InstructionIndex]];
+				FRigVMMemoryHandle& HandleB = CachedMemoryHandles[FirstHandleForInstruction[Context.PublicData.InstructionIndex] + 1];
 				const bool Result = HandleA.GetProperty()->Identical(HandleA.GetData(true), HandleB.GetData(true));
 
-				*((bool*)CachedMemoryHandles[FirstHandleForInstruction[Context.InstructionIndex]+2].GetData()) = Result;
-				Context.InstructionIndex++;
+				*((bool*)CachedMemoryHandles[FirstHandleForInstruction[Context.PublicData.InstructionIndex]+2].GetData()) = Result;
+				Context.PublicData.InstructionIndex++;
 				break;
 			}
 			case ERigVMOpCode::JumpAbsolute:
 			{
 				const FRigVMJumpOp& Op = ByteCode.GetOpAt<FRigVMJumpOp>(Instruction);
-				Context.InstructionIndex = Op.InstructionIndex;
+				Context.PublicData.InstructionIndex = Op.InstructionIndex;
 				break;
 			}
 			case ERigVMOpCode::JumpForward:
 			{
 				const FRigVMJumpOp& Op = ByteCode.GetOpAt<FRigVMJumpOp>(Instruction);
-				Context.InstructionIndex += Op.InstructionIndex;
+				Context.PublicData.InstructionIndex += Op.InstructionIndex;
 				break;
 			}
 			case ERigVMOpCode::JumpBackward:
 			{
 				const FRigVMJumpOp& Op = ByteCode.GetOpAt<FRigVMJumpOp>(Instruction);
-				Context.InstructionIndex -= Op.InstructionIndex;
+				Context.PublicData.InstructionIndex -= Op.InstructionIndex;
 				break;
 			}
 			case ERigVMOpCode::JumpAbsoluteIf:
 			{
 				const FRigVMJumpIfOp& Op = ByteCode.GetOpAt<FRigVMJumpIfOp>(Instruction);
-				const bool Condition = *(bool*)CachedMemoryHandles[FirstHandleForInstruction[Context.InstructionIndex]].GetData();
+				const bool Condition = *(bool*)CachedMemoryHandles[FirstHandleForInstruction[Context.PublicData.InstructionIndex]].GetData();
 				if (Condition == Op.Condition)
 				{
-					Context.InstructionIndex = Op.InstructionIndex;
+					Context.PublicData.InstructionIndex = Op.InstructionIndex;
 				}
 				else
 				{
-					Context.InstructionIndex++;
+					Context.PublicData.InstructionIndex++;
 				}
 				break;
 			}
 			case ERigVMOpCode::JumpForwardIf:
 			{
 				const FRigVMJumpIfOp& Op = ByteCode.GetOpAt<FRigVMJumpIfOp>(Instruction);
-				const bool Condition = *(bool*)CachedMemoryHandles[FirstHandleForInstruction[Context.InstructionIndex]].GetData();
+				const bool Condition = *(bool*)CachedMemoryHandles[FirstHandleForInstruction[Context.PublicData.InstructionIndex]].GetData();
 				if (Condition == Op.Condition)
 				{
-					Context.InstructionIndex += Op.InstructionIndex;
+					Context.PublicData.InstructionIndex += Op.InstructionIndex;
 				}
 				else
 				{
-					Context.InstructionIndex++;
+					Context.PublicData.InstructionIndex++;
 				}
 				break;
 			}
 			case ERigVMOpCode::JumpBackwardIf:
 			{
 				const FRigVMJumpIfOp& Op = ByteCode.GetOpAt<FRigVMJumpIfOp>(Instruction);
-				const bool Condition = *(bool*)CachedMemoryHandles[FirstHandleForInstruction[Context.InstructionIndex]].GetData();
+				const bool Condition = *(bool*)CachedMemoryHandles[FirstHandleForInstruction[Context.PublicData.InstructionIndex]].GetData();
 				if (Condition == Op.Condition)
 				{
-					Context.InstructionIndex -= Op.InstructionIndex;
+					Context.PublicData.InstructionIndex -= Op.InstructionIndex;
 				}
 				else
 				{
-					Context.InstructionIndex++;
+					Context.PublicData.InstructionIndex++;
 				}
 				break;
 			}
@@ -1850,21 +1850,21 @@ bool URigVM::Execute(TArrayView<URigVMMemoryStorage*> Memory, TArrayView<void*>
 			}
 			case ERigVMOpCode::BeginBlock:
 			{
-				const int32 Count = (*((int32*)CachedMemoryHandles[FirstHandleForInstruction[Context.InstructionIndex]].GetData()));
-				const int32 Index = (*((int32*)CachedMemoryHandles[FirstHandleForInstruction[Context.InstructionIndex] + 1].GetData()));
+				const int32 Count = (*((int32*)CachedMemoryHandles[FirstHandleForInstruction[Context.PublicData.InstructionIndex]].GetData()));
+				const int32 Index = (*((int32*)CachedMemoryHandles[FirstHandleForInstruction[Context.PublicData.InstructionIndex] + 1].GetData()));
 				Context.BeginSlice(Count, Index);
-				Context.InstructionIndex++;
+				Context.PublicData.InstructionIndex++;
 				break;
 			}
 			case ERigVMOpCode::EndBlock:
 			{
 				Context.EndSlice();
-				Context.InstructionIndex++;
+				Context.PublicData.InstructionIndex++;
 				break;
 			}
 			case ERigVMOpCode::ArrayReset:
 			{
-				FRigVMMemoryHandle& ArrayHandle = CachedMemoryHandles[FirstHandleForInstruction[Context.InstructionIndex]]; 					
+				FRigVMMemoryHandle& ArrayHandle = CachedMemoryHandles[FirstHandleForInstruction[Context.PublicData.InstructionIndex]]; 					
 				FScriptArrayHelper ArrayHelper(CastFieldChecked<FArrayProperty>(ArrayHandle.GetProperty()), ArrayHandle.GetData());
 				ArrayHelper.Resize(0);
 
@@ -1874,31 +1874,31 @@ bool URigVM::Execute(TArrayView<URigVMMemoryStorage*> Memory, TArrayView<void*>
 					CopyOperandForDebuggingIfNeeded(Op.Arg, ArrayHandle);
 				}
 
-				Context.InstructionIndex++;
+				Context.PublicData.InstructionIndex++;
 				break;
 			}
 			case ERigVMOpCode::ArrayGetNum:
 			{
-				FRigVMMemoryHandle& ArrayHandle = CachedMemoryHandles[FirstHandleForInstruction[Context.InstructionIndex]]; 					
+				FRigVMMemoryHandle& ArrayHandle = CachedMemoryHandles[FirstHandleForInstruction[Context.PublicData.InstructionIndex]]; 					
 				FScriptArrayHelper ArrayHelper(CastFieldChecked<FArrayProperty>(ArrayHandle.GetProperty()), ArrayHandle.GetData());
-				int32& Count = (*((int32*)CachedMemoryHandles[FirstHandleForInstruction[Context.InstructionIndex] + 1].GetData()));
+				int32& Count = (*((int32*)CachedMemoryHandles[FirstHandleForInstruction[Context.PublicData.InstructionIndex] + 1].GetData()));
 				Count = ArrayHelper.Num();
 
 				if(DebugMemoryStorageObject->Num() > 0)
 				{
 					const FRigVMBinaryOp& Op = ByteCode.GetOpAt<FRigVMBinaryOp>(Instruction);
 					CopyOperandForDebuggingIfNeeded(Op.ArgA, ArrayHandle);
-					CopyOperandForDebuggingIfNeeded(Op.ArgB, CachedMemoryHandles[FirstHandleForInstruction[Context.InstructionIndex] + 1]);
+					CopyOperandForDebuggingIfNeeded(Op.ArgB, CachedMemoryHandles[FirstHandleForInstruction[Context.PublicData.InstructionIndex] + 1]);
 				}
 
-				Context.InstructionIndex++;
+				Context.PublicData.InstructionIndex++;
 				break;
 			}
 			case ERigVMOpCode::ArraySetNum:
 			{
-				FRigVMMemoryHandle& ArrayHandle = CachedMemoryHandles[FirstHandleForInstruction[Context.InstructionIndex]]; 					
+				FRigVMMemoryHandle& ArrayHandle = CachedMemoryHandles[FirstHandleForInstruction[Context.PublicData.InstructionIndex]]; 					
 				FScriptArrayHelper ArrayHelper(CastFieldChecked<FArrayProperty>(ArrayHandle.GetProperty()), ArrayHandle.GetData());
-				const int32 Count = (*((int32*)CachedMemoryHandles[FirstHandleForInstruction[Context.InstructionIndex] + 1].GetData()));
+				const int32 Count = (*((int32*)CachedMemoryHandles[FirstHandleForInstruction[Context.PublicData.InstructionIndex] + 1].GetData()));
 				if(Context.IsValidArraySize(Count))
 				{
 					ArrayHelper.Resize(Count);
@@ -1908,16 +1908,16 @@ bool URigVM::Execute(TArrayView<URigVMMemoryStorage*> Memory, TArrayView<void*>
 				{
 					const FRigVMBinaryOp& Op = ByteCode.GetOpAt<FRigVMBinaryOp>(Instruction);
 					CopyOperandForDebuggingIfNeeded(Op.ArgA, ArrayHandle);
-					CopyOperandForDebuggingIfNeeded(Op.ArgB, CachedMemoryHandles[FirstHandleForInstruction[Context.InstructionIndex] + 1]);
+					CopyOperandForDebuggingIfNeeded(Op.ArgB, CachedMemoryHandles[FirstHandleForInstruction[Context.PublicData.InstructionIndex] + 1]);
 				}
 
-				Context.InstructionIndex++;
+				Context.PublicData.InstructionIndex++;
 				break;
 			}
 			case ERigVMOpCode::ArrayAppend:
 			{
-				FRigVMMemoryHandle& ArrayHandle = CachedMemoryHandles[FirstHandleForInstruction[Context.InstructionIndex]]; 					
-				FRigVMMemoryHandle& OtherArrayHandle = CachedMemoryHandles[FirstHandleForInstruction[Context.InstructionIndex] + 1];
+				FRigVMMemoryHandle& ArrayHandle = CachedMemoryHandles[FirstHandleForInstruction[Context.PublicData.InstructionIndex]]; 					
+				FRigVMMemoryHandle& OtherArrayHandle = CachedMemoryHandles[FirstHandleForInstruction[Context.PublicData.InstructionIndex] + 1];
 				const FArrayProperty* ArrayProperty = CastFieldChecked<FArrayProperty>(ArrayHandle.GetProperty()); 
 				const FArrayProperty* OtherArrayProperty = CastFieldChecked<FArrayProperty>(OtherArrayHandle.GetProperty()); 
 
@@ -1945,16 +1945,16 @@ bool URigVM::Execute(TArrayView<URigVMMemoryStorage*> Memory, TArrayView<void*>
 				{
 					const FRigVMBinaryOp& Op = ByteCode.GetOpAt<FRigVMBinaryOp>(Instruction);
 					CopyOperandForDebuggingIfNeeded(Op.ArgA, ArrayHandle);
-					CopyOperandForDebuggingIfNeeded(Op.ArgB, CachedMemoryHandles[FirstHandleForInstruction[Context.InstructionIndex] + 1]);
+					CopyOperandForDebuggingIfNeeded(Op.ArgB, CachedMemoryHandles[FirstHandleForInstruction[Context.PublicData.InstructionIndex] + 1]);
 				}
 
-				Context.InstructionIndex++;
+				Context.PublicData.InstructionIndex++;
 				break;
 			}
 			case ERigVMOpCode::ArrayClone:
 			{
-				FRigVMMemoryHandle& ArrayHandle = CachedMemoryHandles[FirstHandleForInstruction[Context.InstructionIndex]]; 					
-				FRigVMMemoryHandle& ClonedArrayHandle = CachedMemoryHandles[FirstHandleForInstruction[Context.InstructionIndex] + 1];
+				FRigVMMemoryHandle& ArrayHandle = CachedMemoryHandles[FirstHandleForInstruction[Context.PublicData.InstructionIndex]]; 					
+				FRigVMMemoryHandle& ClonedArrayHandle = CachedMemoryHandles[FirstHandleForInstruction[Context.PublicData.InstructionIndex] + 1];
 				const FArrayProperty* ArrayProperty = CastFieldChecked<FArrayProperty>(ArrayHandle.GetProperty()); 
 				const FArrayProperty* ClonedArrayProperty = CastFieldChecked<FArrayProperty>(ClonedArrayHandle.GetProperty()); 
 				FScriptArrayHelper ArrayHelper(ArrayProperty, ArrayHandle.GetData());
@@ -1969,18 +1969,18 @@ bool URigVM::Execute(TArrayView<URigVMMemoryStorage*> Memory, TArrayView<void*>
 					CopyOperandForDebuggingIfNeeded(Op.ArgB, ClonedArrayHandle);
 				}
 
-				Context.InstructionIndex++;
+				Context.PublicData.InstructionIndex++;
 				break;
 			}
 			case ERigVMOpCode::ArrayGetAtIndex:
 			{
-				FRigVMMemoryHandle& ArrayHandle = CachedMemoryHandles[FirstHandleForInstruction[Context.InstructionIndex]];
+				FRigVMMemoryHandle& ArrayHandle = CachedMemoryHandles[FirstHandleForInstruction[Context.PublicData.InstructionIndex]];
 				const FArrayProperty* ArrayProperty = CastFieldChecked<FArrayProperty>(ArrayHandle.GetProperty());
 				FScriptArrayHelper ArrayHelper(ArrayProperty, ArrayHandle.GetData());
-				const int32 Index = (*((int32*)CachedMemoryHandles[FirstHandleForInstruction[Context.InstructionIndex] + 1].GetData()));
+				const int32 Index = (*((int32*)CachedMemoryHandles[FirstHandleForInstruction[Context.PublicData.InstructionIndex] + 1].GetData()));
 				if(Context.IsValidArrayIndex(Index, ArrayHelper.Num()))
 				{
-					FRigVMMemoryHandle& ElementHandle = CachedMemoryHandles[FirstHandleForInstruction[Context.InstructionIndex] + 2];
+					FRigVMMemoryHandle& ElementHandle = CachedMemoryHandles[FirstHandleForInstruction[Context.PublicData.InstructionIndex] + 2];
 					uint8* TargetMemory = ElementHandle.GetData();
 					const uint8* SourceMemory = ArrayHelper.GetRawPtr(Index);
 					URigVMMemoryStorage::CopyProperty(ElementHandle.GetProperty(), TargetMemory, ArrayProperty->Inner, SourceMemory);
@@ -1990,22 +1990,22 @@ bool URigVM::Execute(TArrayView<URigVMMemoryStorage*> Memory, TArrayView<void*>
 				{
 					const FRigVMTernaryOp& Op = ByteCode.GetOpAt<FRigVMTernaryOp>(Instruction);
 					CopyOperandForDebuggingIfNeeded(Op.ArgA, ArrayHandle);
-					CopyOperandForDebuggingIfNeeded(Op.ArgB, CachedMemoryHandles[FirstHandleForInstruction[Context.InstructionIndex] + 1]);
-					CopyOperandForDebuggingIfNeeded(Op.ArgC, CachedMemoryHandles[FirstHandleForInstruction[Context.InstructionIndex] + 2]);
+					CopyOperandForDebuggingIfNeeded(Op.ArgB, CachedMemoryHandles[FirstHandleForInstruction[Context.PublicData.InstructionIndex] + 1]);
+					CopyOperandForDebuggingIfNeeded(Op.ArgC, CachedMemoryHandles[FirstHandleForInstruction[Context.PublicData.InstructionIndex] + 2]);
 				}
 
-				Context.InstructionIndex++;
+				Context.PublicData.InstructionIndex++;
 				break;
 			}
 			case ERigVMOpCode::ArraySetAtIndex:
 			{
-				FRigVMMemoryHandle& ArrayHandle = CachedMemoryHandles[FirstHandleForInstruction[Context.InstructionIndex]];
+				FRigVMMemoryHandle& ArrayHandle = CachedMemoryHandles[FirstHandleForInstruction[Context.PublicData.InstructionIndex]];
 				const FArrayProperty* ArrayProperty = CastFieldChecked<FArrayProperty>(ArrayHandle.GetProperty());
 				FScriptArrayHelper ArrayHelper(ArrayProperty, ArrayHandle.GetData());
-				const int32 Index = (*((int32*)CachedMemoryHandles[FirstHandleForInstruction[Context.InstructionIndex] + 1].GetData()));
+				const int32 Index = (*((int32*)CachedMemoryHandles[FirstHandleForInstruction[Context.PublicData.InstructionIndex] + 1].GetData()));
 				if(Context.IsValidArrayIndex(Index, ArrayHelper.Num()))
 				{
-					FRigVMMemoryHandle& ElementHandle = CachedMemoryHandles[FirstHandleForInstruction[Context.InstructionIndex] + 2];
+					FRigVMMemoryHandle& ElementHandle = CachedMemoryHandles[FirstHandleForInstruction[Context.PublicData.InstructionIndex] + 2];
 					uint8* TargetMemory = ArrayHelper.GetRawPtr(Index);
 					const uint8* SourceMemory = ElementHandle.GetData();
 					URigVMMemoryStorage::CopyProperty(ArrayProperty->Inner, TargetMemory, ElementHandle.GetProperty(), SourceMemory);
@@ -2015,25 +2015,25 @@ bool URigVM::Execute(TArrayView<URigVMMemoryStorage*> Memory, TArrayView<void*>
 				{
 					const FRigVMTernaryOp& Op = ByteCode.GetOpAt<FRigVMTernaryOp>(Instruction);
 					CopyOperandForDebuggingIfNeeded(Op.ArgA, ArrayHandle);
-					CopyOperandForDebuggingIfNeeded(Op.ArgB, CachedMemoryHandles[FirstHandleForInstruction[Context.InstructionIndex] + 1]);
-					CopyOperandForDebuggingIfNeeded(Op.ArgC, CachedMemoryHandles[FirstHandleForInstruction[Context.InstructionIndex] + 2]);
+					CopyOperandForDebuggingIfNeeded(Op.ArgB, CachedMemoryHandles[FirstHandleForInstruction[Context.PublicData.InstructionIndex] + 1]);
+					CopyOperandForDebuggingIfNeeded(Op.ArgC, CachedMemoryHandles[FirstHandleForInstruction[Context.PublicData.InstructionIndex] + 2]);
 				}
 
-				Context.InstructionIndex++;
+				Context.PublicData.InstructionIndex++;
 				break;
 			}
 			case ERigVMOpCode::ArrayInsert:
 			{
-				FRigVMMemoryHandle& ArrayHandle = CachedMemoryHandles[FirstHandleForInstruction[Context.InstructionIndex]];
+				FRigVMMemoryHandle& ArrayHandle = CachedMemoryHandles[FirstHandleForInstruction[Context.PublicData.InstructionIndex]];
 				const FArrayProperty* ArrayProperty = CastFieldChecked<FArrayProperty>(ArrayHandle.GetProperty());
 				FScriptArrayHelper ArrayHelper(ArrayProperty, ArrayHandle.GetData());
 				if(Context.IsValidArraySize(ArrayHelper.Num() + 1))
 				{
-					int32 Index = (*((int32*)CachedMemoryHandles[FirstHandleForInstruction[Context.InstructionIndex] + 1].GetData()));
+					int32 Index = (*((int32*)CachedMemoryHandles[FirstHandleForInstruction[Context.PublicData.InstructionIndex] + 1].GetData()));
 					Index = FMath::Clamp<int32>(Index, 0, ArrayHelper.Num());
 					ArrayHelper.InsertValues(Index, 1);
 
-					FRigVMMemoryHandle& ElementHandle = CachedMemoryHandles[FirstHandleForInstruction[Context.InstructionIndex] + 2];
+					FRigVMMemoryHandle& ElementHandle = CachedMemoryHandles[FirstHandleForInstruction[Context.PublicData.InstructionIndex] + 2];
 					uint8* TargetMemory = ArrayHelper.GetRawPtr(Index);
 					const uint8* SourceMemory = ElementHandle.GetData();
 					URigVMMemoryStorage::CopyProperty(ArrayProperty->Inner, TargetMemory, ElementHandle.GetProperty(), SourceMemory);
@@ -2043,19 +2043,19 @@ bool URigVM::Execute(TArrayView<URigVMMemoryStorage*> Memory, TArrayView<void*>
 				{
 					const FRigVMTernaryOp& Op = ByteCode.GetOpAt<FRigVMTernaryOp>(Instruction);
 					CopyOperandForDebuggingIfNeeded(Op.ArgA, ArrayHandle);
-					CopyOperandForDebuggingIfNeeded(Op.ArgB, CachedMemoryHandles[FirstHandleForInstruction[Context.InstructionIndex] + 1]);
-					CopyOperandForDebuggingIfNeeded(Op.ArgC, CachedMemoryHandles[FirstHandleForInstruction[Context.InstructionIndex] + 2]);
+					CopyOperandForDebuggingIfNeeded(Op.ArgB, CachedMemoryHandles[FirstHandleForInstruction[Context.PublicData.InstructionIndex] + 1]);
+					CopyOperandForDebuggingIfNeeded(Op.ArgC, CachedMemoryHandles[FirstHandleForInstruction[Context.PublicData.InstructionIndex] + 2]);
 				}
 
-				Context.InstructionIndex++;
+				Context.PublicData.InstructionIndex++;
 				break;
 			}
 			case ERigVMOpCode::ArrayRemove:
 			{
-				FRigVMMemoryHandle& ArrayHandle = CachedMemoryHandles[FirstHandleForInstruction[Context.InstructionIndex]];
+				FRigVMMemoryHandle& ArrayHandle = CachedMemoryHandles[FirstHandleForInstruction[Context.PublicData.InstructionIndex]];
 				const FArrayProperty* ArrayProperty = CastFieldChecked<FArrayProperty>(ArrayHandle.GetProperty());
 				FScriptArrayHelper ArrayHelper(ArrayProperty, ArrayHandle.GetData());
-				const int32 Index = (*((int32*)CachedMemoryHandles[FirstHandleForInstruction[Context.InstructionIndex] + 1].GetData()));
+				const int32 Index = (*((int32*)CachedMemoryHandles[FirstHandleForInstruction[Context.PublicData.InstructionIndex] + 1].GetData()));
 				if(Context.IsValidArrayIndex(Index, ArrayHelper.Num()))
 				{
 					ArrayHelper.RemoveValues(Index, 1);
@@ -2065,21 +2065,21 @@ bool URigVM::Execute(TArrayView<URigVMMemoryStorage*> Memory, TArrayView<void*>
 				{
 					const FRigVMBinaryOp& Op = ByteCode.GetOpAt<FRigVMBinaryOp>(Instruction);
 					CopyOperandForDebuggingIfNeeded(Op.ArgA, ArrayHandle);
-					CopyOperandForDebuggingIfNeeded(Op.ArgB, CachedMemoryHandles[FirstHandleForInstruction[Context.InstructionIndex] + 1]);
+					CopyOperandForDebuggingIfNeeded(Op.ArgB, CachedMemoryHandles[FirstHandleForInstruction[Context.PublicData.InstructionIndex] + 1]);
 				}
 
-				Context.InstructionIndex++;
+				Context.PublicData.InstructionIndex++;
 				break;
 			}
 			case ERigVMOpCode::ArrayAdd:
 			{
-				FRigVMMemoryHandle& ArrayHandle = CachedMemoryHandles[FirstHandleForInstruction[Context.InstructionIndex]];
+				FRigVMMemoryHandle& ArrayHandle = CachedMemoryHandles[FirstHandleForInstruction[Context.PublicData.InstructionIndex]];
 				const FArrayProperty* ArrayProperty = CastFieldChecked<FArrayProperty>(ArrayHandle.GetProperty());
 				FScriptArrayHelper ArrayHelper(ArrayProperty, ArrayHandle.GetData());
-				int32& Index = (*((int32*)CachedMemoryHandles[FirstHandleForInstruction[Context.InstructionIndex] + 2].GetData()));
+				int32& Index = (*((int32*)CachedMemoryHandles[FirstHandleForInstruction[Context.PublicData.InstructionIndex] + 2].GetData()));
 				if(Context.IsValidArraySize(ArrayHelper.Num() + 1))
 				{
-					FRigVMMemoryHandle& ElementHandle = CachedMemoryHandles[FirstHandleForInstruction[Context.InstructionIndex] + 1];
+					FRigVMMemoryHandle& ElementHandle = CachedMemoryHandles[FirstHandleForInstruction[Context.PublicData.InstructionIndex] + 1];
 					Index = ArrayHelper.AddValue();
 
 					uint8* TargetMemory = ArrayHelper.GetRawPtr(Index);
@@ -2095,22 +2095,22 @@ bool URigVM::Execute(TArrayView<URigVMMemoryStorage*> Memory, TArrayView<void*>
 				{
 					const FRigVMTernaryOp& Op = ByteCode.GetOpAt<FRigVMTernaryOp>(Instruction);
 					CopyOperandForDebuggingIfNeeded(Op.ArgA, ArrayHandle);
-					CopyOperandForDebuggingIfNeeded(Op.ArgB, CachedMemoryHandles[FirstHandleForInstruction[Context.InstructionIndex] + 1]);
-					CopyOperandForDebuggingIfNeeded(Op.ArgC, CachedMemoryHandles[FirstHandleForInstruction[Context.InstructionIndex] + 2]);
+					CopyOperandForDebuggingIfNeeded(Op.ArgB, CachedMemoryHandles[FirstHandleForInstruction[Context.PublicData.InstructionIndex] + 1]);
+					CopyOperandForDebuggingIfNeeded(Op.ArgC, CachedMemoryHandles[FirstHandleForInstruction[Context.PublicData.InstructionIndex] + 2]);
 				}
 
-				Context.InstructionIndex++;
+				Context.PublicData.InstructionIndex++;
 				break;
 			}
 			case ERigVMOpCode::ArrayFind:
 			{
-				FRigVMMemoryHandle& ArrayHandle = CachedMemoryHandles[FirstHandleForInstruction[Context.InstructionIndex]];
+				FRigVMMemoryHandle& ArrayHandle = CachedMemoryHandles[FirstHandleForInstruction[Context.PublicData.InstructionIndex]];
 				const FArrayProperty* ArrayProperty = CastFieldChecked<FArrayProperty>(ArrayHandle.GetProperty());
 				FScriptArrayHelper ArrayHelper(ArrayProperty, ArrayHandle.GetData());
 
-				FRigVMMemoryHandle& ElementHandle = CachedMemoryHandles[FirstHandleForInstruction[Context.InstructionIndex] + 1];
-				int32& FoundIndex = (*((int32*)CachedMemoryHandles[FirstHandleForInstruction[Context.InstructionIndex] + 2].GetData()));
-				bool& bFound = (*((bool*)CachedMemoryHandles[FirstHandleForInstruction[Context.InstructionIndex] + 3].GetData()));
+				FRigVMMemoryHandle& ElementHandle = CachedMemoryHandles[FirstHandleForInstruction[Context.PublicData.InstructionIndex] + 1];
+				int32& FoundIndex = (*((int32*)CachedMemoryHandles[FirstHandleForInstruction[Context.PublicData.InstructionIndex] + 2].GetData()));
+				bool& bFound = (*((bool*)CachedMemoryHandles[FirstHandleForInstruction[Context.PublicData.InstructionIndex] + 3].GetData()));
 
 				FoundIndex = INDEX_NONE;
 				bFound = false;
@@ -2136,32 +2136,32 @@ bool URigVM::Execute(TArrayView<URigVMMemoryStorage*> Memory, TArrayView<void*>
 				else
 				{
 					static const TCHAR IncompatibleTypes[] = TEXT("Array('%s') doesn't support searching for element('%$s').");
-					Context.Logf(EMessageSeverity::Error, IncompatibleTypes, *PropertyB->GetCPPType(), *PropertyA->GetCPPType());
+					Context.PublicData.Logf(EMessageSeverity::Error, IncompatibleTypes, *PropertyB->GetCPPType(), *PropertyA->GetCPPType());
 				}
 
 				if(DebugMemoryStorageObject->Num() > 0)
 				{
 					const FRigVMQuaternaryOp& Op = ByteCode.GetOpAt<FRigVMQuaternaryOp>(Instruction);
 					CopyOperandForDebuggingIfNeeded(Op.ArgA, ArrayHandle);
-					CopyOperandForDebuggingIfNeeded(Op.ArgB, CachedMemoryHandles[FirstHandleForInstruction[Context.InstructionIndex] + 1]);
-					CopyOperandForDebuggingIfNeeded(Op.ArgC, CachedMemoryHandles[FirstHandleForInstruction[Context.InstructionIndex] + 2]);
-					CopyOperandForDebuggingIfNeeded(Op.ArgD, CachedMemoryHandles[FirstHandleForInstruction[Context.InstructionIndex] + 3]);
+					CopyOperandForDebuggingIfNeeded(Op.ArgB, CachedMemoryHandles[FirstHandleForInstruction[Context.PublicData.InstructionIndex] + 1]);
+					CopyOperandForDebuggingIfNeeded(Op.ArgC, CachedMemoryHandles[FirstHandleForInstruction[Context.PublicData.InstructionIndex] + 2]);
+					CopyOperandForDebuggingIfNeeded(Op.ArgD, CachedMemoryHandles[FirstHandleForInstruction[Context.PublicData.InstructionIndex] + 3]);
 				}
 
-				Context.InstructionIndex++;
+				Context.PublicData.InstructionIndex++;
 				break;
 			}
 			case ERigVMOpCode::ArrayIterator:
 			{
-				FRigVMMemoryHandle& ArrayHandle = CachedMemoryHandles[FirstHandleForInstruction[Context.InstructionIndex]];
+				FRigVMMemoryHandle& ArrayHandle = CachedMemoryHandles[FirstHandleForInstruction[Context.PublicData.InstructionIndex]];
 				const FArrayProperty* ArrayProperty = CastFieldChecked<FArrayProperty>(ArrayHandle.GetProperty());
 				FScriptArrayHelper ArrayHelper(ArrayProperty, ArrayHandle.GetData());
 
-				FRigVMMemoryHandle& ElementHandle = CachedMemoryHandles[FirstHandleForInstruction[Context.InstructionIndex] + 1];
-				const int32& Index = (*((int32*)CachedMemoryHandles[FirstHandleForInstruction[Context.InstructionIndex] + 2].GetData()));
-				int32& Count = (*((int32*)CachedMemoryHandles[FirstHandleForInstruction[Context.InstructionIndex] + 3].GetData()));
-				float& Ratio = (*((float*)CachedMemoryHandles[FirstHandleForInstruction[Context.InstructionIndex] + 4].GetData()));
-				bool& bContinue = (*((bool*)CachedMemoryHandles[FirstHandleForInstruction[Context.InstructionIndex] + 5].GetData()));
+				FRigVMMemoryHandle& ElementHandle = CachedMemoryHandles[FirstHandleForInstruction[Context.PublicData.InstructionIndex] + 1];
+				const int32& Index = (*((int32*)CachedMemoryHandles[FirstHandleForInstruction[Context.PublicData.InstructionIndex] + 2].GetData()));
+				int32& Count = (*((int32*)CachedMemoryHandles[FirstHandleForInstruction[Context.PublicData.InstructionIndex] + 3].GetData()));
+				float& Ratio = (*((float*)CachedMemoryHandles[FirstHandleForInstruction[Context.PublicData.InstructionIndex] + 4].GetData()));
+				bool& bContinue = (*((bool*)CachedMemoryHandles[FirstHandleForInstruction[Context.PublicData.InstructionIndex] + 5].GetData()));
 
 				Count = ArrayHelper.Num();
 				bContinue = Index >=0 && Index < Count;
@@ -2182,25 +2182,25 @@ bool URigVM::Execute(TArrayView<URigVMMemoryStorage*> Memory, TArrayView<void*>
 					{
 						const FRigVMSenaryOp& Op = ByteCode.GetOpAt<FRigVMSenaryOp>(Instruction);
 						CopyOperandForDebuggingIfNeeded(Op.ArgA, ArrayHandle); // array
-						CopyOperandForDebuggingIfNeeded(Op.ArgD, CachedMemoryHandles[FirstHandleForInstruction[Context.InstructionIndex] + 3]); // count
+						CopyOperandForDebuggingIfNeeded(Op.ArgD, CachedMemoryHandles[FirstHandleForInstruction[Context.PublicData.InstructionIndex] + 3]); // count
 
 						Context.BeginSlice(Count, Index);
-						CopyOperandForDebuggingIfNeeded(Op.ArgB, CachedMemoryHandles[FirstHandleForInstruction[Context.InstructionIndex] + 1]); // element
-						CopyOperandForDebuggingIfNeeded(Op.ArgC, CachedMemoryHandles[FirstHandleForInstruction[Context.InstructionIndex] + 2]); // index
-						CopyOperandForDebuggingIfNeeded(Op.ArgE, CachedMemoryHandles[FirstHandleForInstruction[Context.InstructionIndex] + 4]); // ratio
+						CopyOperandForDebuggingIfNeeded(Op.ArgB, CachedMemoryHandles[FirstHandleForInstruction[Context.PublicData.InstructionIndex] + 1]); // element
+						CopyOperandForDebuggingIfNeeded(Op.ArgC, CachedMemoryHandles[FirstHandleForInstruction[Context.PublicData.InstructionIndex] + 2]); // index
+						CopyOperandForDebuggingIfNeeded(Op.ArgE, CachedMemoryHandles[FirstHandleForInstruction[Context.PublicData.InstructionIndex] + 4]); // ratio
 						Context.EndSlice();
 					}
 				}
 
-				Context.InstructionIndex++;
+				Context.PublicData.InstructionIndex++;
 				break;
 			}
 			case ERigVMOpCode::ArrayUnion:
 			case ERigVMOpCode::ArrayDifference:
 			case ERigVMOpCode::ArrayIntersection:
 			{
-				FRigVMMemoryHandle& ArrayHandleA = CachedMemoryHandles[FirstHandleForInstruction[Context.InstructionIndex]]; 					
-				FRigVMMemoryHandle& ArrayHandleB = CachedMemoryHandles[FirstHandleForInstruction[Context.InstructionIndex] + 1]; 					
+				FRigVMMemoryHandle& ArrayHandleA = CachedMemoryHandles[FirstHandleForInstruction[Context.PublicData.InstructionIndex]]; 					
+				FRigVMMemoryHandle& ArrayHandleB = CachedMemoryHandles[FirstHandleForInstruction[Context.PublicData.InstructionIndex] + 1]; 					
 				FScriptArrayHelper ArrayHelperA(CastFieldChecked<FArrayProperty>(ArrayHandleA.GetProperty()), ArrayHandleA.GetData());
 				FScriptArrayHelper ArrayHelperB(CastFieldChecked<FArrayProperty>(ArrayHandleB.GetProperty()), ArrayHandleB.GetData());
 				const FArrayProperty* ArrayPropertyA = CastFieldChecked<FArrayProperty>(ArrayHandleA.GetProperty());
@@ -2310,7 +2310,7 @@ bool URigVM::Execute(TArrayView<URigVMMemoryStorage*> Memory, TArrayView<void*>
 				}
 				else
 				{
-					FRigVMMemoryHandle& ResultArrayHandle = CachedMemoryHandles[FirstHandleForInstruction[Context.InstructionIndex] + 2]; 					
+					FRigVMMemoryHandle& ResultArrayHandle = CachedMemoryHandles[FirstHandleForInstruction[Context.PublicData.InstructionIndex] + 2]; 					
 					FScriptArrayHelper ResultArrayHelper(CastFieldChecked<FArrayProperty>(ResultArrayHandle.GetProperty()), ResultArrayHandle.GetData());
 					const FArrayProperty* ResultArrayProperty = CastFieldChecked<FArrayProperty>(ResultArrayHandle.GetProperty());
 					const FProperty* ResultElementProperty = ResultArrayProperty->Inner;
@@ -2372,12 +2372,12 @@ bool URigVM::Execute(TArrayView<URigVMMemoryStorage*> Memory, TArrayView<void*>
 					}
 				}
 
-				Context.InstructionIndex++;
+				Context.PublicData.InstructionIndex++;
 				break;
 			}
 			case ERigVMOpCode::ArrayReverse:
 			{
-				FRigVMMemoryHandle& ArrayHandle = CachedMemoryHandles[FirstHandleForInstruction[Context.InstructionIndex]]; 					
+				FRigVMMemoryHandle& ArrayHandle = CachedMemoryHandles[FirstHandleForInstruction[Context.PublicData.InstructionIndex]]; 					
 				FScriptArrayHelper ArrayHelper(CastFieldChecked<FArrayProperty>(ArrayHandle.GetProperty()), ArrayHandle.GetData());
 				for(int32 A=0, B=ArrayHelper.Num()-1; A<B; A++, B--)
 				{
@@ -2390,7 +2390,7 @@ bool URigVM::Execute(TArrayView<URigVMMemoryStorage*> Memory, TArrayView<void*>
 					CopyOperandForDebuggingIfNeeded(Op.Arg, ArrayHandle);
 				}
 
-				Context.InstructionIndex++;
+				Context.PublicData.InstructionIndex++;
 				break;
 			}
 			case ERigVMOpCode::Invalid:
@@ -2401,7 +2401,7 @@ bool URigVM::Execute(TArrayView<URigVMMemoryStorage*> Memory, TArrayView<void*>
 		}
 
 #if WITH_EDITOR
-		if(Context.RuntimeSettings.bEnableProfiling && !InstructionVisitOrder.IsEmpty())
+		if(Context.PublicData.RuntimeSettings.bEnableProfiling && !InstructionVisitOrder.IsEmpty())
 		{
 			const uint64 EndCycles = FPlatformTime::Cycles64();
 			const uint64 Cycles = EndCycles - StartCycles;
diff --git a/Engine/Source/Runtime/RigVM/Public/RigVMCore/RigVM.h b/Engine/Source/Runtime/RigVM/Public/RigVMCore/RigVM.h
--- a/Engine/Source/Runtime/RigVM/Public/RigVMCore/RigVM.h
+++ b/Engine/Source/Runtime/RigVM/Public/RigVMCore/RigVM.h
@@ -618,7 +618,7 @@ public:
 #endif
 
 	uint32 GetNumExecutions() const { return NumExecutions; }
-	const FRigVMExecuteContext& GetContext() const { return Context; }
+	const FRigVMExtendedExecuteContext& GetContext() const { return Context; }
 
 private:
 
@@ -635,7 +635,7 @@ private:
 	FRigVMInstructionArray Instructions;
 
 	UPROPERTY(transient)
-	FRigVMExecuteContext Context;
+	FRigVMExtendedExecuteContext Context;
 
 	UPROPERTY(transient)
 	uint32 NumExecutions;
diff --git a/Engine/Source/Runtime/RigVM/Public/RigVMCore/RigVMByteCode.h b/Engine/Source/Runtime/RigVM/Public/RigVMCore/RigVMByteCode.h
--- a/Engine/Source/Runtime/RigVM/Public/RigVMCore/RigVMByteCode.h
+++ b/Engine/Source/Runtime/RigVM/Public/RigVMCore/RigVMByteCode.h
@@ -956,7 +956,7 @@ public:
 	static uint32 GetCallstackHash(const TArrayView<UObject* const>& InCallstack);
 
 	// returns the input operands of a given instruction
-	FORCEINLINE_DEBUGGABLE FRigVMOperandArray GetInputOperands(int32 InInstructionIndex)
+	FORCEINLINE_DEBUGGABLE FRigVMOperandArray GetInputOperands(int32 InInstructionIndex) const
 	{
 		if(InputOperandsPerInstruction.IsValidIndex(InInstructionIndex))
 		{
@@ -969,7 +969,7 @@ public:
 	}
 
 	// returns the output operands of a given instruction
-	FORCEINLINE_DEBUGGABLE FRigVMOperandArray GetOutputOperands(int32 InInstructionIndex)
+	FORCEINLINE_DEBUGGABLE FRigVMOperandArray GetOutputOperands(int32 InInstructionIndex) const
 	{
 		if(OutputOperandsPerInstruction.IsValidIndex(InInstructionIndex))
 		{
diff --git a/Engine/Source/Runtime/RigVM/Public/RigVMCore/RigVMExecuteContext.h b/Engine/Source/Runtime/RigVM/Public/RigVMCore/RigVMExecuteContext.h
--- a/Engine/Source/Runtime/RigVM/Public/RigVMCore/RigVMExecuteContext.h
+++ b/Engine/Source/Runtime/RigVM/Public/RigVMCore/RigVMExecuteContext.h
@@ -9,6 +9,7 @@
 #include "Logging/TokenizedMessage.h"
 #include "RigVMExecuteContext.generated.h"
 
+struct FRigVMExtendedExecuteContext;
 struct FRigVMExecuteContext;
 class URigVM;
 
@@ -200,20 +201,52 @@
 		: EventName(NAME_None)
 		, FunctionName(NAME_None)
 		, InstructionIndex(0)
-		, VM(nullptr)
 		, RuntimeSettings()
-		, LastExecutionMicroSeconds() 
 	{
 		Reset();
 	}
 
+	FORCEINLINE void Log(EMessageSeverity::Type InSeverity, const FString& InMessage) const
+	{
+		if(RuntimeSettings.LogFunction)
+		{
+			RuntimeSettings.LogFunction(InSeverity, this, InMessage);
+		}
+		else
+		{
+			if(InSeverity == EMessageSeverity::Error)
+			{
+				UE_LOG(LogRigVM, Error, TEXT("Instruction %d: %s"), InstructionIndex, *InMessage);
+			}
+			else if(InSeverity == EMessageSeverity::Warning)
+			{
+				UE_LOG(LogRigVM, Warning, TEXT("Instruction %d: %s"), InstructionIndex, *InMessage);
+			}
+			else
+			{
+				UE_LOG(LogRigVM, Display, TEXT("Instruction %d: %s"), InstructionIndex, *InMessage);
+			}
+		}
+	}
+
+	template <typename FmtType, typename... Types>
+	FORCEINLINE void Logf(EMessageSeverity::Type InSeverity, const FmtType& Fmt, Types... Args) const
+	{
+		Log(InSeverity, FString::Printf(Fmt, Args...));
+	}
+
+	FORCEINLINE uint16 GetInstructionIndex() const { return InstructionIndex; }
+
+	FORCEINLINE FName GetFunctionName() const { return FunctionName; }
+	
+	FORCEINLINE FName GetEventName() const { return EventName; }
+
+
+private:
+
 	FORCEINLINE void Reset()
 	{
-		Slices.Reset();
-		Slices.Add(FRigVMSlice());
-		SliceOffsets.Reset();
 		InstructionIndex = 0;
-		VM = nullptr;
 	}
 
 	FORCEINLINE void CopyFrom(const FRigVMExecuteContext& Other)
@@ -221,8 +254,57 @@
 		EventName = Other.EventName;
 		FunctionName = Other.FunctionName;
 		InstructionIndex = Other.InstructionIndex;
-		VM = Other.VM;
 		RuntimeSettings = Other.RuntimeSettings;
+	}
+	
+	FName EventName;
+	
+	FName FunctionName;
+	
+	uint16 InstructionIndex;
+
+	FRigVMRuntimeSettings RuntimeSettings;
+
+	friend struct FRigVMExtendedExecuteContext;
+	friend class URigVM;
+	friend class UControlRig;
+	friend class UAdditiveControlRig;
+	friend struct FRigUnit_BeginExecution;
+	friend struct FRigUnit_InverseExecution;
+	friend struct FRigUnit_PrepareForExecution;
+	friend struct FEngineTestRigVM_Begin;
+	friend struct FEngineTestRigVM_Setup; 
+};
+
+/**
+ * The execute context is used for mutable nodes to
+ * indicate execution order.
+ */
+USTRUCT()
+struct RIGVM_API FRigVMExtendedExecuteContext
+{
+	GENERATED_BODY()
+
+	FORCEINLINE FRigVMExtendedExecuteContext()
+		: VM(nullptr)
+		, LastExecutionMicroSeconds() 
+	{
+		Reset();
+	}
+
+	FORCEINLINE void Reset()
+	{
+		PublicData.Reset();
+		VM = nullptr;
+		Slices.Reset();
+		Slices.Add(FRigVMSlice());
+		SliceOffsets.Reset();
+	}
+
+	FORCEINLINE void CopyFrom(const FRigVMExtendedExecuteContext& Other)
+	{
+		PublicData = Other.PublicData;
+		VM = Other.VM;
 		OpaqueArguments = Other.OpaqueArguments;
 		Slices = Other.Slices;
 		SliceOffsets = Other.SliceOffsets;
@@ -230,7 +312,7 @@
 
 	FORCEINLINE const FRigVMSlice& GetSlice() const
 	{
-		const int32 SliceOffset = (int32)SliceOffsets[InstructionIndex];
+		const int32 SliceOffset = (int32)SliceOffsets[PublicData.InstructionIndex];
 		if (SliceOffset == 0)
 		{
 			return Slices.Last();
@@ -263,41 +345,12 @@
 		return GetSlice().IsComplete();
 	}
 
-	FORCEINLINE void Log(EMessageSeverity::Type InSeverity, const FString& InMessage) const
-	{
-		if(RuntimeSettings.LogFunction)
-		{
-			RuntimeSettings.LogFunction(InSeverity, this, InMessage);
-		}
-		else
-		{
-			if(InSeverity == EMessageSeverity::Error)
-			{
-				UE_LOG(LogRigVM, Error, TEXT("Instruction %d: %s"), InstructionIndex, *InMessage);
-			}
-			else if(InSeverity == EMessageSeverity::Warning)
-			{
-				UE_LOG(LogRigVM, Warning, TEXT("Instruction %d: %s"), InstructionIndex, *InMessage);
-			}
-			else
-			{
-				UE_LOG(LogRigVM, Display, TEXT("Instruction %d: %s"), InstructionIndex, *InMessage);
-			}
-		}
-	}
-
-	template <typename FmtType, typename... Types>
-	FORCEINLINE void Logf(EMessageSeverity::Type InSeverity, const FmtType& Fmt, Types... Args) const
-	{
-		Log(InSeverity, FString::Printf(Fmt, Args...));
-	}
-
 	FORCEINLINE bool IsValidArrayIndex(int32 InIndex, int32 InArraySize)
 	{
 		if(InIndex < 0 || InIndex >= InArraySize)
 		{
 			static const TCHAR OutOfBoundsFormat[] = TEXT("Array Index (%d) out of bounds (count %d).");
-			Logf(EMessageSeverity::Error, OutOfBoundsFormat, InIndex, InArraySize);
+			PublicData.Logf(EMessageSeverity::Error, OutOfBoundsFormat, InIndex, InArraySize);
 			return false;
 		}
 		return true;
@@ -305,10 +358,10 @@
 
 	FORCEINLINE bool IsValidArraySize(int32 InSize) const
 	{
-		if(InSize < 0 || InSize > RuntimeSettings.MaximumArraySize)
+		if(InSize < 0 || InSize > PublicData.RuntimeSettings.MaximumArraySize)
 		{
 			static const TCHAR OutOfBoundsFormat[] = TEXT("Array Size (%d) larger than allowed maximum (%d).\nCheck VMRuntimeSettings in class settings.");
-			Logf(EMessageSeverity::Error, OutOfBoundsFormat, InSize, RuntimeSettings.MaximumArraySize);
+			PublicData.Logf(EMessageSeverity::Error, OutOfBoundsFormat, InSize, PublicData.RuntimeSettings.MaximumArraySize);
 			return false;
 		}
 		return true;
@@ -316,15 +369,12 @@
 
 	FORCEINLINE void SetRuntimeSettings(FRigVMRuntimeSettings InRuntimeSettings)
 	{
-		RuntimeSettings = InRuntimeSettings;
-		check(RuntimeSettings.MaximumArraySize > 0);
+		PublicData.RuntimeSettings = InRuntimeSettings;
+		check(PublicData.RuntimeSettings.MaximumArraySize > 0);
 	}
 
-	FName EventName;
-	FName FunctionName;
-	uint16 InstructionIndex;
+	FRigVMExecuteContext PublicData;
 	URigVM* VM;
-	FRigVMRuntimeSettings RuntimeSettings;
 	TArrayView<void*> OpaqueArguments;
 	TArray<FRigVMSlice> Slices;
 	TArray<uint16> SliceOffsets;
diff --git a/Engine/Source/Runtime/RigVM/Public/RigVMCore/RigVMFunction.h b/Engine/Source/Runtime/RigVM/Public/RigVMCore/RigVMFunction.h
--- a/Engine/Source/Runtime/RigVM/Public/RigVMCore/RigVMFunction.h
+++ b/Engine/Source/Runtime/RigVM/Public/RigVMCore/RigVMFunction.h
@@ -10,7 +10,7 @@
 typedef TArrayView<FRigVMMemoryHandle> FRigVMMemoryHandleArray;
 typedef TArrayView<void*> FRigVMUserDataArray;
 
-typedef void (*FRigVMFunctionPtr)(FRigVMExecuteContext& RigVMExecuteContext, FRigVMMemoryHandleArray RigVMMemoryHandles);
+typedef void (*FRigVMFunctionPtr)(FRigVMExtendedExecuteContext& RigVMExecuteContext, FRigVMMemoryHandleArray RigVMMemoryHandles);
 
 struct FRigVMTemplate;
 
