RigvM: Bracket Uobject / Uinterface functionality

Added new define as an overall toggle - defaults to off.
If turned on the functionality can be disabled using a console variable

diff --git a/Engine/Plugins/Animation/ControlRig/Source/ControlRigDeveloper/Private/Graph/ControlRigGraphSchema.cpp b/Engine/Plugins/Animation/ControlRig/Source/ControlRigDeveloper/Private/Graph/ControlRigGraphSchema.cpp
--- a/Engine/Plugins/Animation/ControlRig/Source/ControlRigDeveloper/Private/Graph/ControlRigGraphSchema.cpp
+++ b/Engine/Plugins/Animation/ControlRig/Source/ControlRigDeveloper/Private/Graph/ControlRigGraphSchema.cpp
@@ -855,7 +855,7 @@ bool UControlRigGraphSchema::SupportsPinType(UScriptStruct* ScriptStruct) const
 				continue;
 			}
 		}
-		else if (CastField<FObjectProperty>(Property))
+		else if (CastField<FObjectProperty>(Property) && RigVMCore::SupportsUObjects())
 		{
 			continue;
 		}
@@ -885,13 +885,16 @@ bool UControlRigGraphSchema::SupportsPinType(TWeakPtr<const FEdGraphSchemaAction
 		return true;
 	}
 
-	if (PinType.PinCategory == UEdGraphSchema_K2::PC_Object ||
-		PinType.PinCategory == UEdGraphSchema_K2::PC_SoftObject ||
-		PinType.PinCategory == UEdGraphSchema_K2::AllObjectTypes)
+	if(RigVMCore::SupportsUObjects())
 	{
-		if (PinType.PinSubCategoryObject.IsValid())
+		if (PinType.PinCategory == UEdGraphSchema_K2::PC_Object ||
+			PinType.PinCategory == UEdGraphSchema_K2::PC_SoftObject ||
+			PinType.PinCategory == UEdGraphSchema_K2::AllObjectTypes)
 		{
-			return PinType.PinSubCategoryObject->IsA<UClass>();
+			if (PinType.PinSubCategoryObject.IsValid())
+			{
+				return PinType.PinSubCategoryObject->IsA<UClass>();
+			}
 		}
 	}
 
diff --git a/Engine/Source/Developer/RigVMDeveloper/Private/RigVMCompiler/RigVMCompiler.cpp b/Engine/Source/Developer/RigVMDeveloper/Private/RigVMCompiler/RigVMCompiler.cpp
--- a/Engine/Source/Developer/RigVMDeveloper/Private/RigVMCompiler/RigVMCompiler.cpp
+++ b/Engine/Source/Developer/RigVMDeveloper/Private/RigVMCompiler/RigVMCompiler.cpp
@@ -281,6 +281,17 @@ bool URigVMCompiler::Compile(URigVMGraph* InGraph, URigVMController* InControlle
 						}
 					}
 				}
+
+				if(VariableDescription.CPPTypeObject && !RigVMCore::SupportsUObjects())
+				{
+					if(VariableDescription.CPPTypeObject->IsA<UClass>() ||
+						VariableDescription.CPPTypeObject->IsA<UInterface>())
+					{
+						static const FString InvalidObjectTypeMessage = TEXT("Variable Node @@ uses an unsupported UClass / UInterface type.");
+						Settings.ASTSettings.Report(EMessageSeverity::Error, ModelNode, InvalidObjectTypeMessage);
+						bEncounteredGraphError = true;
+					}
+				}
 			}
 
 			for(URigVMPin* Pin : ModelNode->Pins)
diff --git a/Engine/Source/Developer/RigVMDeveloper/Private/RigVMModel/RigVMController.cpp b/Engine/Source/Developer/RigVMDeveloper/Private/RigVMModel/RigVMController.cpp
--- a/Engine/Source/Developer/RigVMDeveloper/Private/RigVMModel/RigVMController.cpp
+++ b/Engine/Source/Developer/RigVMDeveloper/Private/RigVMModel/RigVMController.cpp
@@ -11443,7 +11443,16 @@ void URigVMController::ConfigurePinFromProperty(FProperty* InProperty, URigVMPin
 	}
 	else if (FObjectProperty* ObjectProperty = CastField<FObjectProperty>(PropertyForType))
 	{
-		InOutPin->CPPTypeObject = ObjectProperty->PropertyClass;
+		if(RigVMCore::SupportsUObjects())
+		{
+			InOutPin->CPPTypeObject = ObjectProperty->PropertyClass;
+		}
+		else
+		{
+			ReportErrorf(TEXT("Unsupported type '%s' for pin."), *ObjectProperty->PropertyClass->GetName(), *InOutPin->GetName());
+			InOutPin->CPPType = FString();
+			InOutPin->CPPTypeObject = nullptr;
+		}
 	}
 	else if (FEnumProperty* EnumProperty = CastField<FEnumProperty>(PropertyForType))
 	{
diff --git a/Engine/Source/Programs/UnrealHeaderTool/Private/HeaderParser.cpp b/Engine/Source/Programs/UnrealHeaderTool/Private/HeaderParser.cpp
--- a/Engine/Source/Programs/UnrealHeaderTool/Private/HeaderParser.cpp
+++ b/Engine/Source/Programs/UnrealHeaderTool/Private/HeaderParser.cpp
@@ -6449,6 +6449,13 @@ void FHeaderParser::ParseRigVMMethodParameters(FUnrealStructDefinitionInfo& Stru
 			LogError(TEXT("RigVM Struct '%s' - Member '%s' is editor only - WITH_EDITORONLY_DATA not allowed on structs with RIGVM_METHOD."), *StructDef.GetName(), *Parameter.Name, *MemberCPPType);
 		}
 
+#if !UE_RIGVM_UOBJECT_PROPERTIES_ENABLED
+		if(PropertyDef->GetPropertyBase().IsObjectOrInterface())
+		{
+			LogError(TEXT("RigVM Struct '%s' - Member '%s' is a UObject or interface - object types are not allowed on structs with RIGVM_METHOD."), *StructDef.GetName(), *Parameter.Name);
+		}
+#endif
+
 		if (!ExtendedCPPType.IsEmpty())
 		{
 			// we only support arrays - no maps or similar data structures
diff --git a/Engine/Source/Runtime/RigVM/Private/RigVMCore/RigVMMemoryStorage.cpp b/Engine/Source/Runtime/RigVM/Private/RigVMCore/RigVMMemoryStorage.cpp
--- a/Engine/Source/Runtime/RigVM/Private/RigVMCore/RigVMMemoryStorage.cpp
+++ b/Engine/Source/Runtime/RigVM/Private/RigVMCore/RigVMMemoryStorage.cpp
@@ -27,6 +27,13 @@ FRigVMPropertyDescription::FRigVMPropertyDescription(const FProperty* InProperty
 	const FProperty* ChildProperty = InProperty;
 	do
 	{
+		if(ChildProperty->IsA<FObjectProperty>() ||
+			ChildProperty->IsA<FSoftObjectProperty>() ||
+			ChildProperty->IsA<FInterfaceProperty>())
+		{
+			checkf(RigVMCore::SupportsUObjects(), TEXT("UClass / UInterface types are not supported."));
+		}
+
 		if(const FArrayProperty* ArrayProperty = CastField<FArrayProperty>(ChildProperty))
 		{
 			Containers.Add(EPinContainerType::Array);
@@ -63,6 +70,15 @@ FRigVMPropertyDescription::FRigVMPropertyDescription(const FName& InName, const
 		checkf(CPPTypeObject != nullptr, TEXT("CPPType '%s' requires the CPPTypeObject to be provided."), *CPPType);
 	}
 
+	if(CPPTypeObject)
+	{
+		if(CPPTypeObject->IsA<UClass>() ||
+			CPPTypeObject->IsA<UInterface>())
+		{
+			checkf(RigVMCore::SupportsUObjects(), TEXT("UClass / UInterface types are not supported."));
+		}
+	}
+
 	// only allow valid names
 	SanitizeName();
 
@@ -494,6 +510,8 @@ FProperty* URigVMMemoryStorageGeneratorClass::AddProperty(URigVMMemoryStorageGen
 			}
 			else if(UClass* PropertyClass = Cast<UClass>(InProperty.CPPTypeObject))
 			{
+				check(RigVMCore::SupportsUObjects());
+
 				FObjectProperty* ObjectProperty = new FObjectProperty(PropertyOwner, InProperty.Name, RF_Public);
 				ObjectProperty->SetPropertyClass(PropertyClass);
 				(*ValuePropertyPtr) = ObjectProperty;
diff --git a/Engine/Source/Runtime/RigVM/Private/RigVMModule.cpp b/Engine/Source/Runtime/RigVM/Private/RigVMModule.cpp
--- a/Engine/Source/Runtime/RigVM/Private/RigVMModule.cpp
+++ b/Engine/Source/Runtime/RigVM/Private/RigVMModule.cpp
@@ -10,3 +10,16 @@
 IMPLEMENT_MODULE(FDefaultModuleImpl, RigVM);
 
 DEFINE_LOG_CATEGORY(LogRigVM);
+
+#if UE_RIGVM_UOBJECT_PROPERTIES_ENABLED
+TAutoConsoleVariable<bool> CVarRigVMEnableUObjects(TEXT("RigVM.UObjectSupport"), true, TEXT("When true the RigVMCompiler will allow UObjects."));
+#endif
+
+bool RigVMCore::SupportsUObjects()
+{
+#if UE_RIGVM_UOBJECT_PROPERTIES_ENABLED
+	return CVarRigVMEnableUObjects.GetValueOnGameThread();
+#else
+	return false;
+#endif
+}
diff --git a/Engine/Source/Runtime/RigVM/Public/RigVMCore/RigVMExternalVariable.h b/Engine/Source/Runtime/RigVM/Public/RigVMCore/RigVMExternalVariable.h
--- a/Engine/Source/Runtime/RigVM/Public/RigVMCore/RigVMExternalVariable.h
+++ b/Engine/Source/Runtime/RigVM/Public/RigVMCore/RigVMExternalVariable.h
@@ -4,7 +4,9 @@
 
 #include "CoreMinimal.h"
 #include "RigVMTraits.h"
+#include "RigVMDefines.h"
 #include "RigVMMemory.h"
+#include "RigVMModule.h"
 #include "UObject/UnrealType.h"
 
 /**
@@ -78,10 +80,18 @@ struct RIGVM_API FRigVMExternalVariable
 		}
 		else if (const FObjectProperty* ObjectProperty = CastField<FObjectProperty>(InProperty))
 		{
-			OutTypeName = *FString::Printf(TEXT("TObjectPtr<%s%s>"),
-				ObjectProperty->PropertyClass->GetPrefixCPP(),
-				*ObjectProperty->PropertyClass->GetName());
-			OutTypeObject = ObjectProperty->PropertyClass;
+			if(RigVMCore::SupportsUObjects())
+			{
+				OutTypeName = *FString::Printf(TEXT("TObjectPtr<%s%s>"),
+					ObjectProperty->PropertyClass->GetPrefixCPP(),
+					*ObjectProperty->PropertyClass->GetName());
+				OutTypeObject = ObjectProperty->PropertyClass;
+			}
+			else
+			{
+				OutTypeName = NAME_None;
+				OutTypeObject = nullptr;
+			}
 		}
 		else
 		{
diff --git a/Engine/Source/Runtime/RigVM/Public/RigVMDefines.h b/Engine/Source/Runtime/RigVM/Public/RigVMDefines.h
--- a/Engine/Source/Runtime/RigVM/Public/RigVMDefines.h
+++ b/Engine/Source/Runtime/RigVM/Public/RigVMDefines.h
@@ -9,3 +9,7 @@
 #ifndef UE_RIGVM_UCLASS_BASED_STORAGE_DISABLED
 #define UE_RIGVM_UCLASS_BASED_STORAGE_DISABLED 0
 #endif
+
+#ifndef UE_RIGVM_UOBJECT_PROPERTIES_ENABLED
+#define UE_RIGVM_UOBJECT_PROPERTIES_ENABLED 0
+#endif
diff --git a/Engine/Source/Runtime/RigVM/Public/RigVMModule.h b/Engine/Source/Runtime/RigVM/Public/RigVMModule.h
--- a/Engine/Source/Runtime/RigVM/Public/RigVMModule.h
+++ b/Engine/Source/Runtime/RigVM/Public/RigVMModule.h
@@ -9,4 +9,9 @@
 #include "CoreMinimal.h"
 #include "RigVMDefines.h"
 
-DECLARE_LOG_CATEGORY_EXTERN(LogRigVM, Log, All);
\ No newline at end of file
+DECLARE_LOG_CATEGORY_EXTERN(LogRigVM, Log, All);
+
+namespace RigVMCore
+{
+	RIGVM_API bool SupportsUObjects();
+}
