RigVM: Avoid serialization of byte code and other members during CPFUO

* Relying on IsIgnoringArchetypeRef seems to be the right choice, since the FObjectWriter/Reader only set it to true during CPFUO
* GIsReinstancing is too broad of a check sadly

diff --git a/Engine/Source/Runtime/RigVM/Private/RigVMCore/RigVM.cpp b/Engine/Source/Runtime/RigVM/Private/RigVMCore/RigVM.cpp
--- a/Engine/Source/Runtime/RigVM/Private/RigVMCore/RigVM.cpp
+++ b/Engine/Source/Runtime/RigVM/Private/RigVMCore/RigVM.cpp
@@ -133,15 +133,27 @@ void URigVM::Save(FArchive& Ar)
 
 	int32 RigVMUClassBasedStorageDefine = 0; // this used to 
 	Ar << RigVMUClassBasedStorageDefine;
-	Ar << ExternalPropertyPathDescriptions;
-	Ar << FunctionNamesStorage;
-	Ar << ByteCodeStorage;
-	Ar << Parameters;
+
+	// we rely on Ar.IsIgnoringArchetypeRef for determining if we are currently performing
+	// CPFUO (Copy Properties for unrelated objects). During a reinstance pass we don't
+	// want to overwrite the bytecode and some other properties - since that's handled already
+	// by the RigVMCompiler.
+	if(!Ar.IsIgnoringArchetypeRef())
+	{
+		Ar << ExternalPropertyPathDescriptions;
+		Ar << FunctionNamesStorage;
+		Ar << ByteCodeStorage;
+		Ar << Parameters;
+	}
 }
 
 void URigVM::Load(FArchive& Ar)
 {
-	Reset();
+	// we rely on Ar.IsIgnoringArchetypeRef for determining if we are currently performing
+	// CPFUO (Copy Properties for unrelated objects). During a reinstance pass we don't
+	// want to overwrite the bytecode and some other properties - since that's handled already
+	// by the RigVMCompiler.
+	Reset(Ar.IsIgnoringArchetypeRef());
 
 	int32 RigVMUClassBasedStorageDefine = 1;
 	if (Ar.CustomVer(FUE5MainStreamObjectVersion::GUID) >= FUE5MainStreamObjectVersion::RigVMMemoryStorageObject)
@@ -182,11 +194,14 @@ void URigVM::Load(FArchive& Ar)
 	{
 		ClearMemory();
 	}
-	
-	Ar << ExternalPropertyPathDescriptions;
-	Ar << FunctionNamesStorage;
-	Ar << ByteCodeStorage;
-	Ar << Parameters;
+
+	if(!Ar.IsIgnoringArchetypeRef())
+	{
+		Ar << ExternalPropertyPathDescriptions;
+		Ar << FunctionNamesStorage;
+		Ar << ByteCodeStorage;
+		Ar << Parameters;
+	}
 }
 
 void URigVM::PostLoad()
@@ -454,21 +469,27 @@ bool URigVM::ValidateAllOperandsDuringLoad()
 	return bAllOperandsValid;
 }
 
-void URigVM::Reset()
+void URigVM::Reset(bool IsIgnoringArchetypeRef)
 {
-	FunctionNamesStorage.Reset();
-	FunctionsStorage.Reset();
-	ExternalPropertyPathDescriptions.Reset();
-	ExternalPropertyPaths.Reset();
-	ByteCodeStorage.Reset();
-	Instructions.Reset();
-	Parameters.Reset();
-	ParametersNameMap.Reset();
+	if(!IsIgnoringArchetypeRef)
+	{
+		FunctionNamesStorage.Reset();
+		FunctionsStorage.Reset();
+		ExternalPropertyPathDescriptions.Reset();
+		ExternalPropertyPaths.Reset();
+		ByteCodeStorage.Reset();
+		Instructions.Reset();
+		Parameters.Reset();
+		ParametersNameMap.Reset();
+	}
 	DeferredVMToCopy = nullptr;
 
-	FunctionNamesPtr = &FunctionNamesStorage;
-	FunctionsPtr = &FunctionsStorage;
-	ByteCodePtr = &ByteCodeStorage;
+	if(!IsIgnoringArchetypeRef)
+	{
+		FunctionNamesPtr = &FunctionNamesStorage;
+		FunctionsPtr = &FunctionsStorage;
+		ByteCodePtr = &ByteCodeStorage;
+	}
 
 	InvalidateCachedMemory();
 	
diff --git a/Engine/Source/Runtime/RigVM/Public/RigVMCore/RigVM.h b/Engine/Source/Runtime/RigVM/Public/RigVMCore/RigVM.h
--- a/Engine/Source/Runtime/RigVM/Public/RigVMCore/RigVM.h
+++ b/Engine/Source/Runtime/RigVM/Public/RigVMCore/RigVM.h
@@ -153,7 +153,7 @@ public:
 	void PostLoad() override;
 	
 	// resets the container and maintains all memory
-	void Reset();
+	void Reset(bool IsIgnoringArchetypeRef = false);
 
 	// resets the container and removes all memory
 	void Empty();
