RigVM: Minor security update to RigVMMemory::CopyProperty

diff --git a/Engine/Source/Runtime/RigVM/Private/RigVMCore/RigVMMemoryStorage.cpp b/Engine/Source/Runtime/RigVM/Private/RigVMCore/RigVMMemoryStorage.cpp
--- a/Engine/Source/Runtime/RigVM/Private/RigVMCore/RigVMMemoryStorage.cpp
+++ b/Engine/Source/Runtime/RigVM/Private/RigVMCore/RigVMMemoryStorage.cpp
@@ -749,26 +749,32 @@ bool URigVMMemoryStorage::CopyProperty(
 		{
 			if(const FDoubleProperty* SourceDoubleProperty = CastField<FDoubleProperty>(InSourceProperty))
 			{
-				float* TargetFloats = (float*)InTargetPtr;
-				double* SourceDoubles = (double*)InSourcePtr;
-				for(int32 Index=0;Index<TargetFloatProperty->ArrayDim;Index++)
+				if(TargetFloatProperty->ArrayDim == SourceDoubleProperty->ArrayDim)
 				{
-					TargetFloats[Index] = (float)SourceDoubles[Index];
+					float* TargetFloats = (float*)InTargetPtr;
+					double* SourceDoubles = (double*)InSourcePtr;
+					for(int32 Index=0;Index<TargetFloatProperty->ArrayDim;Index++)
+					{
+						TargetFloats[Index] = (float)SourceDoubles[Index];
+					}
+					return true;
 				}
-				return true;
 			}
 		}
 		else if(const FDoubleProperty* TargetDoubleProperty = CastField<FDoubleProperty>(InTargetProperty))
 		{
 			if(const FFloatProperty* SourceFloatProperty = CastField<FFloatProperty>(InSourceProperty))
 			{
-				double* TargetDoubles = (double*)InTargetPtr;
-				float* SourceFloats = (float*)InSourcePtr;
-				for(int32 Index=0;Index<TargetDoubleProperty->ArrayDim;Index++)
+				if(TargetDoubleProperty->ArrayDim == SourceFloatProperty->ArrayDim)
 				{
-					TargetDoubles[Index] = (double)SourceFloats[Index];
+					double* TargetDoubles = (double*)InTargetPtr;
+					float* SourceFloats = (float*)InSourcePtr;
+					for(int32 Index=0;Index<TargetDoubleProperty->ArrayDim;Index++)
+					{
+						TargetDoubles[Index] = (double)SourceFloats[Index];
+					}
+					return true;
 				}
-				return true;
 			}
 		}
 		else if (const FByteProperty* TargetByteProperty = CastField<FByteProperty>(InTargetProperty))
