RigVM: Support for negative array indexing

diff --git a/Engine/Source/Runtime/RigVM/Private/RigVMCore/RigVM.cpp b/Engine/Source/Runtime/RigVM/Private/RigVMCore/RigVM.cpp
--- a/Engine/Source/Runtime/RigVM/Private/RigVMCore/RigVM.cpp
+++ b/Engine/Source/Runtime/RigVM/Private/RigVMCore/RigVM.cpp
@@ -1998,8 +1998,8 @@ bool URigVM::Execute(TArrayView<URigVMMemoryStorage*> Memory, TArrayView<void*>
 				FRigVMMemoryHandle& ArrayHandle = CachedMemoryHandles[FirstHandleForInstruction[Context.PublicData.InstructionIndex]];
 				const FArrayProperty* ArrayProperty = CastFieldChecked<FArrayProperty>(ArrayHandle.GetProperty());
 				FScriptArrayHelper ArrayHelper(ArrayProperty, ArrayHandle.GetData());
-				const int32 Index = (*((int32*)CachedMemoryHandles[FirstHandleForInstruction[Context.PublicData.InstructionIndex] + 1].GetData()));
-				if(Context.IsValidArrayIndex(Index, ArrayHelper.Num()))
+				int32 Index = (*((int32*)CachedMemoryHandles[FirstHandleForInstruction[Context.PublicData.InstructionIndex] + 1].GetData()));
+				if(Context.IsValidArrayIndex(Index, ArrayHelper))
 				{
 					FRigVMMemoryHandle& ElementHandle = CachedMemoryHandles[FirstHandleForInstruction[Context.PublicData.InstructionIndex] + 2];
 					uint8* TargetMemory = ElementHandle.GetData();
@@ -2023,8 +2023,8 @@ bool URigVM::Execute(TArrayView<URigVMMemoryStorage*> Memory, TArrayView<void*>
 				FRigVMMemoryHandle& ArrayHandle = CachedMemoryHandles[FirstHandleForInstruction[Context.PublicData.InstructionIndex]];
 				const FArrayProperty* ArrayProperty = CastFieldChecked<FArrayProperty>(ArrayHandle.GetProperty());
 				FScriptArrayHelper ArrayHelper(ArrayProperty, ArrayHandle.GetData());
-				const int32 Index = (*((int32*)CachedMemoryHandles[FirstHandleForInstruction[Context.PublicData.InstructionIndex] + 1].GetData()));
-				if(Context.IsValidArrayIndex(Index, ArrayHelper.Num()))
+				int32 Index = (*((int32*)CachedMemoryHandles[FirstHandleForInstruction[Context.PublicData.InstructionIndex] + 1].GetData()));
+				if(Context.IsValidArrayIndex(Index, ArrayHelper))
 				{
 					FRigVMMemoryHandle& ElementHandle = CachedMemoryHandles[FirstHandleForInstruction[Context.PublicData.InstructionIndex] + 2];
 					uint8* TargetMemory = ArrayHelper.GetRawPtr(Index);
@@ -2051,6 +2051,13 @@ bool URigVM::Execute(TArrayView<URigVMMemoryStorage*> Memory, TArrayView<void*>
 				if(Context.IsValidArraySize(ArrayHelper.Num() + 1))
 				{
 					int32 Index = (*((int32*)CachedMemoryHandles[FirstHandleForInstruction[Context.PublicData.InstructionIndex] + 1].GetData()));
+
+					// we support wrapping the index around similar to python
+					if(Index < 0)
+					{
+						Index = ArrayHelper.Num() + Index;
+					}
+					
 					Index = FMath::Clamp<int32>(Index, 0, ArrayHelper.Num());
 					ArrayHelper.InsertValues(Index, 1);
 
@@ -2076,8 +2083,8 @@ bool URigVM::Execute(TArrayView<URigVMMemoryStorage*> Memory, TArrayView<void*>
 				FRigVMMemoryHandle& ArrayHandle = CachedMemoryHandles[FirstHandleForInstruction[Context.PublicData.InstructionIndex]];
 				const FArrayProperty* ArrayProperty = CastFieldChecked<FArrayProperty>(ArrayHandle.GetProperty());
 				FScriptArrayHelper ArrayHelper(ArrayProperty, ArrayHandle.GetData());
-				const int32 Index = (*((int32*)CachedMemoryHandles[FirstHandleForInstruction[Context.PublicData.InstructionIndex] + 1].GetData()));
-				if(Context.IsValidArrayIndex(Index, ArrayHelper.Num()))
+				int32 Index = (*((int32*)CachedMemoryHandles[FirstHandleForInstruction[Context.PublicData.InstructionIndex] + 1].GetData()));
+				if(Context.IsValidArrayIndex(Index, ArrayHelper))
 				{
 					ArrayHelper.RemoveValues(Index, 1);
 				}
diff --git a/Engine/Source/Runtime/RigVM/Public/RigVMCore/RigVMExecuteContext.h b/Engine/Source/Runtime/RigVM/Public/RigVMCore/RigVMExecuteContext.h
--- a/Engine/Source/Runtime/RigVM/Public/RigVMCore/RigVMExecuteContext.h
+++ b/Engine/Source/Runtime/RigVM/Public/RigVMCore/RigVMExecuteContext.h
@@ -345,12 +345,20 @@ struct RIGVM_API FRigVMExtendedExecuteContext
 		return GetSlice().IsComplete();
 	}
 
-	FORCEINLINE bool IsValidArrayIndex(int32 InIndex, int32 InArraySize)
+	FORCEINLINE bool IsValidArrayIndex(int32& InOutIndex, const FScriptArrayHelper& InArrayHelper) const
 	{
-		if(InIndex < 0 || InIndex >= InArraySize)
+		const int32 InOriginalIndex = InOutIndex;
+
+		// we support wrapping the index around similar to python
+		if(InOutIndex < 0)
+		{
+			InOutIndex = InArrayHelper.Num() + InOutIndex;
+		}
+
+		if(!InArrayHelper.IsValidIndex(InOutIndex))
 		{
 			static const TCHAR OutOfBoundsFormat[] = TEXT("Array Index (%d) out of bounds (count %d).");
-			PublicData.Logf(EMessageSeverity::Error, OutOfBoundsFormat, InIndex, InArraySize);
+			PublicData.Logf(EMessageSeverity::Error, OutOfBoundsFormat, InOriginalIndex, InArrayHelper.Num());
 			return false;
 		}
 		return true;
