RigVM: Add introspection on function arguments

diff --git a/Engine/Source/Programs/UnrealHeaderTool/Private/CodeGenerator.cpp b/Engine/Source/Programs/UnrealHeaderTool/Private/CodeGenerator.cpp
--- a/Engine/Source/Programs/UnrealHeaderTool/Private/CodeGenerator.cpp
+++ b/Engine/Source/Programs/UnrealHeaderTool/Private/CodeGenerator.cpp
@@ -3549,8 +3549,20 @@ void FNativeClassHeaderGenerator::ExportGeneratedStructBodyMacros(FOutputDevice&
 		{
 			for (const FRigVMMethodInfo& MethodInfo : StructRigVMInfo.Methods)
 			{
-				Out.Logf(TEXT("\t\tFRigVMRegistry::Get().Register(TEXT(\"%s::%s\"), &%s::RigVM%s, %s);\r\n"),
-					*StructNameCPP, *MethodInfo.Name, *StructNameCPP, *MethodInfo.Name, *OuterSingletonName);
+				const FString ArgumentsName = FString::Printf(TEXT("Arguments_%s_%s"), *StructNameCPP, *MethodInfo.Name);
+				Out.Logf(TEXT("\t\tTArray<FRigVMFunctionArgument> %s;\r\n"), *ArgumentsName);
+				for (int32 MemberIndex = 0; MemberIndex < StructRigVMInfo.Members.Num(); MemberIndex++)
+				{
+					const FRigVMParameter& Parameter = StructRigVMInfo.Members[MemberIndex];
+					Out.Logf(TEXT("\t\t%s.Emplace(TEXT(\"%s\"), TEXT(\"%s\"));\r\n"), *ArgumentsName, *Parameter.NameOriginal(), *Parameter.TypeOriginal());
+				}
+				for (int32 ParameterIndex = 0; ParameterIndex < MethodInfo.Parameters.Num(); ParameterIndex++)
+				{
+					const FRigVMParameter& Parameter = MethodInfo.Parameters[ParameterIndex];
+					Out.Logf(TEXT("\t\t%s.Emplace(TEXT(\"%s\"), TEXT(\"%s\"));\r\n"), *ArgumentsName, *Parameter.NameOriginal(), *Parameter.TypeOriginal());
+				}
+				Out.Logf(TEXT("\t\tFRigVMRegistry::Get().Register(TEXT(\"%s::%s\"), &%s::RigVM%s, %s, %s);\r\n"),
+					*StructNameCPP, *MethodInfo.Name, *StructNameCPP, *MethodInfo.Name, *OuterSingletonName, *ArgumentsName);
 			}
 		}
 
diff --git a/Engine/Source/Runtime/RigVM/Private/RigVMCore/RigVMFunction.cpp b/Engine/Source/Runtime/RigVM/Private/RigVMCore/RigVMFunction.cpp
--- a/Engine/Source/Runtime/RigVM/Private/RigVMCore/RigVMFunction.cpp
+++ b/Engine/Source/Runtime/RigVM/Private/RigVMCore/RigVMFunction.cpp
@@ -44,6 +44,17 @@ FString FRigVMFunction::GetModuleRelativeHeaderPath() const
 	return FString();
 }
 
+bool FRigVMFunction::IsAdditionalArgument(const FRigVMFunctionArgument& InArgument) const
+{
+#if WITH_EDITOR
+	if (Struct)
+	{
+		return Struct->FindPropertyByName(InArgument.Name) == nullptr;
+	}
+#endif
+	return false;
+}
+
 const FRigVMTemplate* FRigVMFunction::GetTemplate() const
 {
 	if(TemplateIndex == INDEX_NONE)
diff --git a/Engine/Source/Runtime/RigVM/Private/RigVMCore/RigVMRegistry.cpp b/Engine/Source/Runtime/RigVM/Private/RigVMCore/RigVMRegistry.cpp
--- a/Engine/Source/Runtime/RigVM/Private/RigVMCore/RigVMRegistry.cpp
+++ b/Engine/Source/Runtime/RigVM/Private/RigVMCore/RigVMRegistry.cpp
@@ -16,14 +16,14 @@ void FRigVMRegistry::Refresh()
 {
 }
 
-void FRigVMRegistry::Register(const TCHAR* InName, FRigVMFunctionPtr InFunctionPtr, UScriptStruct* InStruct)
+void FRigVMRegistry::Register(const TCHAR* InName, FRigVMFunctionPtr InFunctionPtr, UScriptStruct* InStruct, const TArray<FRigVMFunctionArgument>& InArguments)
 {
 	if (FindFunction(InName) != nullptr)
 	{
 		return;
 	}
 
-	const FRigVMFunction Function(InName, InFunctionPtr, InStruct, Functions.Num());
+	const FRigVMFunction Function(InName, InFunctionPtr, InStruct, Functions.Num(), InArguments);
 	Functions.Add(Function);
 	FunctionNameToIndex.Add(InName, Function.Index);
 
diff --git a/Engine/Source/Runtime/RigVM/Public/RigVMCore/RigVMFunction.h b/Engine/Source/Runtime/RigVM/Public/RigVMCore/RigVMFunction.h
--- a/Engine/Source/Runtime/RigVM/Public/RigVMCore/RigVMFunction.h
+++ b/Engine/Source/Runtime/RigVM/Public/RigVMCore/RigVMFunction.h
@@ -29,6 +29,27 @@ enum class ERigVMPinDirection : uint8
 	Invalid // The max value for this enum - used for guarding.
 };
 
+/**
+ * The FRigVMFunctionArgument describes an argument necessary for the C++ invocation of the RIGVM_METHOD backed function
+ */
+struct RIGVM_API FRigVMFunctionArgument
+{
+	const TCHAR* Name;
+	const TCHAR* Type;
+
+	FRigVMFunctionArgument()
+		: Name(nullptr)
+		, Type(nullptr)
+	{
+	}
+
+	FRigVMFunctionArgument(const TCHAR* InName, const TCHAR* InType)
+		: Name(InName)
+		, Type(InType)
+	{
+	}
+};
+
 /**
  * The FRigVMFunction is used to represent a function pointer generated by UHT
  * for a given name. The name might be something like "FMyStruct::MyVirtualMethod"
@@ -40,6 +61,7 @@ struct RIGVM_API FRigVMFunction
 	FRigVMFunctionPtr FunctionPtr;
 	int32 Index;
 	int32 TemplateIndex;
+	TArray<FRigVMFunctionArgument> Arguments;
 
 	FRigVMFunction()
 		: Name(nullptr)
@@ -47,15 +69,17 @@ struct RIGVM_API FRigVMFunction
 		, FunctionPtr(nullptr)
 		, Index(INDEX_NONE)
 		, TemplateIndex(INDEX_NONE)
+		, Arguments()
 	{
 	}
 
-	FRigVMFunction(const TCHAR* InName, FRigVMFunctionPtr InFunctionPtr, UScriptStruct* InStruct = nullptr, int32 InIndex = INDEX_NONE)
+	FRigVMFunction(const TCHAR* InName, FRigVMFunctionPtr InFunctionPtr, UScriptStruct* InStruct = nullptr, int32 InIndex = INDEX_NONE, const TArray<FRigVMFunctionArgument>& InArguments = TArray<FRigVMFunctionArgument>())
 		: Name(InName)
 		, Struct(InStruct)
 		, FunctionPtr(InFunctionPtr)
 		, Index(InIndex)
 		, TemplateIndex(INDEX_NONE)
+		, Arguments(InArguments)
 	{
 	}
 
@@ -63,5 +87,7 @@ struct RIGVM_API FRigVMFunction
 	FName GetMethodName() const;
 	FString GetModuleName() const;
 	FString GetModuleRelativeHeaderPath() const;
+	const TArray<FRigVMFunctionArgument>& GetArguments() const { return Arguments; }
+	bool IsAdditionalArgument(const FRigVMFunctionArgument& InArgument) const;
 	const FRigVMTemplate* GetTemplate() const;
 };
diff --git a/Engine/Source/Runtime/RigVM/Public/RigVMCore/RigVMRegistry.h b/Engine/Source/Runtime/RigVM/Public/RigVMCore/RigVMRegistry.h
--- a/Engine/Source/Runtime/RigVM/Public/RigVMCore/RigVMRegistry.h
+++ b/Engine/Source/Runtime/RigVM/Public/RigVMCore/RigVMRegistry.h
@@ -23,7 +23,7 @@ public:
 	// Registers a function given its name.
 	// The name will be the name of the struct and virtual method,
 	// for example "FMyStruct::MyVirtualMethod"
-	void Register(const TCHAR* InName, FRigVMFunctionPtr InFunctionPtr, UScriptStruct* InStruct = nullptr);
+	void Register(const TCHAR* InName, FRigVMFunctionPtr InFunctionPtr, UScriptStruct* InStruct = nullptr, const TArray<FRigVMFunctionArgument>& InArguments = TArray<FRigVMFunctionArgument>());
 
 	// Refreshes the list and finds the function pointers
 	// based on the names.
