Missing include added.

diff --git a/Engine/Source/Runtime/RigVM/Private/RigVMCore/RigVM.cpp b/Engine/Source/Runtime/RigVM/Private/RigVMCore/RigVM.cpp
--- a/Engine/Source/Runtime/RigVM/Private/RigVMCore/RigVM.cpp
+++ b/Engine/Source/Runtime/RigVM/Private/RigVMCore/RigVM.cpp
@@ -8,6 +8,7 @@
 #include "HAL/PlatformTLS.h"
 #include "Async/ParallelFor.h"
 #include "UObject/UObjectHash.h"
+#include "UObject/UObjectIterator.h"
 
 void FRigVMParameter::Serialize(FArchive& Ar)
 {
