ControlRig / RigVM: Offer user workflows per unit / node

diff --git a/Engine/Source/Developer/RigVMDeveloper/Private/RigVMModel/Nodes/RigVMUnitNode.cpp b/Engine/Source/Developer/RigVMDeveloper/Private/RigVMModel/Nodes/RigVMUnitNode.cpp
--- a/Engine/Source/Developer/RigVMDeveloper/Private/RigVMModel/Nodes/RigVMUnitNode.cpp
+++ b/Engine/Source/Developer/RigVMDeveloper/Private/RigVMModel/Nodes/RigVMUnitNode.cpp
@@ -4,6 +4,7 @@
 
 #include "Animation/Rig.h"
 #include "RigVMCore/RigVMStruct.h"
+#include "RigVMUserWorkflowRegistry.h"
 
 FString URigVMUnitNode::GetNodeTitle() const
 {
@@ -109,6 +110,25 @@ FString URigVMUnitNode::GetDeprecatedMetadata() const
 	return FString();
 }
 
+TArray<FRigVMUserWorkflow> URigVMUnitNode::GetSupportedWorkflows(ERigVMUserWorkflowType InType) const
+{
+	TArray<FRigVMUserWorkflow> Workflows = Super::GetSupportedWorkflows(InType);
+
+	if(UScriptStruct* Struct = GetScriptStruct())
+	{
+		check(Struct->IsChildOf(FRigVMStruct::StaticStruct()));
+
+		const TSharedPtr<FStructOnScope> StructOnScope = ConstructStructInstance();
+		const FRigVMStruct* StructMemory = (const FRigVMStruct*)StructOnScope->GetStructMemory();
+		TArray<FRigVMUserWorkflow> StructWorkflows = StructMemory->GetWorkflows(InType);
+		StructWorkflows.Append(URigVMUserWorkflowRegistry::Get()->GetWorkflows(InType, Struct, this));
+		Swap(Workflows, StructWorkflows);
+		Workflows.Append(StructWorkflows);
+	}
+
+	return Workflows;
+}
+
 bool URigVMUnitNode::IsAggregate() const
 {
 #if UE_RIGVM_AGGREGATE_NODES_ENABLED
diff --git a/Engine/Source/Developer/RigVMDeveloper/Private/RigVMModel/RigVMController.cpp b/Engine/Source/Developer/RigVMDeveloper/Private/RigVMModel/RigVMController.cpp
--- a/Engine/Source/Developer/RigVMDeveloper/Private/RigVMModel/RigVMController.cpp
+++ b/Engine/Source/Developer/RigVMDeveloper/Private/RigVMModel/RigVMController.cpp
@@ -7163,6 +7163,16 @@ bool URigVMController::SetPinDefaultValue(URigVMPin* InPin, const FString& InDef
 		}
 	}
 
+	if(bSetupUndoRedo)
+	{
+		TArray<FRigVMUserWorkflow> WorkflowsOnPinChange = InPin->GetNode()->GetSupportedWorkflows(ERigVMUserWorkflowType::OnPinDefaultChanged);
+		for(const FRigVMUserWorkflow& Workflow : WorkflowsOnPinChange)
+		{
+			const URigVMUserWorkflowOptions* Options = MakeOptionsForWorkflow(InPin, Workflow);
+			PerformUserWorkflow(Workflow, Options, bSetupUndoRedo);
+		}
+	}
+
 	if (bSetupUndoRedo)
 	{
 		ActionStack->EndAction(Action, bMergeUndoAction);
@@ -10905,6 +10915,124 @@ bool URigVMController::SetLocalVariableDefaultValue(const FName& InVariableName,
 	return true;
 }
 
+URigVMUserWorkflowOptions* URigVMController::MakeOptionsForWorkflow(UObject* InSubject, const FRigVMUserWorkflow& InWorkflow)
+{
+	URigVMUserWorkflowOptions* Options = nullptr;
+
+	UClass* Class = InWorkflow.GetOptionsClass();
+	if(Class == nullptr)
+	{
+		return Options;
+	}
+
+	if(!Class->IsChildOf(URigVMUserWorkflowOptions::StaticClass()))
+	{
+		return Options;
+	}
+	
+	Options = NewObject<URigVMUserWorkflowOptions>(GetTransientPackage(), Class, NAME_None, RF_Transient);
+	Options->Subject = InSubject;
+	Options->Workflow = InWorkflow;
+
+	TWeakObjectPtr<URigVMController> WeakThis = this;
+	Options->ReportDelegate = FRigVMReportDelegate::CreateLambda([WeakThis](
+		EMessageSeverity::Type InSeverity, UObject* InSubject, const FString& InMessage)
+		{
+			if(URigVMController* StrongThis = WeakThis.Get())
+			{
+				if(InSeverity == EMessageSeverity::Error ||
+					InSeverity == EMessageSeverity::CriticalError)
+				{
+					StrongThis->ReportAndNotifyError(InMessage);
+				}
+				else if(InSeverity == EMessageSeverity::Warning ||
+					InSeverity == EMessageSeverity::PerformanceWarning)
+				{
+					StrongThis->ReportAndNotifyWarning(InMessage);
+				}
+				else 
+				{
+					StrongThis->ReportInfo(InMessage);
+				}
+			}
+		}
+	);
+
+	if(ConfigureWorkflowOptionsDelegate.IsBound())
+	{
+		ConfigureWorkflowOptionsDelegate.Execute(Options);
+	}
+	
+	return Options;
+}
+
+bool URigVMController::PerformUserWorkflowAction(const FRigVMUserWorkflow& InWorkflow, const FRigVMUserWorkflowAction& InAction, bool bSetupUndoRedo)
+{
+	if(!InAction.IsValid())
+	{
+		return false;
+	}
+
+	static bool bIsPerformingUserWorkflowAction = false;
+	if(bIsPerformingUserWorkflowAction)
+	{
+		return false;
+	}
+	TGuardValue<bool> ReentryGuard(bIsPerformingUserWorkflowAction, true);
+	
+	switch(InAction.GetType())
+	{
+		case ERigVMUserWorkflowActionType::SetPinDefaultValue:
+		{
+			if(URigVMPin* Pin = InAction.GetSubject<URigVMPin>())
+			{
+				return SetPinDefaultValue(Pin->GetPinPath(), InAction.GetData(), true, bSetupUndoRedo, false, false);
+			}
+			break;
+		}
+		case ERigVMUserWorkflowActionType::Invalid:
+		default:
+		{
+			break;
+		}
+	}
+	return false;
+}
+
+bool URigVMController::PerformUserWorkflowActions(const FRigVMUserWorkflow& InWorkflow, const TArray<FRigVMUserWorkflowAction>& InActions, bool bSetupUndoRedo)
+{
+	if(InActions.IsEmpty())
+	{
+		return false;
+	}
+	
+	for(const FRigVMUserWorkflowAction& Action : InActions)
+	{
+		if(!Action.IsValid())
+		{
+			return false;
+		}
+	}
+
+	FRigVMBaseAction Bracket;
+	Bracket.Title = InWorkflow.GetTitle();
+	ActionStack->BeginAction(Bracket);
+
+	for(const FRigVMUserWorkflowAction& Action : InActions)
+	{
+		PerformUserWorkflowAction(InWorkflow, Action, bSetupUndoRedo);
+	}
+
+	ActionStack->EndAction(Bracket);
+	return true;
+}
+
+bool URigVMController::PerformUserWorkflow(const FRigVMUserWorkflow& InWorkflow,
+	const URigVMUserWorkflowOptions* InOptions, bool bSetupUndoRedo)
+{
+	return PerformUserWorkflowActions(InWorkflow, InWorkflow.GetActions(InOptions), bSetupUndoRedo);
+}
+
 TArray<TSoftObjectPtr<URigVMFunctionReferenceNode>> URigVMController::GetAffectedReferences(ERigVMControllerBulkEditType InEditType, bool bForceLoad, bool bNotify)
 {
 	TArray<TSoftObjectPtr<URigVMFunctionReferenceNode>> FunctionReferencePtrs;
diff --git a/Engine/Source/Developer/RigVMDeveloper/Private/RigVMModel/RigVMNode.cpp b/Engine/Source/Developer/RigVMDeveloper/Private/RigVMModel/RigVMNode.cpp
--- a/Engine/Source/Developer/RigVMDeveloper/Private/RigVMModel/RigVMNode.cpp
+++ b/Engine/Source/Developer/RigVMDeveloper/Private/RigVMModel/RigVMNode.cpp
@@ -1,9 +1,11 @@
 // Copyright Epic Games, Inc. All Rights Reserved.
 
 #include "RigVMModel/RigVMNode.h"
+#include "RigVMModel/Nodes/RigVMUnitNode.h"
 #include "RigVMModel/RigVMGraph.h"
 #include "RigVMCore/RigVMExecuteContext.h"
 #include "RigVMCore/RigVMStruct.h"
+#include "RigVMUserWorkflowRegistry.h"
 
 const FString URigVMNode::NodeColorName = TEXT("NodeColor");
 
@@ -528,6 +530,11 @@ double URigVMNode::GetInstructionMicroSeconds(URigVM* InVM, const FRigVMASTProxy
 	return -1.0;
 }
 
+TArray<FRigVMUserWorkflow> URigVMNode::GetSupportedWorkflows(ERigVMUserWorkflowType InType) const
+{
+	return URigVMUserWorkflowRegistry::Get()->GetWorkflows(InType, nullptr, this);
+}
+
 #if WITH_EDITOR
 
 const URigVMNode::FProfilingCache* URigVMNode::UpdateProfilingCacheIfNeeded(URigVM* InVM, const FRigVMASTProxy& InProxy) const
diff --git a/Engine/Source/Developer/RigVMDeveloper/Private/RigVMUserWorkflowRegistry.cpp b/Engine/Source/Developer/RigVMDeveloper/Private/RigVMUserWorkflowRegistry.cpp
new file mode 100644
--- /dev/null
+++ b/Engine/Source/Developer/RigVMDeveloper/Private/RigVMUserWorkflowRegistry.cpp
@@ -0,0 +1,51 @@
+﻿// Copyright Epic Games, Inc. All Rights Reserved.
+
+#include "RigVMUserWorkflowRegistry.h"
+
+URigVMUserWorkflowRegistry* URigVMUserWorkflowRegistry::Get()
+{
+	return StaticClass()->GetDefaultObject<URigVMUserWorkflowRegistry>();
+}
+
+int32 URigVMUserWorkflowRegistry::RegisterProvider(const UScriptStruct* InStruct, FRigVMUserWorkflowProvider InProvider)
+{
+	const int32 Handle = ++MaxHandle;
+	Providers.Emplace(Handle, InStruct, InProvider);
+	return Handle;
+}
+
+void URigVMUserWorkflowRegistry::UnregisterProvider(int32 InHandle)
+{
+	Providers.RemoveAll([InHandle](const TTuple<int32, const UScriptStruct*, FRigVMUserWorkflowProvider>& Provider) -> bool
+	{
+		return Provider.Get<0>() == InHandle;
+	});
+}
+
+TArray<FRigVMUserWorkflow> URigVMUserWorkflowRegistry::GetWorkflows(ERigVMUserWorkflowType InType, const UScriptStruct* InStruct, const URigVMNode* InNode) const
+{
+	TArray<FRigVMUserWorkflow> Workflows;
+
+	// remove stale delegates
+	Providers.RemoveAll([](const TTuple<int32, const UScriptStruct*, FRigVMUserWorkflowProvider>& Provider) -> bool
+	{
+		return !Provider.Get<2>().IsBound();
+	});
+	
+
+	for(const TTuple<int32, const UScriptStruct*, FRigVMUserWorkflowProvider>& Provider : Providers)
+	{
+		if(Provider.Get<1>() == InStruct && Provider.Get<1>() != nullptr)
+		{
+			Workflows.Append(Provider.Get<2>().Execute(InNode));
+		}
+	}
+
+	Workflows = Workflows.FilterByPredicate([InType](const FRigVMUserWorkflow& InWorkflow) -> bool
+	{
+		return uint32(InWorkflow.GetType()) & uint32(InType) &&
+			InWorkflow.IsValid();
+	});
+
+	return Workflows;
+}
diff --git a/Engine/Source/Developer/RigVMDeveloper/Public/RigVMCompiler/RigVMAST.h b/Engine/Source/Developer/RigVMDeveloper/Public/RigVMCompiler/RigVMAST.h
--- a/Engine/Source/Developer/RigVMDeveloper/Public/RigVMCompiler/RigVMAST.h
+++ b/Engine/Source/Developer/RigVMDeveloper/Public/RigVMCompiler/RigVMAST.h
@@ -11,8 +11,6 @@
 #include "Logging/TokenizedMessage.h"
 #include "RigVMAST.generated.h"
 
-DECLARE_DELEGATE_ThreeParams(FRigVMReportDelegate, EMessageSeverity::Type, UObject*, const FString&);
-
 class FRigVMParserAST;
 class FRigVMBlockExprAST;
 class FRigVMEntryExprAST;
diff --git a/Engine/Source/Developer/RigVMDeveloper/Public/RigVMModel/Nodes/RigVMAggregateNode.h b/Engine/Source/Developer/RigVMDeveloper/Public/RigVMModel/Nodes/RigVMAggregateNode.h
--- a/Engine/Source/Developer/RigVMDeveloper/Public/RigVMModel/Nodes/RigVMAggregateNode.h
+++ b/Engine/Source/Developer/RigVMDeveloper/Public/RigVMModel/Nodes/RigVMAggregateNode.h
@@ -36,14 +36,8 @@ public:
 
 	URigVMNode* GetFirstInnerNode() const;
 	URigVMNode* GetLastInnerNode() const;
-
-	UFUNCTION(BlueprintPure, Category = Aggregate)
 	URigVMPin* GetFirstAggregatePin() const;
-
-	UFUNCTION(BlueprintPure, Category = Aggregate)
 	URigVMPin* GetSecondAggregatePin() const;
-
-	UFUNCTION(BlueprintPure, Category = Aggregate)
 	URigVMPin* GetOppositeAggregatePin() const;
 
 private:
diff --git a/Engine/Source/Developer/RigVMDeveloper/Public/RigVMModel/Nodes/RigVMUnitNode.h b/Engine/Source/Developer/RigVMDeveloper/Public/RigVMModel/Nodes/RigVMUnitNode.h
--- a/Engine/Source/Developer/RigVMDeveloper/Public/RigVMModel/Nodes/RigVMUnitNode.h
+++ b/Engine/Source/Developer/RigVMDeveloper/Public/RigVMModel/Nodes/RigVMUnitNode.h
@@ -27,6 +27,7 @@ public:
 	virtual FName GetEventName() const override;
 	virtual bool IsLoopNode() const override;
 	virtual bool CanBeUpgraded() const override { return GetUpgradeInfo().IsValid(); }
+	virtual TArray<FRigVMUserWorkflow> GetSupportedWorkflows(ERigVMUserWorkflowType InType = ERigVMUserWorkflowType::All) const override;
 	virtual bool IsAggregate() const override;
 	virtual URigVMPin* GetFirstAggregatePin() const override;
 	virtual URigVMPin* GetSecondAggregatePin() const override;
diff --git a/Engine/Source/Developer/RigVMDeveloper/Public/RigVMModel/RigVMController.h b/Engine/Source/Developer/RigVMDeveloper/Public/RigVMModel/RigVMController.h
--- a/Engine/Source/Developer/RigVMDeveloper/Public/RigVMModel/RigVMController.h
+++ b/Engine/Source/Developer/RigVMDeveloper/Public/RigVMModel/RigVMController.h
@@ -18,6 +18,7 @@
 #include "RigVMModel/Nodes/RigVMFunctionReferenceNode.h"
 #include "RigVMModel/Nodes/RigVMArrayNode.h"
 #include "RigVMModel/RigVMBuildData.h"
+#include "RigVMCore/RigVMUserWorkflow.h"
 #include "RigVMController.generated.h"
 
 #ifndef UE_RIGVM_ENABLE_TEMPLATE_NODES
@@ -87,7 +88,7 @@ DECLARE_DELEGATE_RetVal_TwoParams(FRigVMController_BulkEditResult, FRigVMControl
 DECLARE_DELEGATE_FiveParams(FRigVMController_OnBulkEditProgressDelegate, TSoftObjectPtr<URigVMFunctionReferenceNode>, ERigVMControllerBulkEditType, ERigVMControllerBulkEditProgress, int32, int32)
 DECLARE_DELEGATE_RetVal_TwoParams(FString, FRigVMController_PinPathRemapDelegate, const FString& /* InPinPath */, bool /* bIsInput */);
 DECLARE_DELEGATE_OneParam(FRigVMController_RequestJumpToHyperlinkDelegate, const UObject* InSubject);
-
+DECLARE_DELEGATE_OneParam(FRigVMController_ConfigureWorkflowOptionsDelegate, URigVMUserWorkflowOptions* InOutOptions);
 
 /**
  * The Controller is the sole authority to perform changes
@@ -733,6 +734,14 @@ public:
 	UFUNCTION(BlueprintCallable, Category = RigVMController)
 	bool SetLocalVariableDefaultValue(const FName& InVariableName, const FString& InDefaultValue, bool bSetupUndoRedo = true, bool bPrintPythonCommand = false, bool bNotify = true);
 
+	// creates the options struct for a given workflow
+	UFUNCTION(BlueprintCallable, Category = RigVMController)
+	URigVMUserWorkflowOptions* MakeOptionsForWorkflow(UObject* InSubject, const FRigVMUserWorkflow& InWorkflow);
+
+	// performs all actions representing the workflow
+	UFUNCTION(BlueprintCallable, Category = RigVMController)
+	bool PerformUserWorkflow(const FRigVMUserWorkflow& InWorkflow, const URigVMUserWorkflowOptions* InOptions, bool bSetupUndoRedo = true);
+
 	// Determine affected function references for a potential bulk edit on a library node
 	TArray<TSoftObjectPtr<URigVMFunctionReferenceNode>> GetAffectedReferences(ERigVMControllerBulkEditType InEditType, bool bForceLoad = false, bool bNotify = true);
 
@@ -770,7 +779,10 @@ public:
 	FRigVMController_OnBulkEditProgressDelegate OnBulkEditProgressDelegate;
 
 	// A delegate to request the client to follow a hyper link
-	FRigVMController_RequestJumpToHyperlinkDelegate RequestJumpToHyperlinkDelegate; 
+	FRigVMController_RequestJumpToHyperlinkDelegate RequestJumpToHyperlinkDelegate;
+
+	// A delegate to request to configure an options instance for a node workflow
+	FRigVMController_ConfigureWorkflowOptionsDelegate ConfigureWorkflowOptionsDelegate; 
 
 	// Returns the build data of the host
 	static URigVMBuildData* GetBuildData(bool bCreateIfNeeded = true);
@@ -902,6 +914,8 @@ private:
 	URigVMInjectionInfo* InjectNodeIntoPin(URigVMPin* InPin, bool bAsInput, const FName& InInputPinName, const FName& InOutputPinName, bool bSetupUndoRedo = true);
 	URigVMNode* EjectNodeFromPin(URigVMPin* InPin, bool bSetupUndoRedo = true, bool bPrintPythonCommands = false);
 	bool EjectAllInjectedNodes(URigVMNode* InNode, bool bSetupUndoRedo = true, bool bPrintPythonCommands = false);
+	bool PerformUserWorkflowAction(const FRigVMUserWorkflow& InWorkflow, const FRigVMUserWorkflowAction& InAction, bool bSetupUndoRedo);
+	bool PerformUserWorkflowActions(const FRigVMUserWorkflow& InWorkflow, const TArray<FRigVMUserWorkflowAction>& InActions, bool bSetupUndoRedo);
 
 	// try to reconnect source and target pins after a node deletion
 	void RelinkSourceAndTargetPins(URigVMNode* RigNode, bool bSetupUndoRedo = true);
diff --git a/Engine/Source/Developer/RigVMDeveloper/Public/RigVMModel/RigVMNode.h b/Engine/Source/Developer/RigVMDeveloper/Public/RigVMModel/RigVMNode.h
--- a/Engine/Source/Developer/RigVMDeveloper/Public/RigVMModel/RigVMNode.h
+++ b/Engine/Source/Developer/RigVMDeveloper/Public/RigVMModel/RigVMNode.h
@@ -4,6 +4,7 @@
 
 #include "RigVMPin.h"
 #include "RigVMCore/RigVM.h"
+#include "RigVMCore/RigVMUserWorkflow.h"
 #include "RigVMNode.generated.h"
 
 class URigVMGraph;
@@ -201,11 +202,17 @@ public:
 	double GetInstructionMicroSeconds(URigVM* InVM, const FRigVMASTProxy& InProxy = FRigVMASTProxy()) const;
 
 	// return true if this node is a loop node
+	UFUNCTION(BlueprintPure, Category = RigVMNode)
 	virtual bool IsLoopNode() const { return false; }
 
 	// returns true if the node can be upgraded
+	UFUNCTION(BlueprintPure, Category = RigVMNode)
 	virtual bool CanBeUpgraded() const { return false; }
 
+	// returns all supported workflows of the node
+	UFUNCTION(BlueprintCallable, Category = RigVMNode)
+	virtual TArray<FRigVMUserWorkflow> GetSupportedWorkflows(ERigVMUserWorkflowType InType = ERigVMUserWorkflowType::All) const;
+
 	UFUNCTION(BlueprintCallable, Category = RigVMNode)
 	bool HasBreakpoint() const { return bHasBreakpoint; }
 
@@ -218,11 +225,21 @@ public:
 	UFUNCTION(BlueprintCallable, Category = RigVMNode)
 	void SetExecutionIsHaltedAtThisNode(const bool bValue) { bHaltedAtThisNode = bValue; }
 
+	UFUNCTION(BlueprintPure, Category = RigVMNode)
 	virtual bool IsAggregate() const { return false; }
+
+	UFUNCTION(BlueprintPure, Category = RigVMNode)
 	virtual URigVMPin* GetFirstAggregatePin() const { return nullptr; }
+
+	UFUNCTION(BlueprintPure, Category = RigVMNode)
 	virtual URigVMPin* GetSecondAggregatePin() const { return nullptr; }
+
+	UFUNCTION(BlueprintPure, Category = RigVMNode)
 	virtual URigVMPin* GetOppositeAggregatePin() const { return nullptr; }
+
+	UFUNCTION(BlueprintPure, Category = RigVMNode)
 	virtual bool IsInputAggregate() const { return false; }
+	
 private:
 
 	static const FString NodeColorName;
diff --git a/Engine/Source/Developer/RigVMDeveloper/Public/RigVMUserWorkflowRegistry.h b/Engine/Source/Developer/RigVMDeveloper/Public/RigVMUserWorkflowRegistry.h
new file mode 100644
--- /dev/null
+++ b/Engine/Source/Developer/RigVMDeveloper/Public/RigVMUserWorkflowRegistry.h
@@ -0,0 +1,35 @@
+﻿// Copyright Epic Games, Inc. All Rights Reserved.
+
+#pragma once
+
+#include "CoreMinimal.h"
+#include "RigVMCore/RigVMUserWorkflow.h"
+#include "RigVMModel/RigVMNode.h"
+#include "RigVMUserWorkflowRegistry.generated.h"
+
+DECLARE_DYNAMIC_DELEGATE_RetVal_OneParam(TArray<FRigVMUserWorkflow>, FRigVMUserWorkflowProvider, const URigVMNode*, InNode);
+
+UCLASS(BlueprintType)
+class RIGVMDEVELOPER_API URigVMUserWorkflowRegistry : public UObject
+{
+	GENERATED_BODY()
+
+public:
+
+	UFUNCTION(BlueprintPure, Category = FRigVMUserWorkflowRegistry)
+   	static URigVMUserWorkflowRegistry* Get();
+
+	UFUNCTION(BlueprintCallable, Category = FRigVMUserWorkflowRegistry)
+	int32 RegisterProvider(const UScriptStruct* InStruct, FRigVMUserWorkflowProvider InProvider);
+
+	UFUNCTION(BlueprintCallable, Category = FRigVMUserWorkflowRegistry)
+	void UnregisterProvider(int32 InHandle);
+
+	UFUNCTION(BlueprintPure, Category = FRigVMUserWorkflowRegistry)
+	TArray<FRigVMUserWorkflow> GetWorkflows(ERigVMUserWorkflowType InType, const UScriptStruct* InStruct, const URigVMNode* InNode) const;
+
+private:
+
+	int32 MaxHandle = 0;
+	mutable TArray<TTuple<int32, const UScriptStruct*, FRigVMUserWorkflowProvider>> Providers;
+};
diff --git a/Engine/Source/Runtime/RigVM/Private/RigVMCore/RigVMStruct.cpp b/Engine/Source/Runtime/RigVM/Private/RigVMCore/RigVMStruct.cpp
--- a/Engine/Source/Runtime/RigVM/Private/RigVMCore/RigVMStruct.cpp
+++ b/Engine/Source/Runtime/RigVM/Private/RigVMCore/RigVMStruct.cpp
@@ -221,6 +221,15 @@ float FRigVMStruct::GetRatioFromIndex(int32 InIndex, int32 InCount)
 	return ((float)FMath::Clamp<int32>(InIndex, 0, InCount - 1)) / ((float)(InCount - 1));
 }
 
+TArray<FRigVMUserWorkflow> FRigVMStruct::GetWorkflows(ERigVMUserWorkflowType InType) const
+{
+	return GetSupportedWorkflows().FilterByPredicate([InType](const FRigVMUserWorkflow& InWorkflow) -> bool
+	{
+		return uint32(InWorkflow.GetType()) & uint32(InType) &&
+			InWorkflow.IsValid();
+	});
+}
+
 #if WITH_EDITOR
 
 bool FRigVMStruct::ValidateStruct(UScriptStruct* InStruct, FString* OutErrorMessage)
diff --git a/Engine/Source/Runtime/RigVM/Private/RigVMCore/RigVMUserWorkflow.cpp b/Engine/Source/Runtime/RigVM/Private/RigVMCore/RigVMUserWorkflow.cpp
new file mode 100644
--- /dev/null
+++ b/Engine/Source/Runtime/RigVM/Private/RigVMCore/RigVMUserWorkflow.cpp
@@ -0,0 +1,121 @@
+// Copyright Epic Games, Inc. All Rights Reserved.
+
+#include "RigVMCore/RigVMUserWorkflow.h"
+#include "RigVMCore/RigVMMemoryStorage.h"
+#include "RigVMModule.h"
+
+////////////////////////////////////////////////////////////////////////////////////////////////////////////////////////
+
+bool URigVMUserWorkflowOptions::RequiresDialog() const
+{
+	for (TFieldIterator<FProperty> It(GetClass()); It; ++It)
+	{
+		if(const FProperty* Property = CastField<FProperty>(*It))
+		{
+			if(RequiresDialog(Property))
+			{
+				return true;
+			}
+		}
+	}
+	return false;
+}
+
+bool URigVMUserWorkflowOptions::RequiresDialog(const FProperty* InProperty) const
+{
+	if(!InProperty->HasAnyPropertyFlags(CPF_Edit))
+	{
+		return false;
+	}
+
+	static const FName SubjectName = TEXT("Subject");
+	static const FName WorkflowName = TEXT("Workflow");
+	return InProperty->GetFName() != SubjectName && InProperty->GetFName() != WorkflowName;
+}
+
+void URigVMUserWorkflowOptions::Report(EMessageSeverity::Type InSeverity,  const FString& InMessage) const
+{
+	if (ReportDelegate.IsBound())
+	{
+		ReportDelegate.Execute(InSeverity, GetSubject(), InMessage);
+	}
+	else
+	{
+		if (InSeverity == EMessageSeverity::Error)
+		{
+			FScriptExceptionHandler::Get().HandleException(ELogVerbosity::Error, *InMessage, *FString());
+		}
+		else if (InSeverity == EMessageSeverity::Warning)
+		{
+			FScriptExceptionHandler::Get().HandleException(ELogVerbosity::Warning, *InMessage, *FString());
+		}
+		else
+		{
+			UE_LOG(LogRigVM, Display, TEXT("%s"), *InMessage);
+		}
+	}
+}
+
+////////////////////////////////////////////////////////////////////////////////////////////////////////////////////////
+
+TArray<FRigVMUserWorkflowAction> FRigVMUserWorkflow::GetActions(const URigVMUserWorkflowOptions* InOptions) const
+{
+	TArray<FRigVMUserWorkflowAction> Actions;
+	if(!IsValid() || !ValidateOptions(InOptions))
+	{
+		return Actions;
+	}
+
+	if(OnGetActionsDelegate.IsBound())
+	{
+		Actions = OnGetActionsDelegate.Execute(InOptions);
+	}
+	else if(OnGetActionsDynamicDelegate.IsBound())
+	{
+		Actions = OnGetActionsDynamicDelegate.Execute(InOptions);
+	}
+	
+	Actions = Actions.FilterByPredicate([](const FRigVMUserWorkflowAction& Action) -> bool
+	{
+		return Action.IsValid();
+	});
+	
+	return Actions;
+}
+
+bool FRigVMUserWorkflow::ValidateOptions(const URigVMUserWorkflowOptions* InOptions) const
+{
+	UClass* ExpectedOptionsClass = GetOptionsClass();
+	if(!ensure(ExpectedOptionsClass != nullptr))
+	{
+		return false;
+	}
+
+	if(!ensure(ExpectedOptionsClass->IsChildOf(URigVMUserWorkflowOptions::StaticClass())))
+	{
+		return false;
+	}
+
+	if(!ensure(InOptions != nullptr))
+	{
+		return false;
+	}
+
+	if(!ensure(InOptions->IsValid()))
+	{
+		return false;
+	}
+
+	if(!ensure(InOptions->GetClass()->IsChildOf(ExpectedOptionsClass)))
+	{
+		InOptions->Reportf(
+			EMessageSeverity::Error,
+			TEXT("Workflow '%s' cannot execute: Options are of type '%s' - expected was type '%s'."),
+			*GetTitle(),
+			*InOptions->GetClass()->GetName(),
+			*ExpectedOptionsClass->GetName());
+		return false;
+	}
+
+	return true;
+}
diff --git a/Engine/Source/Runtime/RigVM/Public/RigVMCore/RigVMStruct.h b/Engine/Source/Runtime/RigVM/Public/RigVMCore/RigVMStruct.h
--- a/Engine/Source/Runtime/RigVM/Public/RigVMCore/RigVMStruct.h
+++ b/Engine/Source/Runtime/RigVM/Public/RigVMCore/RigVMStruct.h
@@ -7,6 +7,7 @@
 #include "RigVMCore/RigVMRegistry.h"
 #include "RigVMCore/RigVMExternalVariable.h"
 #include "RigVMCore/RigVMTraits.h"
+#include "RigVMCore/RigVMUserWorkflow.h"
 #include "RigVMStruct.generated.h"
 
 // delegates used for variable introspection / creation
@@ -221,6 +222,9 @@ public:
 	// node creation
 	FORCEINLINE virtual void OnUnitNodeCreated(FRigVMUnitNodeCreatedContext& InContext) const {}
 
+	// user workflow
+	TArray<FRigVMUserWorkflow> GetWorkflows(ERigVMUserWorkflowType InType = ERigVMUserWorkflowType::All) const; 
+
 #if WITH_EDITOR
 	static bool ValidateStruct(UScriptStruct* InStruct, FString* OutErrorMessage);
 	static bool CheckPinType(UScriptStruct* InStruct, const FName& PinName, const FString& ExpectedType, FString* OutErrorMessage = nullptr);
@@ -232,7 +236,27 @@ public:
 #endif
 	static FString ExportToFullyQualifiedText(const FProperty* InMemberProperty, const uint8* InMemberMemoryPtr, bool bUseQuotes = true);
 	static FString ExportToFullyQualifiedText(const UScriptStruct* InStruct, const uint8* InStructMemoryPtr);
+
+	template <
+		typename T,
+		typename TEnableIf<TRigVMIsBaseStructure<T>::Value>::Type * = nullptr
+	>
+	FORCEINLINE static FString ExportToFullyQualifiedText(const T& InStructValue)
+	{
+		return ExportToFullyQualifiedText(TBaseStructure<T>::Get(), (const uint8*)&InStructValue);
+	}
+	
+	template <
+		typename T,
+		typename TEnableIf<TModels<CRigVMUStruct, T>::Value>::Type * = nullptr
+	>
+	FORCEINLINE static FString ExportToFullyQualifiedText(const T& InStructValue)
+	{
+		return ExportToFullyQualifiedText(T::StaticStruct(), (const uint8*)&InStructValue);
+	}
+
 	FString ExportToFullyQualifiedText(const UScriptStruct* InScriptStruct, const FName& InPropertyName, const uint8* InStructMemoryPointer = nullptr) const;
+	
 	virtual FName GetNextAggregateName(const FName& InLastAggregatePinName) const { return FName(); }
 	virtual FRigVMStructUpgradeInfo GetUpgradeInfo() const { return FRigVMStructUpgradeInfo(); }
 
@@ -276,6 +300,7 @@ protected:
 	static float GetRatioFromIndex(int32 InIndex, int32 InCount);
 	TMap<FName, FString> GetDefaultValues(UScriptStruct* InScriptStruct) const;
 	bool ApplyUpgradeInfo(const FRigVMStructUpgradeInfo& InUpgradeInfo);
+	FORCEINLINE virtual TArray<FRigVMUserWorkflow> GetSupportedWorkflows() const { return TArray<FRigVMUserWorkflow>(); } 
 
 	friend struct FRigVMStructUpgradeInfo;
 	friend class FRigVMGraphStructUpgradeInfoTest;
diff --git a/Engine/Source/Runtime/RigVM/Public/RigVMCore/RigVMUserWorkflow.h b/Engine/Source/Runtime/RigVM/Public/RigVMCore/RigVMUserWorkflow.h
new file mode 100644
--- /dev/null
+++ b/Engine/Source/Runtime/RigVM/Public/RigVMCore/RigVMUserWorkflow.h
@@ -0,0 +1,193 @@
+// Copyright Epic Games, Inc. All Rights Reserved.
+
+#pragma once
+
+#include "CoreMinimal.h"
+#include "Logging/TokenizedMessage.h"
+#include "RigVMCore/RigVMTraits.h"
+#include "RigVMUserWorkflow.generated.h"
+
+class URigVMUserWorkflowOptions;
+
+DECLARE_DELEGATE_ThreeParams(FRigVMReportDelegate, EMessageSeverity::Type, UObject*, const FString&);
+
+// Types of actions within a workflow
+UENUM(BlueprintType)
+enum class ERigVMUserWorkflowActionType : uint8
+{
+	Invalid = 0 UMETA(Hidden),
+	SetPinDefaultValue = 1
+};
+
+USTRUCT(BlueprintType)
+struct RIGVM_API FRigVMUserWorkflowAction
+{
+	GENERATED_BODY()
+
+public:
+
+	FORCEINLINE FRigVMUserWorkflowAction()
+    : Type(ERigVMUserWorkflowActionType::Invalid)
+    , Subject()
+    , Data()
+    {}
+
+	FORCEINLINE FRigVMUserWorkflowAction(
+		ERigVMUserWorkflowActionType InType,
+		UObject* InSubject,
+		const FString& InData)
+	: Type(InType)
+	, Subject(InSubject)
+	, Data(InData)
+	{}
+
+	FORCEINLINE bool IsValid() const { return Type != ERigVMUserWorkflowActionType::Invalid && Subject != nullptr; }
+	FORCEINLINE ERigVMUserWorkflowActionType GetType() const { return Type; }
+	FORCEINLINE UObject* GetSubject() const { return Subject; }
+	FORCEINLINE const FString& GetData() const { return Data; }
+	
+	template<typename T>
+	FORCEINLINE T* GetSubject() const { return Cast<T>(Subject.Get()); }
+
+protected:
+
+	UPROPERTY(BlueprintReadWrite, EditAnywhere, Category = Action, meta = (AllowPrivateAccess = true))
+	ERigVMUserWorkflowActionType Type;
+
+	UPROPERTY(BlueprintReadWrite, EditAnywhere, Category = Action, meta = (AllowPrivateAccess = true))
+	TObjectPtr<UObject> Subject;
+
+	UPROPERTY(BlueprintReadWrite, EditAnywhere, Category = Action, meta = (AllowPrivateAccess = true))
+	FString Data;
+};
+
+// Types of workflows offered by a rigvm struct node
+UENUM(BlueprintType)
+enum class ERigVMUserWorkflowType : uint8
+{
+	Invalid = 0 UMETA(Hidden),
+	NodeContext = 0x001,
+	PinContext = 0x002,
+	OnPinDefaultChanged  = 0x004,
+	All = NodeContext | PinContext | OnPinDefaultChanged
+};
+
+DECLARE_DELEGATE_RetVal_OneParam(TArray<FRigVMUserWorkflowAction>, FRigVMWorkflowGetActionsDelegate, const URigVMUserWorkflowOptions*);
+DECLARE_DYNAMIC_DELEGATE_RetVal_OneParam(TArray<FRigVMUserWorkflowAction>, FRigVMWorkflowGetActionsDynamicDelegate, const URigVMUserWorkflowOptions*, InOptions);
+
+USTRUCT(BlueprintType)
+struct RIGVM_API FRigVMUserWorkflow
+{
+	GENERATED_BODY()
+
+public:
+
+	FORCEINLINE FRigVMUserWorkflow()
+	: Title()
+	, Tooltip()
+	, Type(ERigVMUserWorkflowType::Invalid)
+	, OnGetActionsDelegate()
+	, OnGetActionsDynamicDelegate()
+	, OptionsClass(nullptr)
+	{}
+
+	FORCEINLINE FRigVMUserWorkflow(
+		const FString& InTitle,
+		const FString& InTooltip,
+		ERigVMUserWorkflowType InType,
+		FRigVMWorkflowGetActionsDelegate InGetActionsDelegate,
+		UClass* InOptionsClass)
+	: Title(InTitle)
+	, Tooltip(InTooltip)
+	, Type(InType)
+	, OnGetActionsDelegate(InGetActionsDelegate)
+	, OnGetActionsDynamicDelegate()
+	, OptionsClass(InOptionsClass)
+	{}
+	
+	FORCEINLINE virtual ~FRigVMUserWorkflow() {}
+
+	FORCEINLINE bool IsValid() const
+	{
+		return Type != ERigVMUserWorkflowType::Invalid &&
+			GetOptionsClass() != nullptr &&
+			(OnGetActionsDelegate.IsBound() || OnGetActionsDynamicDelegate.IsBound());
+	}
+
+	FORCEINLINE const FString& GetTitle() const { return Title; }
+	FORCEINLINE const FString& GetTooltip() const { return Tooltip; }
+	FORCEINLINE ERigVMUserWorkflowType GetType() const { return Type; }
+	FORCEINLINE UClass* GetOptionsClass() const { return OptionsClass; }
+
+	TArray<FRigVMUserWorkflowAction> GetActions(const URigVMUserWorkflowOptions* InOptions) const;
+
+protected:
+
+	bool ValidateOptions(const URigVMUserWorkflowOptions* InOptions) const;
+
+	UPROPERTY(BlueprintReadWrite, EditAnywhere, Category = Workflow, meta = (AllowPrivateAccess = true))
+	FString Title;
+
+	UPROPERTY(BlueprintReadWrite, EditAnywhere, Category = Workflow, meta = (AllowPrivateAccess = true))
+	FString Tooltip;
+
+	UPROPERTY(BlueprintReadWrite, EditAnywhere, Category = Workflow, meta = (AllowPrivateAccess = true))
+	ERigVMUserWorkflowType Type;
+
+	FRigVMWorkflowGetActionsDelegate OnGetActionsDelegate;
+
+	UPROPERTY(BlueprintReadWrite, EditAnywhere, Category = Workflow, meta = (ScriptName = "OnGetActions", AllowPrivateAccess = true))
+	FRigVMWorkflowGetActionsDynamicDelegate OnGetActionsDynamicDelegate;
+
+	UPROPERTY(BlueprintReadWrite, EditAnywhere, Category = Workflow, meta = (AllowPrivateAccess = true))
+	TObjectPtr<UClass> OptionsClass;
+};
+
+UCLASS(BlueprintType)
+class RIGVM_API URigVMUserWorkflowOptions : public UObject
+{
+	GENERATED_BODY()
+
+public:
+
+	FORCEINLINE bool IsValid() const { return Subject != nullptr; }
+	bool RequiresDialog() const;
+	FORCEINLINE UObject* GetSubject() const { return Subject.Get(); }
+
+	template<typename T>
+	FORCEINLINE T* GetSubject() const { return Cast<T>(Subject.Get()); }
+
+	FORCEINLINE UObject* GetSubjectChecked() const
+	{
+		UObject* Object = Subject.Get();
+		check(Object);
+		return Object;
+	}
+
+	template<typename T>
+	FORCEINLINE T* GetSubjectChecked() const { return CastChecked<T>(Subject.Get()); }
+
+	const FRigVMUserWorkflow& GetWorkflow() const { return Workflow; }
+
+	void Report(EMessageSeverity::Type InSeverity, const FString& InMessage) const;
+
+	template <typename FmtType, typename... Types>
+	void Reportf(EMessageSeverity::Type InSeverity, const FmtType& Fmt, Types... Args) const
+	{
+		Report(InSeverity, FString::Printf(Fmt, Args...));
+	}
+
+protected:
+
+	virtual bool RequiresDialog(const FProperty* InProperty) const;
+
+	UPROPERTY(BlueprintReadOnly, EditDefaultsOnly, Category = Options, meta = (AllowPrivateAccess = true))
+	TObjectPtr<UObject> Subject;
+
+	UPROPERTY(BlueprintReadOnly, EditDefaultsOnly, Category = Options, meta = (AllowPrivateAccess = true))
+	FRigVMUserWorkflow Workflow;
+
+	FRigVMReportDelegate ReportDelegate;
+
+	friend class URigVMController;
+};
