RigVM: Improve Workflow access in python

diff --git a/Engine/Source/Developer/RigVMDeveloper/Private/RigVMModel/RigVMController.cpp b/Engine/Source/Developer/RigVMDeveloper/Private/RigVMModel/RigVMController.cpp
--- a/Engine/Source/Developer/RigVMDeveloper/Private/RigVMModel/RigVMController.cpp
+++ b/Engine/Source/Developer/RigVMDeveloper/Private/RigVMModel/RigVMController.cpp
@@ -900,8 +900,34 @@ URigVMUnitNode* URigVMController::AddUnitNodeWithDefaults(UScriptStruct* InScrip
 	return Node;
 }
 
+bool URigVMController::SetUnitNodeDefaults(URigVMUnitNode* InNode, const FString& InDefaults, bool bSetupUndoRedo,
+	bool bPrintPythonCommand)
+{
+	if(InNode == nullptr)
+	{
+		return false; 
+	}
+
+	UScriptStruct* ScriptStruct = InNode->GetScriptStruct();
+	if(ScriptStruct == nullptr)
+	{
+		return false;
+	}
+	
+	FStructOnScope StructOnScope(ScriptStruct);
+	FRigVMPinDefaultValueImportErrorContext ErrorPipe;
+	ScriptStruct->ImportText(*InDefaults, StructOnScope.GetStructMemory(), nullptr, PPF_None, &ErrorPipe, FString());
+
+	if(ErrorPipe.NumErrors > 0)
+	{
+		return false;
+	}
+
+	return SetUnitNodeDefaults(InNode, StructOnScope, bSetupUndoRedo, bPrintPythonCommand);
+}
+
 bool URigVMController::SetUnitNodeDefaults(URigVMUnitNode* InNode, const FRigStructScope& InDefaults,
-	bool bSetupUndoRedo, bool bPrintPythonCommand)
+                                           bool bSetupUndoRedo, bool bPrintPythonCommand)
 {
 	if(InNode == nullptr || !InDefaults.IsValid())
 	{
diff --git a/Engine/Source/Developer/RigVMDeveloper/Private/RigVMModel/RigVMNode.cpp b/Engine/Source/Developer/RigVMDeveloper/Private/RigVMModel/RigVMNode.cpp
--- a/Engine/Source/Developer/RigVMDeveloper/Private/RigVMModel/RigVMNode.cpp
+++ b/Engine/Source/Developer/RigVMDeveloper/Private/RigVMModel/RigVMNode.cpp
@@ -536,7 +536,13 @@ TArray<FRigVMUserWorkflow> URigVMNode::GetSupportedWorkflows(ERigVMUserWorkflowT
 	{
 		InSubject = this;
 	}
-	return URigVMUserWorkflowRegistry::Get()->GetWorkflows(InType, nullptr, InSubject);
+
+	const UScriptStruct* Struct = nullptr;
+	if(const URigVMUnitNode* UnitNode = Cast<URigVMUnitNode>(this))
+	{
+		Struct = UnitNode->GetScriptStruct();
+	}
+	return URigVMUserWorkflowRegistry::Get()->GetWorkflows(InType, Struct, InSubject);
 }
 
 #if WITH_EDITOR
diff --git a/Engine/Source/Developer/RigVMDeveloper/Public/RigVMModel/RigVMController.h b/Engine/Source/Developer/RigVMDeveloper/Public/RigVMModel/RigVMController.h
--- a/Engine/Source/Developer/RigVMDeveloper/Public/RigVMModel/RigVMController.h
+++ b/Engine/Source/Developer/RigVMDeveloper/Public/RigVMModel/RigVMController.h
@@ -268,6 +268,7 @@ public:
 	// UnitNode represent a RIGVM_METHOD declaration on a USTRUCT.
 	// This causes a NodeAdded modified event.
 	UFUNCTION(BlueprintCallable, Category = RigVMController)
+	bool SetUnitNodeDefaults(URigVMUnitNode* InNode, const FString& InDefaults, bool bSetupUndoRedo = true, bool bPrintPythonCommand = false);
 	bool SetUnitNodeDefaults(URigVMUnitNode* InNode, const FRigStructScope& InDefaults, bool bSetupUndoRedo = true, bool bPrintPythonCommand = false);
 
 	// Adds a Variable Node to the edited Graph.
diff --git a/Engine/Source/Runtime/RigVM/Public/RigVMCore/RigVMUserWorkflow.h b/Engine/Source/Runtime/RigVM/Public/RigVMCore/RigVMUserWorkflow.h
--- a/Engine/Source/Runtime/RigVM/Public/RigVMCore/RigVMUserWorkflow.h
+++ b/Engine/Source/Runtime/RigVM/Public/RigVMCore/RigVMUserWorkflow.h
@@ -95,11 +95,13 @@ class RIGVM_API URigVMUserWorkflowOptions : public UObject
 
 public:
 
+	UFUNCTION(BlueprintCallable, Category=Options)
 	FORCEINLINE bool IsValid() const { return Subject != nullptr; }
+	
+	UFUNCTION(BlueprintCallable, Category=Options)
 	bool RequiresDialog() const;
-	FORCEINLINE UObject* GetSubject() const { return Subject.Get(); }
 
-	template<typename T>
+	template<typename T = UObject>
 	FORCEINLINE T* GetSubject() const { return Cast<T>(Subject.Get()); }
 
 	FORCEINLINE UObject* GetSubjectChecked() const
@@ -122,6 +124,24 @@ public:
 		Report(InSeverity, FString::Printf(Fmt, Args...));
 	}
 
+	UFUNCTION(BlueprintCallable, Category=Options)
+	FORCEINLINE void ReportError(const FString& InMessage)
+	{
+		Report(EMessageSeverity::Error, InMessage);
+	}
+
+	UFUNCTION(BlueprintCallable, Category=Options)
+	FORCEINLINE void ReportWarning(const FString& InMessage)
+	{
+		Report(EMessageSeverity::Warning, InMessage);
+	}
+
+	UFUNCTION(BlueprintCallable, Category=Options)
+	FORCEINLINE void ReportInfo(const FString& InMessage)
+	{
+		Report(EMessageSeverity::Info, InMessage);
+	}
+
 protected:
 
 	virtual bool RequiresDialog(const FProperty* InProperty) const;
