RigVM: Change inheritance order - make UnitNode inherit from TemplateNode
Also implement a way to register external / auxiliary templates to the registry + unit test

diff --git a/Engine/Config/BaseEngine.ini b/Engine/Config/BaseEngine.ini
--- a/Engine/Config/BaseEngine.ini
+++ b/Engine/Config/BaseEngine.ini
@@ -1055,6 +1055,7 @@ DefaultWorldPartitionHLODLayer=/Engine/Maps/Templates/HLODs/HLODLayer_Instanced.
 +FunctionRedirects=(OldName="KismetMathLibrary.NotEqual_FloatFloat", NewName="NotEqual_DoubleDouble")
 
 ;5.1
++ClassRedirects=(OldName="RigVMTemplateNode", NewName="/Script/RigVMDeveloper.RigVMUnitNode")
 
 ; Blueprint real number support
 +FunctionRedirects=(OldName="KismetSystemLibrary.MakeLiteralFloat", NewName="MakeLiteralDouble")
diff --git a/Engine/Plugins/Animation/ControlRig/Source/ControlRigDeveloper/Private/ControlRigBlueprint.cpp b/Engine/Plugins/Animation/ControlRig/Source/ControlRigDeveloper/Private/ControlRigBlueprint.cpp
--- a/Engine/Plugins/Animation/ControlRig/Source/ControlRigDeveloper/Private/ControlRigBlueprint.cpp
+++ b/Engine/Plugins/Animation/ControlRig/Source/ControlRigDeveloper/Private/ControlRigBlueprint.cpp
@@ -473,9 +473,6 @@ void UControlRigBlueprint::PostLoad()
 		PatchRigElementKeyCacheOnLoad();
 		PatchBoundVariables();
 		PatchPropagateToChildren();
-#if UE_RIGVM_ENABLE_TEMPLATE_NODES
-		ConvertUnitNodesToTemplateNodes();
-#endif
 		PatchParameterNodesOnLoad();
 
 #if WITH_EDITOR
@@ -3918,36 +3915,6 @@ void UControlRigBlueprint::PatchPropagateToChildren()
 	}
 }
 
-void UControlRigBlueprint::ConvertUnitNodesToTemplateNodes()
-{
-	TGuardValue<bool> GuardNotifsSelf(bSuspendModelNotificationsForSelf, true);
-
-	bool bWasDirty = false;
-	if (const UPackage* Package = GetOutermost())
-	{
-		bWasDirty = Package->IsDirty();
-	}
-
-	for (URigVMGraph* Graph : GetAllModels())
-	{
-		URigVMController* Controller = GetOrCreateController(Graph);
-		TArray<URigVMNode*> Nodes = Graph->GetNodes();
-		for (URigVMNode* Node : Nodes)
-		{
-			if(Node->IsInjected())
-			{
-				continue;
-			}
-			Controller->ReplaceUnitNodeWithTemplateNode(Node->GetFName(), false);
-		}
-	}
-
-	if (UPackage* Package = GetOutermost())
-	{
-		Package->SetDirtyFlag(bWasDirty);
-	}
-}
-
 void UControlRigBlueprint::PatchParameterNodesOnLoad()
 {	
 #if WITH_EDITOR
diff --git a/Engine/Plugins/Animation/ControlRig/Source/ControlRigDeveloper/Private/Graph/ControlRigGraphNode.cpp b/Engine/Plugins/Animation/ControlRig/Source/ControlRigDeveloper/Private/Graph/ControlRigGraphNode.cpp
--- a/Engine/Plugins/Animation/ControlRig/Source/ControlRigDeveloper/Private/Graph/ControlRigGraphNode.cpp
+++ b/Engine/Plugins/Animation/ControlRig/Source/ControlRigDeveloper/Private/Graph/ControlRigGraphNode.cpp
@@ -775,9 +775,8 @@ bool UControlRigGraphNode::ShowPaletteIconOnNode() const
 			ModelNode->IsA<URigVMFunctionReturnNode>() ||
 			ModelNode->IsA<URigVMFunctionReferenceNode>() ||
 			ModelNode->IsA<URigVMCollapseNode>() ||
-			ModelNode->IsA<URigVMTemplateNode>() ||
-			ModelNode->IsLoopNode() ||
-			Cast<URigVMUnitNode>(ModelNode) != nullptr;
+			ModelNode->IsA<URigVMUnitNode>() ||
+			ModelNode->IsLoopNode();
 	}
 	return false;
 }
@@ -824,13 +823,16 @@ FSlateIcon UControlRigGraphNode::GetIconAndTint(FLinearColor& OutColor) const
 			}
 		}
 
-		if(URigVMTemplateNode* TemplateNode = Cast<URigVMTemplateNode>(ModelNode))
-		{
-			return TemplateNodeIcon;
-		}
-
 		if (URigVMUnitNode *UnitNode = Cast<URigVMUnitNode>(ModelNode))
 		{
+			if(const FRigVMTemplate* Template = UnitNode->GetTemplate())
+			{
+				if(Template->NumPermutations() > 1)
+				{
+					return TemplateNodeIcon;
+				}
+			}
+			
 			FString IconPath;
 			const int32 NumOfIconPathNames = 4;
 			
diff --git a/Engine/Plugins/Animation/ControlRig/Source/ControlRigDeveloper/Public/ControlRigBlueprint.h b/Engine/Plugins/Animation/ControlRig/Source/ControlRigDeveloper/Public/ControlRigBlueprint.h
--- a/Engine/Plugins/Animation/ControlRig/Source/ControlRigDeveloper/Public/ControlRigBlueprint.h
+++ b/Engine/Plugins/Animation/ControlRig/Source/ControlRigDeveloper/Public/ControlRigBlueprint.h
@@ -570,7 +570,6 @@ private:
 	void PatchBoundVariables();
 	void PatchVariableNodesWithIncorrectType();
 	void PatchPropagateToChildren();
-	void ConvertUnitNodesToTemplateNodes();
 	void PatchParameterNodesOnLoad();
 
 	TMap<FName, int32> AddedMemberVariableMap;
diff --git a/Engine/Plugins/Animation/ControlRig/Source/ControlRigEditor/Private/ControlRigEditorModule.cpp b/Engine/Plugins/Animation/ControlRig/Source/ControlRigEditor/Private/ControlRigEditorModule.cpp
--- a/Engine/Plugins/Animation/ControlRig/Source/ControlRigEditor/Private/ControlRigEditorModule.cpp
+++ b/Engine/Plugins/Animation/ControlRig/Source/ControlRigEditor/Private/ControlRigEditorModule.cpp
@@ -1015,6 +1015,12 @@ void FControlRigEditorModule::GetTypeActions(UControlRigBlueprint* CRB, FBluepri
 			continue;
 		}
 
+		// ignore templates which don't have a function backing it up
+		if(Template.GetPermutation(0) == nullptr)
+		{
+			continue;
+		}
+
 		FText NodeCategory = FText::FromString(Template.GetCategory());
 		FText MenuDesc = FText::FromName(Template.GetName());
 		FText ToolTip = Template.GetTooltipText();
diff --git a/Engine/Source/Developer/RigVMDeveloper/Private/RigVMModel/Nodes/RigVMTemplateNode.cpp b/Engine/Source/Developer/RigVMDeveloper/Private/RigVMModel/Nodes/RigVMTemplateNode.cpp
--- a/Engine/Source/Developer/RigVMDeveloper/Private/RigVMModel/Nodes/RigVMTemplateNode.cpp
+++ b/Engine/Source/Developer/RigVMDeveloper/Private/RigVMModel/Nodes/RigVMTemplateNode.cpp
@@ -298,15 +298,6 @@ bool URigVMTemplateNode::IsFullyUnresolved() const
 	return GetResolvedPermutations().Num() == GetTemplate()->NumPermutations();
 }
 
-TArray<UScriptStruct*> URigVMTemplateNode::GetSupportedUnitStructs() const
-{
-	if(const FRigVMTemplate* Template = GetTemplate())
-	{
-		return URigVMController::GetUnitStructsForTemplate(Template->GetNotation());
-	}
-	return TArray<UScriptStruct*>();
-}
-
 FString URigVMTemplateNode::GetInitialDefaultValueForPin(const FName& InRootPinName, const TArray<int32>& InPermutationIndices) const
 {
 	if(GetTemplate() == nullptr)
diff --git a/Engine/Source/Developer/RigVMDeveloper/Private/RigVMModel/Nodes/RigVMUnitNode.cpp b/Engine/Source/Developer/RigVMDeveloper/Private/RigVMModel/Nodes/RigVMUnitNode.cpp
--- a/Engine/Source/Developer/RigVMDeveloper/Private/RigVMModel/Nodes/RigVMUnitNode.cpp
+++ b/Engine/Source/Developer/RigVMDeveloper/Private/RigVMModel/Nodes/RigVMUnitNode.cpp
@@ -7,6 +7,25 @@
 #include "RigVMModel/RigVMControllerActions.h"
 #include "RigVMUserWorkflowRegistry.h"
 
+void URigVMUnitNode::PostLoad()
+{
+	Super::PostLoad();
+
+	// if we have a script struct but no notation let's figure out the template
+	if(GetScriptStruct() != nullptr && GetTemplate() == nullptr)
+	{
+		if(const FRigVMFunction* Function = FRigVMRegistry::Get().FindFunction(GetScriptStruct(), *GetMethodName().ToString()))
+		{
+			if(Function->TemplateIndex != INDEX_NONE)
+			{
+				const FRigVMTemplate& Template = FRigVMRegistry::Get().GetTemplates()[Function->TemplateIndex];
+				TemplateNotation = Template.GetNotation();
+				ResolvedFunctionName = Function->GetName();
+			}	
+		}
+	}
+}
+
 FString URigVMUnitNode::GetNodeTitle() const
 {
 	if (UScriptStruct* Struct = GetScriptStruct())
@@ -283,7 +302,11 @@ FName URigVMUnitNode::GetNextAggregateName(const FName& InLastAggregatePinName)
 
 UScriptStruct* URigVMUnitNode::GetScriptStruct() const
 {
-	return ScriptStruct;
+	if(UScriptStruct* ResolvedStruct = Super::GetScriptStruct())
+	{
+		return ResolvedStruct;
+	}
+	return ScriptStruct_DEPRECATED;
 }
 
 bool URigVMUnitNode::IsLoopNode() const
@@ -299,7 +322,12 @@ bool URigVMUnitNode::IsLoopNode() const
 
 FName URigVMUnitNode::GetMethodName() const
 {
-	return MethodName;
+	const FName ResolvedMethodName = Super::GetMethodName();
+	if(!ResolvedMethodName.IsNone())
+	{
+		return ResolvedMethodName;
+	}
+	return MethodName_DEPRECATED;
 }
 
 FString URigVMUnitNode::GetStructDefaultValue() const
diff --git a/Engine/Source/Developer/RigVMDeveloper/Private/RigVMModel/RigVMController.cpp b/Engine/Source/Developer/RigVMDeveloper/Private/RigVMModel/RigVMController.cpp
--- a/Engine/Source/Developer/RigVMDeveloper/Private/RigVMModel/RigVMController.cpp
+++ b/Engine/Source/Developer/RigVMDeveloper/Private/RigVMModel/RigVMController.cpp
@@ -707,7 +707,7 @@ URigVMUnitNode* URigVMController::AddUnitNode(UScriptStruct* InScriptStruct, con
 		}
 
 		const FString Name = GetValidNodeName(InNodeName.IsEmpty() ? InScriptStruct->GetName() : InNodeName);
-		URigVMTemplateNode* TemplateNode = AddTemplateNode(Template->GetNotation(), InPosition, Name, bSetupUndoRedo, bPrintPythonCommand);
+		URigVMUnitNode* TemplateNode = Cast<URigVMUnitNode>(AddTemplateNode(Template->GetNotation(), InPosition, Name, bSetupUndoRedo, bPrintPythonCommand));
 		if(TemplateNode == nullptr)
 		{
 			CancelUndoBracket();
@@ -752,8 +752,7 @@ URigVMUnitNode* URigVMController::AddUnitNode(UScriptStruct* InScriptStruct, con
 	
 	FString Name = GetValidNodeName(InNodeName.IsEmpty() ? InScriptStruct->GetName() : InNodeName);
 	URigVMUnitNode* Node = NewObject<URigVMUnitNode>(Graph, *Name);
-	Node->ScriptStruct = InScriptStruct;
-	Node->MethodName = InMethodName;
+	Node->ResolvedFunctionName = Function->GetName();
 	Node->Position = InPosition;
 	Node->NodeTitle = InScriptStruct->GetMetaData(TEXT("DisplayName"));
 	
@@ -1626,82 +1625,6 @@ URigVMVariableNode* URigVMController::ReplaceParameterNodeWithVariable(const FNa
 	return nullptr;
 }
 
-URigVMTemplateNode* URigVMController::ReplaceUnitNodeWithTemplateNode(const FName& InNodeName, bool bSetupUndoRedo)
-{
-	if(!IsValidGraph())
-	{
-		return nullptr;
-	}
-
-	if (!bIsTransacting && !IsGraphEditable())
-	{
-		return nullptr;
-	}
-
-	URigVMUnitNode* UnitNode = Cast<URigVMUnitNode>(GetGraph()->FindNodeByName(InNodeName));
-	if(UnitNode == nullptr)
-	{
-		return nullptr;
-	}
-	if(UnitNode->IsA<URigVMTemplateNode>())
-	{
-		return Cast<URigVMTemplateNode>(UnitNode);
-	}
-
-	UScriptStruct* ScriptStruct = UnitNode->GetScriptStruct();
-	if(ScriptStruct == nullptr)
-	{
-		return nullptr;
-	}
-	
-	const FRigVMFunction* Function = FRigVMRegistry::Get().FindFunction(ScriptStruct, *UnitNode->GetMethodName().ToString());
-	if(Function == nullptr)
-	{
-		return nullptr;
-	}
-			
-	const FRigVMTemplate* Template = Function->GetTemplate();
-	if(Template == nullptr)
-	{
-		return nullptr;
-	}
-
-	FRigVMBaseAction Action;
-	if (bSetupUndoRedo)
-	{
-		Action.Title = FString::Printf(TEXT("Replace Unit with Template"));
-		ActionStack->BeginAction(Action);
-	}
-
-	const TMap<FString, FPinState> PinStates = GetPinStates(UnitNode);
-	TArray<URigVMLink*> Links = UnitNode->GetLinks();
-	DetachLinksFromPinObjects(&Links);
-
-	const FVector2D NodePosition = UnitNode->GetPosition();
-	RemoveNode(UnitNode, bSetupUndoRedo);
-	URigVMTemplateNode* TemplateNode = AddTemplateNode(Template->GetNotation(), NodePosition, InNodeName.ToString(), bSetupUndoRedo);
-	if(TemplateNode == nullptr)
-	{
-		if(bSetupUndoRedo)
-		{
-			ActionStack->CancelAction(Action, this);
-		}
-		return nullptr;
-	}
-
-	FullyResolveTemplateNode(TemplateNode, Function, bSetupUndoRedo);
-
-	ApplyPinStates(TemplateNode, PinStates);
-	ReattachLinksToPinObjects(false, &Links);
-
-	if(bSetupUndoRedo)
-	{
-		ActionStack->EndAction(Action);
-	}
-	
-	return TemplateNode;
-}
-
 bool URigVMController::UnresolveTemplateNodes(const TArray<FName>& InNodeNames, bool bSetupUndoRedo, bool bPrintPythonCommand)
 {
 	if (!IsValidGraph())
@@ -11752,7 +11675,21 @@ URigVMTemplateNode* URigVMController::AddTemplateNode(const FName& InNotation, c
 	}
 
 	FString Name = GetValidNodeName(InNodeName.IsEmpty() ? Template->GetName().ToString() : InNodeName);
-	URigVMTemplateNode* Node = NewObject<URigVMTemplateNode>(Graph, *Name);
+	URigVMTemplateNode* Node = nullptr;
+
+	// determine what kind of node we need to create
+	const UScriptStruct* PotentialUnitStruct = Template->GetPermutation(0)->Struct;
+	if(PotentialUnitStruct && PotentialUnitStruct->IsChildOf(FRigVMStruct::StaticStruct()))
+	{
+		Node = NewObject<URigVMUnitNode>(Graph, *Name); 
+	}
+
+	if(Node == nullptr)
+	{
+		ReportErrorf(TEXT("Template node '%s' cannot be created. Unknown template."), *InNotation.ToString());
+		return nullptr;
+	}
+	
 	Node->TemplateNotation = Template->GetNotation();
 	Node->Position = InPosition;
 
@@ -13074,7 +13011,7 @@ FProperty* URigVMController::FindPropertyForPin(const FString& InPinPath)
 	{
 		int32 PartIndex = 1; // cut off the first one since it's the node
 
-		UStruct* Struct = UnitNode->ScriptStruct;
+		UStruct* Struct = UnitNode->GetScriptStruct();
 		FProperty* Property = Struct->FindPropertyByName(*Parts[PartIndex++]);
 
 		while (PartIndex < Parts.Num() && Property != nullptr)
@@ -14629,8 +14566,7 @@ bool URigVMController::FullyResolveTemplateNode(URigVMTemplateNode* InNode, cons
 	{
 		ReportErrorf(TEXT("Template %s doesn't support function %s::%s."),
 			*Template->GetNotation().ToString(),
-			*InFunctionToResolve->Struct->GetStructCPPName(),
-			*InFunctionToResolve->GetMethodName().ToString()
+			*InFunctionToResolve->GetName()
 		);
 		return false;
 	}
@@ -14766,11 +14702,14 @@ bool URigVMController::FullyResolveTemplateNode(URigVMTemplateNode* InNode, cons
 
 	if (UnitNodeCreatedContext.IsValid())
 	{
-		if (TSharedPtr<FStructOnScope> StructScope = InNode->ConstructStructInstance())
+		if(URigVMUnitNode* UnitNode = Cast<URigVMUnitNode>(InNode))
 		{
-			TGuardValue<FName> NodeNameScope(UnitNodeCreatedContext.NodeName, InNode->GetFName());
-			FRigVMStruct* StructInstance = (FRigVMStruct*)StructScope->GetStructMemory();
-			StructInstance->OnUnitNodeCreated(UnitNodeCreatedContext);
+			if (TSharedPtr<FStructOnScope> StructScope = UnitNode->ConstructStructInstance())
+			{
+				TGuardValue<FName> NodeNameScope(UnitNodeCreatedContext.NodeName, UnitNode->GetFName());
+				FRigVMStruct* StructInstance = (FRigVMStruct*)StructScope->GetStructMemory();
+				StructInstance->OnUnitNodeCreated(UnitNodeCreatedContext);
+			}
 		}
 	}
 	
diff --git a/Engine/Source/Developer/RigVMDeveloper/Public/RigVMModel/Nodes/RigVMTemplateNode.h b/Engine/Source/Developer/RigVMDeveloper/Public/RigVMModel/Nodes/RigVMTemplateNode.h
--- a/Engine/Source/Developer/RigVMDeveloper/Public/RigVMModel/Nodes/RigVMTemplateNode.h
+++ b/Engine/Source/Developer/RigVMDeveloper/Public/RigVMModel/Nodes/RigVMTemplateNode.h
@@ -2,8 +2,7 @@
 
 #pragma once
 
-#include "RigVMUnitNode.h"
-#include "RigVMModel/Nodes/RigVMUnitNode.h"
+#include "RigVMModel/RigVMNode.h"
 #include "RigVMCore/RigVMTemplate.h"
 #include "RigVMTemplateNode.generated.h"
 
@@ -13,7 +12,7 @@
  * the template's template.
  */
 UCLASS(BlueprintType)
-class RIGVMDEVELOPER_API URigVMTemplateNode : public URigVMUnitNode
+class RIGVMDEVELOPER_API URigVMTemplateNode : public URigVMNode
 {
 	GENERATED_BODY()
 
@@ -22,13 +21,17 @@ public:
 	// default constructor
 	URigVMTemplateNode();
 
-	// URigVMUnitNode interface
-	virtual UScriptStruct* GetScriptStruct() const;
+	// URigVMNode interface
 	virtual FString GetNodeTitle() const override;
 	virtual FName GetMethodName() const;
 	virtual  FText GetToolTipText() const override;
 	virtual FText GetToolTipTextForPin(const URigVMPin* InPin) const override;
 
+	// Returns the UStruct for this unit node
+	// (the struct declaring the RIGVM_METHOD)
+	UFUNCTION(BlueprintCallable, Category = RigVMUnitNode)
+	virtual UScriptStruct* GetScriptStruct() const;
+
 	// Returns the notation of the node
 	UFUNCTION(BlueprintPure, Category = Template)
 	FName GetNotation() const;
@@ -62,21 +65,13 @@ public:
 	UFUNCTION(BlueprintPure, Category = Template)
 	bool IsFullyUnresolved() const;
 
-	// returns the sole resolved unit struct
-	UFUNCTION(BlueprintPure, Category = Template)
-	UScriptStruct* GetResolvedUnitStruct() const { return GetScriptStruct(); }
-
-	// returns the list of all supported unit structs
-	UFUNCTION(BlueprintPure, Category = Template)
-	TArray<UScriptStruct*> GetSupportedUnitStructs() const;
-
 	// returns a default value for pin if it is known
 	FString GetInitialDefaultValueForPin(const FName& InRootPinName, const TArray<int32>& InPermutationIndices = TArray<int32>()) const;
 
 	// returns the display name for a pin
 	FName GetDisplayNameForPin(const FName& InRootPinName, const TArray<int32>& InPermutationIndices = TArray<int32>()) const;
 
-private:
+protected:
 
 	void InvalidateCache();
 
diff --git a/Engine/Source/Developer/RigVMDeveloper/Public/RigVMModel/Nodes/RigVMUnitNode.h b/Engine/Source/Developer/RigVMDeveloper/Public/RigVMModel/Nodes/RigVMUnitNode.h
--- a/Engine/Source/Developer/RigVMDeveloper/Public/RigVMModel/Nodes/RigVMUnitNode.h
+++ b/Engine/Source/Developer/RigVMDeveloper/Public/RigVMModel/Nodes/RigVMUnitNode.h
@@ -2,7 +2,7 @@
 
 #pragma once
 
-#include "RigVMModel/RigVMNode.h"
+#include "RigVMModel/Nodes/RigVMTemplateNode.h"
 #include "RigVMCore/RigVMStruct.h"
 #include "UObject/StructOnScope.h"
 #include "RigVMUnitNode.generated.h"
@@ -13,12 +13,15 @@
  * struct UPROPERTY members.
  */
 UCLASS(BlueprintType)
-class RIGVMDEVELOPER_API URigVMUnitNode : public URigVMNode
+class RIGVMDEVELOPER_API URigVMUnitNode : public URigVMTemplateNode
 {
 	GENERATED_BODY()
 
 public:
 
+	// UObject interface
+	virtual void PostLoad() override;
+
 	// Override node functions
 	virtual FString GetNodeTitle() const override;
 	virtual FText GetToolTipText() const override;
@@ -41,12 +44,9 @@ public:
 	bool IsDeprecated() const;
 	FString GetDeprecatedMetadata() const;
 
-	
-	
-	// Returns the UStruct for this unit node
-	// (the struct declaring the RIGVM_METHOD)
-	UFUNCTION(BlueprintCallable, Category = RigVMUnitNode)
-	virtual UScriptStruct* GetScriptStruct() const;
+
+	// URigVMTemplateNode interface
+	virtual UScriptStruct* GetScriptStruct() const override;
 
 	// Returns the name of the declared RIGVM_METHOD
 	UFUNCTION(BlueprintCallable, Category = RigVMUnitNode)
@@ -86,10 +86,10 @@ protected:
 private:
 
 	UPROPERTY()
-	TObjectPtr<UScriptStruct> ScriptStruct;
+	TObjectPtr<UScriptStruct> ScriptStruct_DEPRECATED;
 
 	UPROPERTY()
-	FName MethodName;
+	FName MethodName_DEPRECATED;
 
 	friend class URigVMController;
 };
diff --git a/Engine/Source/Developer/RigVMDeveloper/Public/RigVMModel/RigVMController.h b/Engine/Source/Developer/RigVMDeveloper/Public/RigVMModel/RigVMController.h
--- a/Engine/Source/Developer/RigVMDeveloper/Public/RigVMModel/RigVMController.h
+++ b/Engine/Source/Developer/RigVMDeveloper/Public/RigVMModel/RigVMController.h
@@ -303,10 +303,6 @@ public:
 	UFUNCTION(BlueprintCallable, Category = RigVMController)
 	URigVMVariableNode* ReplaceParameterNodeWithVariable(const FName& InNodeName, const FName& InVariableName, const FString& InCPPType, UObject* InCPPTypeObject, bool bSetupUndoRedo);
 
-	// Refreshes the unit node to be backed up by a template instead 
-	UFUNCTION(BlueprintCallable, Category = RigVMController)
-	URigVMTemplateNode* ReplaceUnitNodeWithTemplateNode(const FName& InNodeName, bool bSetupUndoRedo);
-
 	// Turns a resolved templated node(s) back into its template.
 	UFUNCTION(BlueprintCallable, Category = RigVMController)
 	bool UnresolveTemplateNodes(const TArray<FName>& InNodeNames, bool bSetupUndoRedo = true, bool bPrintPythonCommand = false);
diff --git a/Engine/Source/Runtime/RigVM/Private/RigVMCore/RigVMFunction.cpp b/Engine/Source/Runtime/RigVM/Private/RigVMCore/RigVMFunction.cpp
--- a/Engine/Source/Runtime/RigVM/Private/RigVMCore/RigVMFunction.cpp
+++ b/Engine/Source/Runtime/RigVM/Private/RigVMCore/RigVMFunction.cpp
@@ -4,6 +4,11 @@
 #include "RigVMCore/RigVMTemplate.h"
 #include "RigVMCore/RigVMRegistry.h"
 
+FString FRigVMFunction::GetName() const
+{
+	return Name;
+}
+
 FName FRigVMFunction::GetMethodName() const
 {
 	FString FullName(Name);
diff --git a/Engine/Source/Runtime/RigVM/Private/RigVMCore/RigVMRegistry.cpp b/Engine/Source/Runtime/RigVM/Private/RigVMCore/RigVMRegistry.cpp
--- a/Engine/Source/Runtime/RigVM/Private/RigVMCore/RigVMRegistry.cpp
+++ b/Engine/Source/Runtime/RigVM/Private/RigVMCore/RigVMRegistry.cpp
@@ -118,4 +118,228 @@ const TArray<FRigVMTemplate>& FRigVMRegistry::GetTemplates() const
 	return Templates;
 }
 
+const FRigVMTemplate* FRigVMRegistry::GetOrAddTemplateFromArguments(const FName& InName, const TArray<FRigVMTemplateArgument>& InArguments)
+{
+	FRigVMTemplate Template(InName, InArguments, INDEX_NONE);
+	if(const FRigVMTemplate* ExistingTemplate = FindTemplate(Template.GetNotation()))
+	{
+		return ExistingTemplate;
+	}
+
+	// we only support to ask for templates here which provide singleton types
+	for(const FRigVMTemplateArgument& Argument : InArguments)
+	{
+		if(!Argument.bSingleton || Argument.Types.Num() > 1)
+		{
+			return nullptr;
+		}
+	}
+
+	int32 NumPermutations = 1;
+
+	// if any of the arguments are wildcards we'll need to update the types
+	for(FRigVMTemplateArgument& Argument : Template.Arguments)
+	{
+		if(Argument.Types[0].IsWildCard())
+		{
+			static TArray<FRigVMTemplateArgument::FType> AllTypes, AllArrayTypes;
+			if(AllTypes.IsEmpty())
+			{
+				AllTypes.Add(FRigVMTemplateArgument::FType(RigVMTypeUtils::BoolType));
+				AllTypes.Add(FRigVMTemplateArgument::FType(RigVMTypeUtils::Int32Type));
+				AllTypes.Add(FRigVMTemplateArgument::FType(RigVMTypeUtils::UInt8Type));
+				AllTypes.Add(FRigVMTemplateArgument::FType(RigVMTypeUtils::FloatType));
+				AllTypes.Add(FRigVMTemplateArgument::FType(RigVMTypeUtils::DoubleType));
+				AllTypes.Add(FRigVMTemplateArgument::FType(RigVMTypeUtils::FNameType));
+				AllTypes.Add(FRigVMTemplateArgument::FType(RigVMTypeUtils::FStringType));
+
+				AllArrayTypes.Add(FRigVMTemplateArgument::FType(RigVMTypeUtils::BoolArrayType));
+				AllArrayTypes.Add(FRigVMTemplateArgument::FType(RigVMTypeUtils::Int32ArrayType));
+				AllArrayTypes.Add(FRigVMTemplateArgument::FType(RigVMTypeUtils::UInt8ArrayType));
+				AllArrayTypes.Add(FRigVMTemplateArgument::FType(RigVMTypeUtils::FloatArrayType));
+				AllArrayTypes.Add(FRigVMTemplateArgument::FType(RigVMTypeUtils::DoubleArrayType));
+				AllArrayTypes.Add(FRigVMTemplateArgument::FType(RigVMTypeUtils::FNameArrayType));
+				AllArrayTypes.Add(FRigVMTemplateArgument::FType(RigVMTypeUtils::FStringArrayType));
+
+				struct FTypeTraverser
+				{
+					static EObjectFlags DisallowedFlags()
+					{
+						return RF_BeginDestroyed | RF_FinishDestroyed | RF_PendingKill;
+					}
+
+					static EObjectFlags NeededFlags()
+					{
+						return RF_Public;
+					}
+					
+					static bool IsAllowedType(const FProperty* InProperty)
+					{
+						if(!InProperty->HasAnyPropertyFlags(
+							CPF_BlueprintVisible |
+							CPF_BlueprintReadOnly |
+							CPF_Edit))
+						{
+							return false;
+						}
+
+						if(InProperty->IsA<FBoolProperty>() ||
+							InProperty->IsA<FUInt32Property>() ||
+							InProperty->IsA<FInt8Property>() ||
+							InProperty->IsA<FInt16Property>() ||
+							InProperty->IsA<FIntProperty>() ||
+							InProperty->IsA<FInt64Property>() ||
+							InProperty->IsA<FFloatProperty>() ||
+							InProperty->IsA<FDoubleProperty>() ||
+							InProperty->IsA<FNumericProperty>() ||
+							InProperty->IsA<FNameProperty>() ||
+							InProperty->IsA<FStrProperty>())
+						{
+							return true;
+						}
+
+						if(const FArrayProperty* ArrayProperty  = CastField<FArrayProperty>(InProperty))
+						{
+							return IsAllowedType(ArrayProperty->Inner);
+						}
+						if(const FStructProperty* StructProperty = CastField<FStructProperty>(InProperty))
+						{
+							return IsAllowedType(StructProperty->Struct);
+						}
+						if(const FObjectProperty* ObjectProperty = CastField<FObjectProperty>(InProperty))
+						{
+							return IsAllowedType(ObjectProperty->PropertyClass);
+						}
+						if(const FEnumProperty* EnumProperty = CastField<FEnumProperty>(InProperty))
+						{
+							return IsAllowedType(EnumProperty->GetEnum());
+						}
+						if(const FByteProperty* ByteProperty = CastField<FByteProperty>(InProperty))
+						{
+							if(const UEnum* Enum = ByteProperty->Enum)
+							{
+								return IsAllowedType(Enum);
+							}
+							return true;
+						}
+						return false;
+					}
+
+					static bool IsAllowedType(const UEnum* InEnum)
+					{
+						return !InEnum->HasAnyFlags(DisallowedFlags()) && InEnum->HasAllFlags(NeededFlags());
+					}
+
+					static bool IsAllowedType(const UStruct* InStruct)
+					{
+						if(InStruct->HasAnyFlags(DisallowedFlags()) || !InStruct->HasAllFlags(NeededFlags()))
+						{
+							return false;
+						}
+						if(InStruct->IsChildOf(FRigVMExecuteContext::StaticStruct()))
+						{
+							return false;
+						}
+						if(InStruct->IsChildOf(FRigVMStruct::StaticStruct()))
+						{
+							return false;
+						}
+						for (TFieldIterator<FProperty> It(InStruct); It; ++It)
+						{
+							if(!IsAllowedType(*It))
+							{
+								return false;
+							}
+						}
+						return true;
+					}
+
+					static bool IsAllowedType(const UClass* InClass)
+					{
+						if(InClass->HasAnyClassFlags(CLASS_Hidden | CLASS_Abstract))
+						{
+							return false;
+						}
+
+						// note: currently we don't allow UObjects
+						return false;
+						//return IsAllowedType(Cast<UStruct>(InClass));
+					}
+				};
+
+				// add all structs
+				for (TObjectIterator<UScriptStruct> ScriptIt; ScriptIt; ++ScriptIt)
+				{
+					UScriptStruct* ScriptStruct = *ScriptIt;
+					if(!FTypeTraverser::IsAllowedType(ScriptStruct))
+					{
+						continue;
+					}
+					const FString CPPType = ScriptStruct->GetStructCPPName();
+					AllTypes.Add(FRigVMTemplateArgument::FType(CPPType, ScriptStruct));
+					AllArrayTypes.Add(FRigVMTemplateArgument::FType(RigVMTypeUtils::ArrayTypeFromBaseType(CPPType), ScriptStruct));
+				}
+
+				// add all enums
+				for (TObjectIterator<UEnum> EnumIt; EnumIt; ++EnumIt)
+				{
+					UEnum* Enum = (*EnumIt);
+					if(!FTypeTraverser::IsAllowedType(Enum))
+					{
+						continue;
+					}
+					const FString CPPType = Enum->CppType.IsEmpty() ? Enum->GetName() : Enum->CppType;
+					AllTypes.Add(FRigVMTemplateArgument::FType(CPPType, Enum));
+					AllArrayTypes.Add(FRigVMTemplateArgument::FType(RigVMTypeUtils::ArrayTypeFromBaseType(CPPType), Enum));
+				}
+
+				// add all classes
+				for (TObjectIterator<UClass> ClassIt; ClassIt; ++ClassIt)
+				{
+					UClass* Class = *ClassIt;
+					if(!FTypeTraverser::IsAllowedType(Class))
+					{
+						continue;
+					}
+
+					const FString CPPType = Class->GetPrefixCPP() + Class->GetName();
+					AllTypes.Add(FRigVMTemplateArgument::FType(CPPType, Class));
+					AllArrayTypes.Add(FRigVMTemplateArgument::FType(RigVMTypeUtils::ArrayTypeFromBaseType(CPPType), Class));
+				}
+			}
 
+			Argument.Types = Argument.Types[0].IsArray() ? AllArrayTypes : AllTypes;
+			Argument.bSingleton = false;
+			NumPermutations = Argument.Types.Num(); 
+		}
+	}
+
+	// if we have more than one permutation we may need to upgrade the types for singleton args
+	if(NumPermutations > 1)
+	{
+		for(FRigVMTemplateArgument& Argument : Template.Arguments)
+		{
+			if(Argument.Types.Num() == 1)
+			{
+				const FRigVMTemplateArgument::FType Type = Argument.Types[0];
+				Argument.Types.SetNum(NumPermutations);
+				for(int32 Index=0;Index<NumPermutations;Index++)
+				{
+					Argument.Types[Index] = Type;
+				}
+				Argument.bSingleton = true;
+			}
+		}
+	}
+
+	Template.Permutations.SetNum(NumPermutations);
+	for(int32 Index=0;Index<NumPermutations;Index++)
+	{
+		Template.Permutations[Index] = INDEX_NONE;
+	}
+
+	const int32 Index = Templates.Add(Template);
+	Templates[Index].Index = Index;
+	TemplateNotationToIndex.Add(Template.GetNotation(), Index);
+	return &Templates[Index];
+}
diff --git a/Engine/Source/Runtime/RigVM/Private/RigVMCore/RigVMTemplate.cpp b/Engine/Source/Runtime/RigVM/Private/RigVMCore/RigVMTemplate.cpp
--- a/Engine/Source/Runtime/RigVM/Private/RigVMCore/RigVMTemplate.cpp
+++ b/Engine/Source/Runtime/RigVM/Private/RigVMCore/RigVMTemplate.cpp
@@ -52,6 +52,15 @@ FRigVMTemplateArgument::FRigVMTemplateArgument(FProperty* InProperty)
 	Types.Add(Type);
 }
 
+FRigVMTemplateArgument::FRigVMTemplateArgument(const FName& InName, ERigVMPinDirection InDirection, const FType& InType)
+: Index(INDEX_NONE)
+, Name(InName)
+, Direction(InDirection)
+, bSingleton(true)
+, Types({InType})
+{
+}
+
 bool FRigVMTemplateArgument::SupportsType(const FString& InCPPType, const TArray<int32>& InPermutationIndices, FType* OutType) const
 {
 	bool bFoundMatch = false;
@@ -197,7 +206,32 @@ FRigVMTemplate::FRigVMTemplate(UScriptStruct* InStruct, const FString& InTemplat
 
 	if (ArgumentNotations.Num() > 0)
 	{
-		FString NotationStr = FString::Printf(TEXT("%s(%s)"), *InTemplateName, *FString::Join(ArgumentNotations, TEXT(",")));
+		const FString NotationStr = FString::Printf(TEXT("%s(%s)"), *InTemplateName, *FString::Join(ArgumentNotations, TEXT(",")));
+		Notation = *NotationStr;
+		Permutations.Add(InFunctionIndex);
+	}
+}
+
+FRigVMTemplate::FRigVMTemplate(const FName& InTemplateName, const TArray<FRigVMTemplateArgument>& InArguments, int32 InFunctionIndex)
+	: Index(INDEX_NONE)
+	, Notation(NAME_None)
+{
+	TArray<FString> ArgumentNotations;
+	for (const FRigVMTemplateArgument& InArgument : InArguments)
+	{
+		FRigVMTemplateArgument Argument = InArgument;
+		Argument.Index = Arguments.Num();
+
+		if(IsValidArgumentForTemplate(Argument))
+		{
+			Arguments.Add(Argument);
+			ArgumentNotations.Add(GetArgumentNotation(Argument));
+		}
+	}
+
+	if (ArgumentNotations.Num() > 0)
+	{
+		const FString NotationStr = FString::Printf(TEXT("%s(%s)"), *InTemplateName.ToString(), *FString::Join(ArgumentNotations, TEXT(",")));
 		Notation = *NotationStr;
 		Permutations.Add(InFunctionIndex);
 	}
@@ -580,7 +614,12 @@ bool FRigVMTemplate::ArgumentSupportsType(const FName& InArgumentName, const FSt
 const FRigVMFunction* FRigVMTemplate::GetPermutation(int32 InIndex) const
 {
 	const FRigVMRegistry& Registry = FRigVMRegistry::Get();
-	return &Registry.GetFunctions()[Permutations[InIndex]];
+	const int32 FunctionIndex = Permutations[InIndex];
+	if(Registry.GetFunctions().IsValidIndex(FunctionIndex))
+	{
+		return &Registry.GetFunctions()[Permutations[InIndex]];
+	}
+	return nullptr;
 }
 
 bool FRigVMTemplate::ContainsPermutation(const FRigVMFunction* InPermutation) const
diff --git a/Engine/Source/Runtime/RigVM/Public/RigVMCore/RigVMFunction.h b/Engine/Source/Runtime/RigVM/Public/RigVMCore/RigVMFunction.h
--- a/Engine/Source/Runtime/RigVM/Public/RigVMCore/RigVMFunction.h
+++ b/Engine/Source/Runtime/RigVM/Public/RigVMCore/RigVMFunction.h
@@ -84,6 +84,7 @@ struct RIGVM_API FRigVMFunction
 	}
 
 	FORCEINLINE bool IsValid() const { return Name != nullptr && FunctionPtr != nullptr; }
+	FString GetName() const;
 	FName GetMethodName() const;
 	FString GetModuleName() const;
 	FString GetModuleRelativeHeaderPath() const;
diff --git a/Engine/Source/Runtime/RigVM/Public/RigVMCore/RigVMRegistry.h b/Engine/Source/Runtime/RigVM/Public/RigVMCore/RigVMRegistry.h
--- a/Engine/Source/Runtime/RigVM/Public/RigVMCore/RigVMRegistry.h
+++ b/Engine/Source/Runtime/RigVM/Public/RigVMCore/RigVMRegistry.h
@@ -44,6 +44,9 @@ public:
 	// Returns all current rigvm functions
 	const TArray<FRigVMTemplate>& GetTemplates() const;
 
+	// Defines and retrieves a template given its arguments
+	const FRigVMTemplate* GetOrAddTemplateFromArguments(const FName& InName, const TArray<FRigVMTemplateArgument>& InArguments);
+
 private:
 
 	static const FName TemplateNameMetaName;
diff --git a/Engine/Source/Runtime/RigVM/Public/RigVMCore/RigVMTemplate.h b/Engine/Source/Runtime/RigVM/Public/RigVMCore/RigVMTemplate.h
--- a/Engine/Source/Runtime/RigVM/Public/RigVMCore/RigVMTemplate.h
+++ b/Engine/Source/Runtime/RigVM/Public/RigVMCore/RigVMTemplate.h
@@ -112,6 +112,8 @@ struct RIGVM_API FRigVMTemplateArgument
 		}
 	};
 
+	FRigVMTemplateArgument(const FName& InName, ERigVMPinDirection InDirection, const FType& InType);
+
 	// returns the name of the argument
 	const FName& GetName() const { return Name; }
 
@@ -152,6 +154,7 @@ protected:
 
 	friend struct FRigVMTemplate;
 	friend class URigVMController;
+	friend struct FRigVMRegistry;
 };
 
 /**
@@ -251,6 +254,9 @@ private:
 	// Constructor from a struct, a template name and a function index
 	FRigVMTemplate(UScriptStruct* InStruct, const FString& InTemplateName, int32 InFunctionIndex);
 
+	// Constructor from a template name, arguments and a function index
+	FRigVMTemplate(const FName& InTemplateName, const TArray<FRigVMTemplateArgument>& InArguments, int32 InFunctionIndex);
+
 	int32 Index;
 	FName Notation;
 	TArray<FRigVMTemplateArgument> Arguments;
