RigVM: Fix build break due to mention of RF_PendingKill

diff --git a/Engine/Source/Runtime/RigVM/Private/RigVMCore/RigVMRegistry.cpp b/Engine/Source/Runtime/RigVM/Private/RigVMCore/RigVMRegistry.cpp
--- a/Engine/Source/Runtime/RigVM/Private/RigVMCore/RigVMRegistry.cpp
+++ b/Engine/Source/Runtime/RigVM/Private/RigVMCore/RigVMRegistry.cpp
@@ -165,7 +165,7 @@ const FRigVMTemplate* FRigVMRegistry::GetOrAddTemplateFromArguments(const FName&
 				{
 					static EObjectFlags DisallowedFlags()
 					{
-						return RF_BeginDestroyed | RF_FinishDestroyed | RF_PendingKill;
+						return RF_BeginDestroyed | RF_FinishDestroyed;
 					}
 
 					static EObjectFlags NeededFlags()
