RigVM: Allow template arguments to support both array and single values if needed

diff --git a/Engine/Plugins/Animation/ControlRig/Source/ControlRigEditor/Private/Graph/SControlRigGraphChangePinType.cpp b/Engine/Plugins/Animation/ControlRig/Source/ControlRigEditor/Private/Graph/SControlRigGraphChangePinType.cpp
--- a/Engine/Plugins/Animation/ControlRig/Source/ControlRigEditor/Private/Graph/SControlRigGraphChangePinType.cpp
+++ b/Engine/Plugins/Animation/ControlRig/Source/ControlRigEditor/Private/Graph/SControlRigGraphChangePinType.cpp
@@ -293,11 +293,32 @@ void SControlRigChangePinType::FillPinTypeMenu(FMenuBuilder& MenuBuilder)
 		static FName PropertyIcon(TEXT("Kismet.VariableList.TypeIcon"));
 		const UControlRigGraphSchema* Schema = GetDefault<UControlRigGraphSchema>();
 
+		const bool bHasAllTypes =
+			SortedTypes.Num() >=
+				FRigVMTemplateArgument::GetCompatibleTypes(FRigVMTemplateArgument::ETypeCategory_SingleAnyValue).Num();
+		
 		for(int32 TypeIndex=0; TypeIndex < SortedTypes.Num(); TypeIndex++)
 		{
 			const FRigVMTemplateArgument::FType& Type = SortedTypes[TypeIndex].Key;
 			const FEdGraphPinType PinType = RigVMTypeUtils::PinTypeFromCPPType(Type.CPPType, Type.CPPTypeObject);
 
+			if(bHasAllTypes && Type.CPPTypeObject)
+			{
+				if(Type.CPPTypeObject->IsA<UEnum>() || Type.CPPTypeObject->IsA<UClass>())
+				{
+					continue;
+				}
+
+				if(UScriptStruct* ScriptStruct = Cast<UScriptStruct>(Type.CPPTypeObject))
+				{
+					if(!FRigVMTemplateArgument::GetCompatibleTypes(FRigVMTemplateArgument::ETypeCategory_SingleMathStructValue).Contains(
+						FRigVMTemplateArgument::FType(ScriptStruct->GetStructCPPName(), ScriptStruct)))
+					{
+						continue;
+					}
+				}
+			}
+
 			if(UScriptStruct* ScriptStruct = Cast<UScriptStruct>(Type.CPPTypeObject))
 			{
 				if(ScriptStruct->IsChildOf(FRigVMExecuteContext::StaticStruct()))
diff --git a/Engine/Source/Developer/RigVMDeveloper/Private/RigVMModel/Nodes/RigVMRerouteNode.cpp b/Engine/Source/Developer/RigVMDeveloper/Private/RigVMModel/Nodes/RigVMRerouteNode.cpp
--- a/Engine/Source/Developer/RigVMDeveloper/Private/RigVMModel/Nodes/RigVMRerouteNode.cpp
+++ b/Engine/Source/Developer/RigVMDeveloper/Private/RigVMModel/Nodes/RigVMRerouteNode.cpp
@@ -3,7 +3,6 @@
 #include "RigVMModel/Nodes/RigVMRerouteNode.h"
 
 const FString URigVMRerouteNode::RerouteName = TEXT("Reroute");
-const FString URigVMRerouteNode::RerouteArrayName = TEXT("RerouteArray");
 const FString URigVMRerouteNode::ValueName = TEXT("Value");
 
 URigVMRerouteNode::URigVMRerouteNode()
@@ -66,28 +65,29 @@ FLinearColor URigVMRerouteNode::GetNodeColor() const
 
 FName URigVMRerouteNode::GetNotation() const
 {
-	return GetTemplate()->GetNotation();
+	static constexpr TCHAR Format[] = TEXT("%s(io %s)");
+	static const FName RerouteNotation = *FString::Printf(Format, *RerouteName, *ValueName);
+	return RerouteNotation;
 }
 
 const FRigVMTemplate* URigVMRerouteNode::GetTemplate() const
 {
+	if(const FRigVMTemplate* SuperTemplate = Super::GetTemplate())
+	{
+		return SuperTemplate;
+	}
+	
 	if(CachedTemplate == nullptr)
 	{
-		bool bIsArray = false;
-		if(const URigVMPin* ValuePin = FindPin(ValueName))
-		{
-			bIsArray = ValuePin->IsArray();
-		}
-		CachedTemplate = FindOrAddTemplate(bIsArray);
+		TArray<FRigVMTemplateArgument::FType> Types;
+		Types.Append(FRigVMTemplateArgument::GetCompatibleTypes(FRigVMTemplateArgument::ETypeCategory_SingleAnyValue));
+		Types.Append(FRigVMTemplateArgument::GetCompatibleTypes(FRigVMTemplateArgument::ETypeCategory_ArrayAnyValue));
+		Types.Append(FRigVMTemplateArgument::GetCompatibleTypes(FRigVMTemplateArgument::ETypeCategory_ExecuteContext));
+		
+		TArray<FRigVMTemplateArgument> Arguments;
+		Arguments.Emplace(TEXT("Value"), ERigVMPinDirection::IO, Types);
+		
+		CachedTemplate = FRigVMRegistry::Get().GetOrAddTemplateFromArguments(*RerouteName, Arguments);
 	}
 	return CachedTemplate;
 }
-
-const FRigVMTemplate* URigVMRerouteNode::FindOrAddTemplate(bool bIsArray)
-{
-	TArray<FRigVMTemplateArgument> Arguments;
-	Arguments.Emplace(TEXT("Value"), ERigVMPinDirection::IO, bIsArray ? FRigVMTemplateArgument::FType::Array() : FRigVMTemplateArgument::FType());
-		
-	return FRigVMRegistry::Get().GetOrAddTemplateFromArguments(
-		bIsArray ? *RerouteArrayName : *RerouteName, Arguments, true);
-}
diff --git a/Engine/Source/Developer/RigVMDeveloper/Private/RigVMModel/Nodes/RigVMTemplateNode.cpp b/Engine/Source/Developer/RigVMDeveloper/Private/RigVMModel/Nodes/RigVMTemplateNode.cpp
--- a/Engine/Source/Developer/RigVMDeveloper/Private/RigVMModel/Nodes/RigVMTemplateNode.cpp
+++ b/Engine/Source/Developer/RigVMDeveloper/Private/RigVMModel/Nodes/RigVMTemplateNode.cpp
@@ -13,6 +13,21 @@ URigVMTemplateNode::URigVMTemplateNode()
 {
 }
 
+void URigVMTemplateNode::PostLoad()
+{
+	Super::PostLoad();
+
+	// if there are brackets in the notation remove them
+	const FString OriginalNotation = TemplateNotation.ToString();
+	const FString SanitizedNotation = OriginalNotation.Replace(TEXT("[]"), TEXT(""));
+	if(OriginalNotation != SanitizedNotation)
+	{
+		TemplateNotation = *SanitizedNotation;
+	}
+	
+	InvalidateCache();
+}
+
 UScriptStruct* URigVMTemplateNode::GetScriptStruct() const
 {
 	if(const FRigVMFunction* Function = GetResolvedFunction())
@@ -113,9 +128,15 @@ bool URigVMTemplateNode::SupportsType(const URigVMPin* InPin, const FString& InC
 
 	const URigVMPin* RootPin = InPin->GetRootPin();
 
+	FString CPPType = InCPPType;
+
+	if(InPin->GetParentPin() == RootPin && RootPin->IsArray())
+	{
+		CPPType = RigVMTypeUtils::ArrayTypeFromBaseType(CPPType);
+	}
+
 	// we always support the unknown type
-	if(((InCPPType == WildCardCPPType) && !InPin->IsArray()) ||
-		((InCPPType == WildCardArrayCPPType) && InPin->IsArray()))
+	if((InCPPType == WildCardCPPType) || (InCPPType == WildCardArrayCPPType))
 	{
 		if(const FRigVMTemplate* Template = GetTemplate())
 		{
@@ -126,6 +147,21 @@ bool URigVMTemplateNode::SupportsType(const URigVMPin* InPin, const FString& InC
 				{
 					return false;
 				}
+
+				if(RigVMTypeUtils::IsArrayType(CPPType))
+				{
+					if(Argument->GetArrayType() == FRigVMTemplateArgument::EArrayType_SingleValue)
+					{
+						return false;
+					}
+				}
+				else
+				{
+					if(Argument->GetArrayType() == FRigVMTemplateArgument::EArrayType_ArrayValue)
+					{
+						return false;
+					}
+				}
 				
 				if(OutCPPType)
 				{
@@ -137,13 +173,6 @@ bool URigVMTemplateNode::SupportsType(const URigVMPin* InPin, const FString& InC
 		return false;
 	}
 	
-	FString CPPType = InCPPType;
-
-	if(InPin->GetParentPin() == RootPin && RootPin->IsArray())
-	{
-		CPPType = RigVMTypeUtils::ArrayTypeFromBaseType(CPPType);
-	}
-
 	if (const FRigVMTemplate* Template = GetTemplate())
 	{
 		const FString CacheKey = RootPin->GetName() + TEXT("|") + CPPType;
@@ -258,7 +287,7 @@ const FRigVMTemplate* URigVMTemplateNode::GetTemplate() const
 {
 	if(CachedTemplate == nullptr)
 	{
-		CachedTemplate = FRigVMRegistry::Get().FindTemplate(TemplateNotation);
+		CachedTemplate = FRigVMRegistry::Get().FindTemplate(GetNotation());
 	}
 	return CachedTemplate;
 }
diff --git a/Engine/Source/Developer/RigVMDeveloper/Private/RigVMModel/RigVMController.cpp b/Engine/Source/Developer/RigVMDeveloper/Private/RigVMModel/RigVMController.cpp
--- a/Engine/Source/Developer/RigVMDeveloper/Private/RigVMModel/RigVMController.cpp
+++ b/Engine/Source/Developer/RigVMDeveloper/Private/RigVMModel/RigVMController.cpp
@@ -11724,8 +11724,7 @@ URigVMTemplateNode* URigVMController::AddTemplateNode(const FName& InNotation, c
 	if(Node == nullptr)
 	{
 		const FString TemplateName = Template->GetName().ToString();
-		if(TemplateName== URigVMRerouteNode::RerouteName ||
-			TemplateName== URigVMRerouteNode::RerouteArrayName)
+		if(TemplateName == URigVMRerouteNode::RerouteName)
 		{
 			Node = NewObject<URigVMRerouteNode>(Graph, *Name);
 		}
@@ -15004,10 +15003,24 @@ bool URigVMController::ResolveWildCardPinImpl(URigVMPin* InPinToResolve, const F
 	}
 
 	FString ResolvedCPPType = RigVMTypeUtils::IsArrayType(CPPType) ? RigVMTypeUtils::BaseTypeFromArrayType(CPPType) : CPPType;
+
+	// change the pin's type to be of an array as well
 	if(InPinToResolve->IsArray())
 	{
 		ResolvedCPPType = RigVMTypeUtils::ArrayTypeFromBaseType(ResolvedCPPType);
 	}
+	else if(InPinToResolve->IsRootPin() && RigVMTypeUtils::IsArrayType(InCPPType))
+	{
+		if(URigVMTemplateNode* TemplateNode = Cast<URigVMTemplateNode>(InPinToResolve->GetNode()))
+		{
+			const FString ArrayCPPType = RigVMTypeUtils::ArrayTypeFromBaseType(ResolvedCPPType);
+			FString ResolvedArrayCPPType;
+			if(TemplateNode->SupportsType(InPinToResolve, ArrayCPPType, &ResolvedArrayCPPType))
+			{
+				ResolvedCPPType = ResolvedArrayCPPType;
+			}
+		}
+	}
 
 	if(ChangePinType(InPinToResolve, ResolvedCPPType, CPPTypeObject, bSetupUndoRedo, false, false, false))
 	{
@@ -15143,11 +15156,20 @@ bool URigVMController::ChangePinType(URigVMPin* InPin, const FString& InCPPType,
 	{
 		if(!TemplateNode->SupportsType(InPin, InCPPType, &CPPType))
 		{
+			ReportErrorf(TEXT("ChangePinType: %s doesn't support type '%s'."), *InPin->GetPinPath(), *InCPPType);
 			return false;
 		}
 
-		const FString BaseCPPType = RigVMTypeUtils::IsArrayType(CPPType) ? RigVMTypeUtils::BaseTypeFromArrayType(CPPType) : CPPType;
-		CPPType = InPin->IsArray() ? RigVMTypeUtils::ArrayTypeFromBaseType(BaseCPPType) : BaseCPPType;
+		// change the pin's type to be of an array as well
+		if(InPin->IsRootPin() && RigVMTypeUtils::IsArrayType(CPPType) != InPin->IsArray())
+		{
+			// nothing to do here - leave the type as is 
+		}
+		else
+		{
+			const FString BaseCPPType = RigVMTypeUtils::IsArrayType(CPPType) ? RigVMTypeUtils::BaseTypeFromArrayType(CPPType) : CPPType;
+			CPPType = InPin->IsArray() ? RigVMTypeUtils::ArrayTypeFromBaseType(BaseCPPType) : BaseCPPType;
+		}
 	}
 
 	FRigVMControllerCompileBracketScope CompileScope(this);
diff --git a/Engine/Source/Developer/RigVMDeveloper/Private/RigVMModel/RigVMPin.cpp b/Engine/Source/Developer/RigVMDeveloper/Private/RigVMModel/RigVMPin.cpp
--- a/Engine/Source/Developer/RigVMDeveloper/Private/RigVMModel/RigVMPin.cpp
+++ b/Engine/Source/Developer/RigVMDeveloper/Private/RigVMModel/RigVMPin.cpp
@@ -1625,7 +1625,6 @@ bool URigVMPin::CanLink(URigVMPin* InSourcePin, URigVMPin* InTargetPin, FString*
 
 		if (bCPPTypesDiffer)
 		{
-			if(InSourcePin->IsArray() == InTargetPin->IsArray())
 			{
 				auto TemplateNodeSupportsType = [](URigVMPin* InPin, const FString& InCPPType, FString* OutFailureReason) -> bool
 				{
diff --git a/Engine/Source/Developer/RigVMDeveloper/Public/RigVMModel/Nodes/RigVMRerouteNode.h b/Engine/Source/Developer/RigVMDeveloper/Public/RigVMModel/Nodes/RigVMRerouteNode.h
--- a/Engine/Source/Developer/RigVMDeveloper/Public/RigVMModel/Nodes/RigVMRerouteNode.h
+++ b/Engine/Source/Developer/RigVMDeveloper/Public/RigVMModel/Nodes/RigVMRerouteNode.h
@@ -37,10 +37,7 @@ public:
 
 private:
 
-	static const FRigVMTemplate* FindOrAddTemplate(bool bIsArray);
-
 	static const FString RerouteName;
-	static const FString RerouteArrayName;
 	static const FString ValueName;
 
 	UPROPERTY()
diff --git a/Engine/Source/Developer/RigVMDeveloper/Public/RigVMModel/Nodes/RigVMTemplateNode.h b/Engine/Source/Developer/RigVMDeveloper/Public/RigVMModel/Nodes/RigVMTemplateNode.h
--- a/Engine/Source/Developer/RigVMDeveloper/Public/RigVMModel/Nodes/RigVMTemplateNode.h
+++ b/Engine/Source/Developer/RigVMDeveloper/Public/RigVMModel/Nodes/RigVMTemplateNode.h
@@ -21,6 +21,9 @@ public:
 	// default constructor
 	URigVMTemplateNode();
 
+	// UObject interface
+	virtual void PostLoad() override;
+
 	// URigVMNode interface
 	virtual FString GetNodeTitle() const override;
 	virtual FName GetMethodName() const;
diff --git a/Engine/Source/Runtime/RigVM/Private/RigVMCore/RigVMRegistry.cpp b/Engine/Source/Runtime/RigVM/Private/RigVMCore/RigVMRegistry.cpp
--- a/Engine/Source/Runtime/RigVM/Private/RigVMCore/RigVMRegistry.cpp
+++ b/Engine/Source/Runtime/RigVM/Private/RigVMCore/RigVMRegistry.cpp
@@ -2,6 +2,7 @@
 
 #include "RigVMCore/RigVMRegistry.h"
 #include "RigVMCore/RigVMStruct.h"
+#include "RigVMModule.h"
 #include "UObject/UObjectIterator.h"
 
 FRigVMRegistry FRigVMRegistry::s_RigVMRegistry;
@@ -118,7 +119,7 @@ const TArray<FRigVMTemplate>& FRigVMRegistry::GetTemplates() const
 	return Templates;
 }
 
-const FRigVMTemplate* FRigVMRegistry::GetOrAddTemplateFromArguments(const FName& InName, const TArray<FRigVMTemplateArgument>& InArguments, bool bAllowExecuteContext)
+const FRigVMTemplate* FRigVMRegistry::GetOrAddTemplateFromArguments(const FName& InName, const TArray<FRigVMTemplateArgument>& InArguments)
 {
 	FRigVMTemplate Template(InName, InArguments, INDEX_NONE);
 	if(const FRigVMTemplate* ExistingTemplate = FindTemplate(Template.GetNotation()))
@@ -127,208 +128,35 @@ const FRigVMTemplate* FRigVMRegistry::GetOrAddTemplateFromArguments(const FName&
 	}
 
 	// we only support to ask for templates here which provide singleton types
+	int32 NumPermutations = 1;
 	for(const FRigVMTemplateArgument& Argument : InArguments)
 	{
-		if(!Argument.bSingleton || Argument.Types.Num() > 1)
+		if(!Argument.bSingleton && NumPermutations > 1)
 		{
-			return nullptr;
+			if(Argument.Types.Num() != NumPermutations)
+			{
+				UE_LOG(LogRigVM, Error, TEXT("Failed to add template '%s' since the arguments' types counts don't match."), *InName.ToString());
+				return nullptr;
+			}
 		}
+		NumPermutations = FMath::Max(NumPermutations, Argument.Types.Num()); 
 	}
 
-	int32 NumPermutations = 1;
-
 	// if any of the arguments are wildcards we'll need to update the types
 	for(FRigVMTemplateArgument& Argument : Template.Arguments)
 	{
-		if(Argument.Types[0].IsWildCard())
+		if(Argument.Types.Num() == 1 && Argument.Types[0].IsWildCard())
 		{
-			static TArray<FRigVMTemplateArgument::FType> AllTypes, AllArrayTypes, ExecuteContextTypes;
-			if(AllTypes.IsEmpty())
-			{
-				AllTypes.Add(FRigVMTemplateArgument::FType(RigVMTypeUtils::BoolType));
-				AllTypes.Add(FRigVMTemplateArgument::FType(RigVMTypeUtils::Int32Type));
-				AllTypes.Add(FRigVMTemplateArgument::FType(RigVMTypeUtils::UInt8Type));
-				AllTypes.Add(FRigVMTemplateArgument::FType(RigVMTypeUtils::FloatType));
-				AllTypes.Add(FRigVMTemplateArgument::FType(RigVMTypeUtils::DoubleType));
-				AllTypes.Add(FRigVMTemplateArgument::FType(RigVMTypeUtils::FNameType));
-				AllTypes.Add(FRigVMTemplateArgument::FType(RigVMTypeUtils::FStringType));
-
-				AllArrayTypes.Add(FRigVMTemplateArgument::FType(RigVMTypeUtils::BoolArrayType));
-				AllArrayTypes.Add(FRigVMTemplateArgument::FType(RigVMTypeUtils::Int32ArrayType));
-				AllArrayTypes.Add(FRigVMTemplateArgument::FType(RigVMTypeUtils::UInt8ArrayType));
-				AllArrayTypes.Add(FRigVMTemplateArgument::FType(RigVMTypeUtils::FloatArrayType));
-				AllArrayTypes.Add(FRigVMTemplateArgument::FType(RigVMTypeUtils::DoubleArrayType));
-				AllArrayTypes.Add(FRigVMTemplateArgument::FType(RigVMTypeUtils::FNameArrayType));
-				AllArrayTypes.Add(FRigVMTemplateArgument::FType(RigVMTypeUtils::FStringArrayType));
-
-				struct FTypeTraverser
-				{
-					EObjectFlags DisallowedFlags()
-					{
-						return RF_BeginDestroyed | RF_FinishDestroyed;
-					}
-
-					EObjectFlags NeededFlags()
-					{
-						return RF_Public;
-					}
-					
-					bool IsAllowedType(const FProperty* InProperty)
-					{
-						if(!InProperty->HasAnyPropertyFlags(
-							CPF_BlueprintVisible |
-							CPF_BlueprintReadOnly |
-							CPF_Edit))
-						{
-							return false;
-						}
-
-						if(InProperty->IsA<FBoolProperty>() ||
-							InProperty->IsA<FUInt32Property>() ||
-							InProperty->IsA<FInt8Property>() ||
-							InProperty->IsA<FInt16Property>() ||
-							InProperty->IsA<FIntProperty>() ||
-							InProperty->IsA<FInt64Property>() ||
-							InProperty->IsA<FFloatProperty>() ||
-							InProperty->IsA<FDoubleProperty>() ||
-							InProperty->IsA<FNumericProperty>() ||
-							InProperty->IsA<FNameProperty>() ||
-							InProperty->IsA<FStrProperty>())
-						{
-							return true;
-						}
-
-						if(const FArrayProperty* ArrayProperty  = CastField<FArrayProperty>(InProperty))
-						{
-							return IsAllowedType(ArrayProperty->Inner);
-						}
-						if(const FStructProperty* StructProperty = CastField<FStructProperty>(InProperty))
-						{
-							return IsAllowedType(StructProperty->Struct);
-						}
-						if(const FObjectProperty* ObjectProperty = CastField<FObjectProperty>(InProperty))
-						{
-							return IsAllowedType(ObjectProperty->PropertyClass);
-						}
-						if(const FEnumProperty* EnumProperty = CastField<FEnumProperty>(InProperty))
-						{
-							return IsAllowedType(EnumProperty->GetEnum());
-						}
-						if(const FByteProperty* ByteProperty = CastField<FByteProperty>(InProperty))
-						{
-							if(const UEnum* Enum = ByteProperty->Enum)
-							{
-								return IsAllowedType(Enum);
-							}
-							return true;
-						}
-						return false;
-					}
-
-					bool IsAllowedType(const UEnum* InEnum)
-					{
-						return !InEnum->HasAnyFlags(DisallowedFlags()) && InEnum->HasAllFlags(NeededFlags());
-					}
-
-					bool IsAllowedType(const UStruct* InStruct)
-					{
-						if(InStruct->HasAnyFlags(DisallowedFlags()) || !InStruct->HasAllFlags(NeededFlags()))
-						{
-							return false;
-						}
-						if(InStruct->IsChildOf(FRigVMStruct::StaticStruct()))
-						{
-							return false;
-						}
-						for (TFieldIterator<FProperty> It(InStruct); It; ++It)
-						{
-							if(!IsAllowedType(*It))
-							{
-								return false;
-							}
-						}
-						return true;
-					}
-
-					bool IsAllowedType(const UClass* InClass)
-					{
-						if(InClass->HasAnyClassFlags(CLASS_Hidden | CLASS_Abstract))
-						{
-							return false;
-						}
-
-						// note: currently we don't allow UObjects
-						return false;
-						//return IsAllowedType(Cast<UStruct>(InClass));
-					}
-				};
-
-				FTypeTraverser Traverser;
-
-				// add all structs
-				for (TObjectIterator<UScriptStruct> ScriptIt; ScriptIt; ++ScriptIt)
-				{
-					UScriptStruct* ScriptStruct = *ScriptIt;
-					if(!Traverser.IsAllowedType(ScriptStruct))
-					{
-						continue;
-					}
-					const FString CPPType = ScriptStruct->GetStructCPPName();
-
-					if(ScriptStruct->IsChildOf(FRigVMExecuteContext::StaticStruct()))
-					{
-						ExecuteContextTypes.Add(FRigVMTemplateArgument::FType(CPPType, ScriptStruct));
-					}
-					else
-					{
-						AllTypes.Add(FRigVMTemplateArgument::FType(CPPType, ScriptStruct));
-						AllArrayTypes.Add(FRigVMTemplateArgument::FType(RigVMTypeUtils::ArrayTypeFromBaseType(CPPType), ScriptStruct));
-					}
-				}
-
-				// add all enums
-				for (TObjectIterator<UEnum> EnumIt; EnumIt; ++EnumIt)
-				{
-					UEnum* Enum = (*EnumIt);
-					if(!Traverser.IsAllowedType(Enum))
-					{
-						continue;
-					}
-					const FString CPPType = Enum->CppType.IsEmpty() ? Enum->GetName() : Enum->CppType;
-					AllTypes.Add(FRigVMTemplateArgument::FType(CPPType, Enum));
-					AllArrayTypes.Add(FRigVMTemplateArgument::FType(RigVMTypeUtils::ArrayTypeFromBaseType(CPPType), Enum));
-				}
-
-				// add all classes
-				for (TObjectIterator<UClass> ClassIt; ClassIt; ++ClassIt)
-				{
-					UClass* Class = *ClassIt;
-					if(!Traverser.IsAllowedType(Class))
-					{
-						continue;
-					}
-
-					const FString CPPType = Class->GetPrefixCPP() + Class->GetName();
-					AllTypes.Add(FRigVMTemplateArgument::FType(CPPType, Class));
-					AllArrayTypes.Add(FRigVMTemplateArgument::FType(RigVMTypeUtils::ArrayTypeFromBaseType(CPPType), Class));
-				}
-			}
-
 			if(Argument.Types[0].IsArray())
 			{
-				Argument.Types = AllArrayTypes;
+				Argument.Types = FRigVMTemplateArgument::GetCompatibleTypes(FRigVMTemplateArgument::ETypeCategory_ArrayAnyValue);
 			}
 			else
 			{
-				Argument.Types = AllTypes;
-				if(bAllowExecuteContext)
-				{
-					Argument.Types.Append(ExecuteContextTypes);
-				}
+				Argument.Types = FRigVMTemplateArgument::GetCompatibleTypes(FRigVMTemplateArgument::ETypeCategory_SingleAnyValue);
 			}
-			
 			Argument.bSingleton = false;
-			NumPermutations = Argument.Types.Num(); 
+			NumPermutations = FMath::Max(NumPermutations, Argument.Types.Num()); 
 		}
 	}
 
diff --git a/Engine/Source/Runtime/RigVM/Private/RigVMCore/RigVMTemplate.cpp b/Engine/Source/Runtime/RigVM/Private/RigVMCore/RigVMTemplate.cpp
--- a/Engine/Source/Runtime/RigVM/Private/RigVMCore/RigVMTemplate.cpp
+++ b/Engine/Source/Runtime/RigVM/Private/RigVMCore/RigVMTemplate.cpp
@@ -61,6 +61,24 @@ FRigVMTemplateArgument::FRigVMTemplateArgument(const FName& InName, ERigVMPinDir
 {
 }
 
+FRigVMTemplateArgument::FRigVMTemplateArgument(const FName& InName, ERigVMPinDirection InDirection, const TArray<FType>& InTypes)
+: Index(INDEX_NONE)
+, Name(InName)
+, Direction(InDirection)
+, bSingleton(true)
+, Types(InTypes)
+{
+	check(Types.Num() > 0);
+	for(int32 TypeIndex=1;TypeIndex<Types.Num();TypeIndex++)
+	{
+		if(Types[TypeIndex] != Types[0])
+		{
+			bSingleton = false;
+			break;
+		}
+	}
+}
+
 bool FRigVMTemplateArgument::SupportsType(const FString& InCPPType, const TArray<int32>& InPermutationIndices, FType* OutType) const
 {
 	bool bFoundMatch = false;
@@ -124,13 +142,275 @@ bool FRigVMTemplateArgument::IsSingleton(const TArray<int32>& InPermutationIndic
 	return true;
 }
 
-bool FRigVMTemplateArgument::IsArray() const
+FRigVMTemplateArgument::EArrayType FRigVMTemplateArgument::GetArrayType() const
 {
 	if(Types.Num() > 0)
 	{
-		return Types[0].IsArray();
+		const EArrayType ArrayType = Types[0].IsArray() ? EArrayType::EArrayType_ArrayValue : EArrayType::EArrayType_SingleValue;
+		
+		if(IsSingleton())
+		{
+			return ArrayType;
+		}
+
+		for(int32 PermutationIndex=1;PermutationIndex<Types.Num();PermutationIndex++)
+		{
+			const EArrayType OtherArrayType = Types[PermutationIndex].IsArray() ? EArrayType::EArrayType_ArrayValue : EArrayType::EArrayType_SingleValue;
+			if(OtherArrayType != ArrayType)
+			{
+				return EArrayType::EArrayType_Mixed;
+			}
+		}
+
+		return ArrayType;
 	}
-	return false;
+
+	return EArrayType_Invalid;
+}
+
+const TArray<FRigVMTemplateArgument::FType> FRigVMTemplateArgument::GetCompatibleTypes(ETypeCategory InCategory)
+{
+	check(InCategory != ETypeCategory_Invalid);
+	
+	static TMap<ETypeCategory, TArray<FType>> CompatibleTypes;
+	if(CompatibleTypes.IsEmpty())
+	{
+		CompatibleTypes.Add(ETypeCategory_SingleAnyValue, TArray<FType>());
+		CompatibleTypes.Add(ETypeCategory_ArrayAnyValue, TArray<FType>());
+		CompatibleTypes.Add(ETypeCategory_SingleSimpleValue, TArray<FType>());
+		CompatibleTypes.Add(ETypeCategory_ArraySimpleValue, TArray<FType>());
+		CompatibleTypes.Add(ETypeCategory_SingleMathStructValue, TArray<FType>());
+		CompatibleTypes.Add(ETypeCategory_ArrayMathStructValue, TArray<FType>());
+		CompatibleTypes.Add(ETypeCategory_SingleScriptStructValue, TArray<FType>());
+		CompatibleTypes.Add(ETypeCategory_ArrayScriptStructValue, TArray<FType>());
+		CompatibleTypes.Add(ETypeCategory_SingleEnumValue, TArray<FType>());
+		CompatibleTypes.Add(ETypeCategory_ArrayEnumValue, TArray<FType>());
+		CompatibleTypes.Add(ETypeCategory_SingleObjectValue, TArray<FType>());
+		CompatibleTypes.Add(ETypeCategory_ArrayObjectValue, TArray<FType>());
+		CompatibleTypes.Add(ETypeCategory_ExecuteContext, TArray<FType>());
+
+		TArray<FType>& SingleAnyValueTypes = CompatibleTypes.FindChecked(ETypeCategory_SingleAnyValue);
+		TArray<FType>& ArrayAnyValueTypes = CompatibleTypes.FindChecked(ETypeCategory_ArrayAnyValue);
+		TArray<FType>& SingleSimpleValueTypes = CompatibleTypes.FindChecked(ETypeCategory_SingleSimpleValue);
+		TArray<FType>& ArraySimpleValueTypes = CompatibleTypes.FindChecked(ETypeCategory_ArraySimpleValue);
+		TArray<FType>& SingleMathStructValueTypes = CompatibleTypes.FindChecked(ETypeCategory_SingleMathStructValue);
+		TArray<FType>& ArrayMathStructValueTypes = CompatibleTypes.FindChecked(ETypeCategory_ArrayMathStructValue);
+		TArray<FType>& SingleScriptStructValueTypes = CompatibleTypes.FindChecked(ETypeCategory_SingleScriptStructValue);
+		TArray<FType>& ArrayScriptStructValueTypes = CompatibleTypes.FindChecked(ETypeCategory_ArrayScriptStructValue);
+		TArray<FType>& SingleEnumValueTypes = CompatibleTypes.FindChecked(ETypeCategory_SingleEnumValue);
+		TArray<FType>& ArrayEnumValueTypes = CompatibleTypes.FindChecked(ETypeCategory_ArrayEnumValue);
+		TArray<FType>& SingleObjectValueTypes = CompatibleTypes.FindChecked(ETypeCategory_SingleObjectValue);
+		TArray<FType>& ArrayObjectValueTypes = CompatibleTypes.FindChecked(ETypeCategory_ArrayObjectValue);
+		TArray<FType>& ExecuteContextTypes = CompatibleTypes.FindChecked(ETypeCategory_ExecuteContext);
+
+		SingleSimpleValueTypes.Add(FType(RigVMTypeUtils::BoolType));
+		SingleSimpleValueTypes.Add(FType(RigVMTypeUtils::Int32Type));
+		SingleSimpleValueTypes.Add(FType(RigVMTypeUtils::UInt8Type));
+		SingleSimpleValueTypes.Add(FType(RigVMTypeUtils::FloatType));
+		SingleSimpleValueTypes.Add(FType(RigVMTypeUtils::DoubleType));
+		SingleSimpleValueTypes.Add(FType(RigVMTypeUtils::FNameType));
+		SingleSimpleValueTypes.Add(FType(RigVMTypeUtils::FStringType));
+
+		for(const FType& Type : SingleSimpleValueTypes)
+		{
+			ArraySimpleValueTypes.Add(FType(RigVMTypeUtils::ArrayTypeFromBaseType(Type.CPPType)));
+		}
+
+		SingleAnyValueTypes = SingleSimpleValueTypes;
+		ArrayAnyValueTypes = ArraySimpleValueTypes;
+
+		static const TArray<UScriptStruct*> MathTypes = { 
+			TBaseStructure<FRotator>::Get(),
+			TBaseStructure<FQuat>::Get(),
+			TBaseStructure<FTransform>::Get(),
+			TBaseStructure<FLinearColor>::Get(),
+			TBaseStructure<FColor>::Get(),
+			TBaseStructure<FPlane>::Get(),
+			TBaseStructure<FVector>::Get(),
+			TBaseStructure<FVector2D>::Get(),
+			TBaseStructure<FVector4>::Get(),
+			TBaseStructure<FBox2D>::Get()
+		};
+
+		for(UScriptStruct* MathType : MathTypes)
+		{
+			SingleMathStructValueTypes.Add(FType(MathType->GetStructCPPName(), MathType));
+			ArrayMathStructValueTypes.Add(FType(RigVMTypeUtils::ArrayTypeFromBaseType(MathType->GetStructCPPName()), MathType));
+		}
+
+		struct FTypeTraverser
+		{
+			static EObjectFlags DisallowedFlags()
+			{
+				return RF_BeginDestroyed | RF_FinishDestroyed;
+			}
+
+			static EObjectFlags NeededFlags()
+			{
+				return RF_Public;
+			}
+			
+			static bool IsAllowedType(const FProperty* InProperty)
+			{
+				if(!InProperty->HasAnyPropertyFlags(
+					CPF_BlueprintVisible |
+					CPF_BlueprintReadOnly |
+					CPF_Edit))
+				{
+					return false;
+				}
+
+				if(InProperty->IsA<FBoolProperty>() ||
+					InProperty->IsA<FUInt32Property>() ||
+					InProperty->IsA<FInt8Property>() ||
+					InProperty->IsA<FInt16Property>() ||
+					InProperty->IsA<FIntProperty>() ||
+					InProperty->IsA<FInt64Property>() ||
+					InProperty->IsA<FFloatProperty>() ||
+					InProperty->IsA<FDoubleProperty>() ||
+					InProperty->IsA<FNumericProperty>() ||
+					InProperty->IsA<FNameProperty>() ||
+					InProperty->IsA<FStrProperty>())
+				{
+					return true;
+				}
+
+				if(const FArrayProperty* ArrayProperty  = CastField<FArrayProperty>(InProperty))
+				{
+					return IsAllowedType(ArrayProperty->Inner);
+				}
+				if(const FStructProperty* StructProperty = CastField<FStructProperty>(InProperty))
+				{
+					return IsAllowedType(StructProperty->Struct);
+				}
+				if(const FObjectProperty* ObjectProperty = CastField<FObjectProperty>(InProperty))
+				{
+					return IsAllowedType(ObjectProperty->PropertyClass);
+				}
+				if(const FEnumProperty* EnumProperty = CastField<FEnumProperty>(InProperty))
+				{
+					return IsAllowedType(EnumProperty->GetEnum());
+				}
+				if(const FByteProperty* ByteProperty = CastField<FByteProperty>(InProperty))
+				{
+					if(const UEnum* Enum = ByteProperty->Enum)
+					{
+						return IsAllowedType(Enum);
+					}
+					return true;
+				}
+				return false;
+			}
+
+			static bool IsAllowedType(const UEnum* InEnum)
+			{
+				return !InEnum->HasAnyFlags(DisallowedFlags()) && InEnum->HasAllFlags(NeededFlags());
+			}
+
+			static bool IsAllowedType(const UStruct* InStruct)
+			{
+				if(InStruct->HasAnyFlags(DisallowedFlags()) || !InStruct->HasAllFlags(NeededFlags()))
+				{
+					return false;
+				}
+				if(InStruct->IsChildOf(FRigVMStruct::StaticStruct()))
+				{
+					return false;
+				}
+				if(InStruct->IsChildOf(FRigVMUnknownType::StaticStruct()))
+				{
+					return false;
+				}
+				for (TFieldIterator<FProperty> It(InStruct); It; ++It)
+				{
+					if(!IsAllowedType(*It))
+					{
+						return false;
+					}
+				}
+				return true;
+			}
+
+			static bool IsAllowedType(const UClass* InClass)
+			{
+				if(InClass->HasAnyClassFlags(CLASS_Hidden | CLASS_Abstract))
+				{
+					return false;
+				}
+
+				// note: currently we don't allow UObjects
+				return false;
+				//return IsAllowedType(Cast<UStruct>(InClass));
+			}
+		};
+
+		// add all structs
+		for (TObjectIterator<UScriptStruct> ScriptIt; ScriptIt; ++ScriptIt)
+		{
+			UScriptStruct* ScriptStruct = *ScriptIt;
+			if(!FTypeTraverser::IsAllowedType(ScriptStruct))
+			{
+				continue;
+			}
+
+			const FString CPPType = ScriptStruct->GetStructCPPName();
+
+			if(ScriptStruct->IsChildOf(FRigVMExecuteContext::StaticStruct()))
+			{
+				ExecuteContextTypes.Add(FType(CPPType, ScriptStruct));
+			}
+			else
+			{
+				SingleAnyValueTypes.Add(FType(CPPType, ScriptStruct));
+				ArrayAnyValueTypes.Add(FType(RigVMTypeUtils::ArrayTypeFromBaseType(CPPType), ScriptStruct));
+				SingleScriptStructValueTypes.Add(FType(CPPType, ScriptStruct));
+				ArrayScriptStructValueTypes.Add(FType(RigVMTypeUtils::ArrayTypeFromBaseType(CPPType), ScriptStruct));
+			}
+		}
+
+		// add all enums
+		for (TObjectIterator<UEnum> EnumIt; EnumIt; ++EnumIt)
+		{
+			UEnum* Enum = (*EnumIt);
+			if(!FTypeTraverser::IsAllowedType(Enum))
+			{
+				continue;
+			}
+			const FString CPPType = Enum->CppType.IsEmpty() ? Enum->GetName() : Enum->CppType;
+			SingleAnyValueTypes.Add(FType(CPPType, Enum));
+			ArrayAnyValueTypes.Add(FType(RigVMTypeUtils::ArrayTypeFromBaseType(CPPType), Enum));
+			SingleEnumValueTypes.Add(FType(CPPType, Enum));
+			ArrayEnumValueTypes.Add(FType(RigVMTypeUtils::ArrayTypeFromBaseType(CPPType), Enum));
+		}
+
+		// add all classes
+		for (TObjectIterator<UClass> ClassIt; ClassIt; ++ClassIt)
+		{
+			UClass* Class = *ClassIt;
+
+			// for now this always returns false
+			if(!FTypeTraverser::IsAllowedType(Class))
+			{
+				continue;
+			}
+
+			const FString CPPType = Class->GetPrefixCPP() + Class->GetName();
+			SingleAnyValueTypes.Add(FType(CPPType, Class));
+			ArrayAnyValueTypes.Add(FType(RigVMTypeUtils::ArrayTypeFromBaseType(CPPType), Class));
+			SingleObjectValueTypes.Add(FType(CPPType, Class));
+			ArrayObjectValueTypes.Add(FType(RigVMTypeUtils::ArrayTypeFromBaseType(CPPType), Class));
+		}
+
+		// check that matching categories have the same size
+		check(CompatibleTypes.FindChecked(ETypeCategory_SingleAnyValue).Num() == CompatibleTypes.FindChecked(ETypeCategory_ArrayAnyValue).Num());
+		check(CompatibleTypes.FindChecked(ETypeCategory_SingleSimpleValue).Num() == CompatibleTypes.FindChecked(ETypeCategory_ArraySimpleValue).Num());
+		check(CompatibleTypes.FindChecked(ETypeCategory_SingleMathStructValue).Num() == CompatibleTypes.FindChecked(ETypeCategory_ArrayMathStructValue).Num());
+		check(CompatibleTypes.FindChecked(ETypeCategory_SingleScriptStructValue).Num() == CompatibleTypes.FindChecked(ETypeCategory_ArrayScriptStructValue).Num());
+		check(CompatibleTypes.FindChecked(ETypeCategory_SingleEnumValue).Num() == CompatibleTypes.FindChecked(ETypeCategory_ArrayEnumValue).Num());
+		check(CompatibleTypes.FindChecked(ETypeCategory_SingleObjectValue).Num() == CompatibleTypes.FindChecked(ETypeCategory_ArrayObjectValue).Num());
+	}
+
+	return CompatibleTypes.FindChecked(InCategory);
 }
 
 const TArray<FRigVMTemplateArgument::FType>& FRigVMTemplateArgument::GetTypes() const
@@ -285,19 +565,11 @@ const FString& FRigVMTemplate::GetArgumentNotationPrefix(const FRigVMTemplateArg
 	return EmptyPrefix;
 }
 
-const FString& FRigVMTemplate::GetArgumentNotationSuffix(const FRigVMTemplateArgument& InArgument)
-{
-	static const FString EmptySuffix = FString();
-	static const FString ArraySuffix = TEXT("[]");
-	return InArgument.IsArray() ? ArraySuffix : EmptySuffix;
-}
-
 FString FRigVMTemplate::GetArgumentNotation(const FRigVMTemplateArgument& InArgument)
 {
-	return FString::Printf(TEXT("%s%s%s"),
+	return FString::Printf(TEXT("%s%s"),
 		*GetArgumentNotationPrefix(InArgument),
-		*InArgument.GetName().ToString(),
-		*GetArgumentNotationSuffix(InArgument));
+		*InArgument.GetName().ToString());
 }
 
 bool FRigVMTemplate::IsValid() const
@@ -730,7 +1002,7 @@ bool FRigVMTemplate::Resolve(FTypeMap& InOutTypes, TArray<int32>& OutPermutation
 			}
 		}
 		
-		if(Argument.IsArray())
+		if(Argument.GetArrayType() == FRigVMTemplateArgument::EArrayType_ArrayValue)
 		{
 			InOutTypes.Add(Argument.Name, FRigVMTemplateArgument::FType::Array());
 		}
diff --git a/Engine/Source/Runtime/RigVM/Public/RigVMCore/RigVMRegistry.h b/Engine/Source/Runtime/RigVM/Public/RigVMCore/RigVMRegistry.h
--- a/Engine/Source/Runtime/RigVM/Public/RigVMCore/RigVMRegistry.h
+++ b/Engine/Source/Runtime/RigVM/Public/RigVMCore/RigVMRegistry.h
@@ -45,7 +45,7 @@ public:
 	const TArray<FRigVMTemplate>& GetTemplates() const;
 
 	// Defines and retrieves a template given its arguments
-	const FRigVMTemplate* GetOrAddTemplateFromArguments(const FName& InName, const TArray<FRigVMTemplateArgument>& InArguments, bool bIncludeExecuteContext = false);
+	const FRigVMTemplate* GetOrAddTemplateFromArguments(const FName& InName, const TArray<FRigVMTemplateArgument>& InArguments);
 
 private:
 
diff --git a/Engine/Source/Runtime/RigVM/Public/RigVMCore/RigVMTemplate.h b/Engine/Source/Runtime/RigVM/Public/RigVMCore/RigVMTemplate.h
--- a/Engine/Source/Runtime/RigVM/Public/RigVMCore/RigVMTemplate.h
+++ b/Engine/Source/Runtime/RigVM/Public/RigVMCore/RigVMTemplate.h
@@ -111,8 +111,35 @@ struct RIGVM_API FRigVMTemplateArgument
 			return CPPType;
 		}
 	};
+	
+	enum EArrayType
+	{
+		EArrayType_SingleValue,
+		EArrayType_ArrayValue,
+		EArrayType_Mixed,
+		EArrayType_Invalid
+	};
+
+	enum ETypeCategory
+	{
+		ETypeCategory_SingleAnyValue,
+		ETypeCategory_ArrayAnyValue,
+		ETypeCategory_SingleSimpleValue,
+		ETypeCategory_ArraySimpleValue,
+		ETypeCategory_SingleMathStructValue,
+		ETypeCategory_ArrayMathStructValue,
+		ETypeCategory_SingleScriptStructValue,
+		ETypeCategory_ArrayScriptStructValue,
+		ETypeCategory_SingleEnumValue,
+		ETypeCategory_ArrayEnumValue,
+		ETypeCategory_SingleObjectValue,
+		ETypeCategory_ArrayObjectValue,
+		ETypeCategory_ExecuteContext,
+		ETypeCategory_Invalid
+	};
 
 	FRigVMTemplateArgument(const FName& InName, ERigVMPinDirection InDirection, const FType& InType);
+	FRigVMTemplateArgument(const FName& InName, ERigVMPinDirection InDirection, const TArray<FType>& InTypes);
 
 	// returns the name of the argument
 	const FName& GetName() const { return Name; }
@@ -136,7 +163,10 @@ struct RIGVM_API FRigVMTemplateArgument
 	bool IsSingleton(const TArray<int32>& InPermutationIndices = TArray<int32>()) const;
 
 	// returns true if the argument uses an array container
-	bool IsArray() const;
+	EArrayType GetArrayType() const;
+	
+	// Returns all compatible types given a category
+	static const TArray<FType> GetCompatibleTypes(ETypeCategory InCategory);
 
 protected:
 
@@ -224,9 +254,6 @@ public:
 	// returns the prefix for an argument in the notation
 	static const FString& GetArgumentNotationPrefix(const FRigVMTemplateArgument& InArgument);
 
-	// returns the suffix for an argument in the notation
-	static const FString& GetArgumentNotationSuffix(const FRigVMTemplateArgument& InArgument);
-
 	// returns the notation of an argument
 	static FString GetArgumentNotation(const FRigVMTemplateArgument& InArgument);
 
