Control Rig: Fix crash in compiler in shipping

diff --git a/Engine/Source/Runtime/RigVM/Private/RigVMCore/RigVMMemoryStorage.cpp b/Engine/Source/Runtime/RigVM/Private/RigVMCore/RigVMMemoryStorage.cpp
--- a/Engine/Source/Runtime/RigVM/Private/RigVMCore/RigVMMemoryStorage.cpp
+++ b/Engine/Source/Runtime/RigVM/Private/RigVMCore/RigVMMemoryStorage.cpp
@@ -162,12 +162,12 @@ FRigVMPropertyDescription::FRigVMPropertyDescription(const FName& InName, const
 		if(TailCPPType.RemoveFromStart(ArrayPrefix))
 		{
 			Containers.Add(EPinContainerType::Array);
-			check(TailCPPType.RemoveFromEnd(ContainerSuffix));
+			verify(TailCPPType.RemoveFromEnd(ContainerSuffix));
 		}
 		else if(TailCPPType.RemoveFromStart(MapPrefix))
 		{
 			Containers.Add(EPinContainerType::Map);
-			check(TailCPPType.RemoveFromEnd(ContainerSuffix));
+			verify(TailCPPType.RemoveFromEnd(ContainerSuffix));
 		}
 		else
 		{
@@ -228,14 +228,14 @@ FString FRigVMPropertyDescription::GetTailCPPType() const
 		{
 			case EPinContainerType::Array:
 			{
-				check(TailCPPType.RemoveFromStart(ArrayPrefix))
-				check(TailCPPType.RemoveFromEnd(ContainerSuffix));
+				verify(TailCPPType.RemoveFromStart(ArrayPrefix));
+				verify(TailCPPType.RemoveFromEnd(ContainerSuffix));
 				break;
 			}		
 			case EPinContainerType::Map:
 			{
-				check(TailCPPType.RemoveFromStart(MapPrefix))
-				check(TailCPPType.RemoveFromEnd(ContainerSuffix));
+				verify(TailCPPType.RemoveFromStart(MapPrefix));
+				verify(TailCPPType.RemoveFromEnd(ContainerSuffix));
 				break;
 			}		
 			case EPinContainerType::None:
