RigVM Template: Make sure to check flags on array container properties only

diff --git a/Engine/Source/Runtime/RigVM/Private/RigVMCore/RigVMTemplate.cpp b/Engine/Source/Runtime/RigVM/Private/RigVMCore/RigVMTemplate.cpp
--- a/Engine/Source/Runtime/RigVM/Private/RigVMCore/RigVMTemplate.cpp
+++ b/Engine/Source/Runtime/RigVM/Private/RigVMCore/RigVMTemplate.cpp
@@ -249,14 +249,17 @@ const TArray<FRigVMTemplateArgument::FType> FRigVMTemplateArgument::GetCompatibl
 				return RF_Public;
 			}
 			
-			static bool IsAllowedType(const FProperty* InProperty)
+			static bool IsAllowedType(const FProperty* InProperty, bool bCheckFlags = true)
 			{
-				if(!InProperty->HasAnyPropertyFlags(
-					CPF_BlueprintVisible |
-					CPF_BlueprintReadOnly |
-					CPF_Edit))
+				if(bCheckFlags)
 				{
-					return false;
+					if(!InProperty->HasAnyPropertyFlags(
+						CPF_BlueprintVisible |
+						CPF_BlueprintReadOnly |
+						CPF_Edit))
+					{
+						return false;
+					}
 				}
 
 				if(InProperty->IsA<FBoolProperty>() ||
@@ -276,7 +279,7 @@ const TArray<FRigVMTemplateArgument::FType> FRigVMTemplateArgument::GetCompatibl
 
 				if(const FArrayProperty* ArrayProperty  = CastField<FArrayProperty>(InProperty))
 				{
-					return IsAllowedType(ArrayProperty->Inner);
+					return IsAllowedType(ArrayProperty->Inner, false);
 				}
 				if(const FStructProperty* StructProperty = CastField<FStructProperty>(InProperty))
 				{
