Control Rig: Avoid copying TFunction in a cpp struct copy context.

#rnx

diff --git a/Engine/Plugins/Animation/ControlRig/Source/ControlRig/Private/ControlRig.cpp b/Engine/Plugins/Animation/ControlRig/Source/ControlRig/Private/ControlRig.cpp
--- a/Engine/Plugins/Animation/ControlRig/Source/ControlRig/Private/ControlRig.cpp
+++ b/Engine/Plugins/Animation/ControlRig/Source/ControlRig/Private/ControlRig.cpp
@@ -1401,7 +1401,7 @@ void UControlRig::UpdateVMSettings()
 #if WITH_EDITOR
 
 		// setup array handling and error reporting on the VM
-		VMRuntimeSettings.LogFunction = [this](EMessageSeverity::Type InSeverity, const FRigVMExecuteContext* InContext, const FString& Message)
+		VMRuntimeSettings.SetLogFunction([this](EMessageSeverity::Type InSeverity, const FRigVMExecuteContext* InContext, const FString& Message)
 		{
 			check(InContext);
 
@@ -1413,8 +1413,7 @@ void UControlRig::UpdateVMSettings()
 			{
 				LogOnce(InSeverity, InContext->InstructionIndex, Message);
 			}
-		};
-		
+		});
 #endif
 		
 		VM->SetRuntimeSettings(VMRuntimeSettings);
diff --git a/Engine/Source/Runtime/RigVM/Public/RigVMCore/RigVMExecuteContext.h b/Engine/Source/Runtime/RigVM/Public/RigVMCore/RigVMExecuteContext.h
--- a/Engine/Source/Runtime/RigVM/Public/RigVMCore/RigVMExecuteContext.h
+++ b/Engine/Source/Runtime/RigVM/Public/RigVMCore/RigVMExecuteContext.h
@@ -177,7 +177,13 @@ struct RIGVM_API FRigVMRuntimeSettings
 	/*
 	 * The function to use for logging anything from the VM to the host
 	 */
-	TFunction<void(EMessageSeverity::Type,const FRigVMExecuteContext*,const FString&)> LogFunction = nullptr;
+	using LogFunctionType = TFunction<void(EMessageSeverity::Type,const FRigVMExecuteContext*,const FString&)>;
+	TSharedPtr<LogFunctionType> LogFunction = nullptr;
+
+	void SetLogFunction(LogFunctionType InLogFunction)
+	{
+		LogFunction = MakeShared<LogFunctionType>(InLogFunction);
+	}
 
 	/*
 	 * Validate the settings
@@ -208,9 +214,9 @@ struct RIGVM_API FRigVMExecuteContext
 
 	FORCEINLINE void Log(EMessageSeverity::Type InSeverity, const FString& InMessage) const
 	{
-		if(RuntimeSettings.LogFunction)
+		if(RuntimeSettings.LogFunction.IsValid())
 		{
-			RuntimeSettings.LogFunction(InSeverity, this, InMessage);
+			(*RuntimeSettings.LogFunction)(InSeverity, this, InMessage);
 		}
 		else
 		{
