ControlRig: Reparent rig vm mem class to the asset obj instead of the package because asset related operation assumes that each package only contains one asset object and actions like rename only acts on that one object and skips other objects in the same package

diff --git a/Engine/Plugins/Animation/ControlRig/Source/ControlRigDeveloper/Private/ControlRigBlueprint.cpp b/Engine/Plugins/Animation/ControlRig/Source/ControlRigDeveloper/Private/ControlRigBlueprint.cpp
--- a/Engine/Plugins/Animation/ControlRig/Source/ControlRigDeveloper/Private/ControlRigBlueprint.cpp
+++ b/Engine/Plugins/Animation/ControlRig/Source/ControlRigDeveloper/Private/ControlRigBlueprint.cpp
@@ -717,6 +717,7 @@ void UControlRigBlueprint::HandlePackageDone(const FEndLoadPackageContext& Conte
 	}
 	
 	PropagateHierarchyFromBPToInstances();
+	RemoveDeprecatedVMMemoryClass();
 	RecompileVM();
 	RequestControlRigInit();
 	BroadcastControlRigPackageDone();
@@ -738,6 +739,20 @@ void UControlRigBlueprint::BroadcastControlRigPackageDone()
 		}
 	}
 }
+
+void UControlRigBlueprint::RemoveDeprecatedVMMemoryClass() const
+{
+	TArray<UObject*> Objects;
+	GetObjectsWithPackage(GetPackage(), Objects, false);
+
+	for (UObject* Object : Objects)
+	{
+		if (URigVMMemoryStorageGeneratorClass* DeprecatedClass = Cast<URigVMMemoryStorageGeneratorClass>(Object))
+		{
+			DeprecatedClass->Rename(nullptr, GetTransientPackage(), REN_ForceNoResetLoaders | REN_DoNotDirty | REN_DontCreateRedirectors | REN_NonTransactional);
+		}
+	}
+}
 #endif
 
 void UControlRigBlueprint::RecompileVM()
diff --git a/Engine/Plugins/Animation/ControlRig/Source/ControlRigDeveloper/Public/ControlRigBlueprint.h b/Engine/Plugins/Animation/ControlRig/Source/ControlRigDeveloper/Public/ControlRigBlueprint.h
--- a/Engine/Plugins/Animation/ControlRig/Source/ControlRigDeveloper/Public/ControlRigBlueprint.h
+++ b/Engine/Plugins/Animation/ControlRig/Source/ControlRigDeveloper/Public/ControlRigBlueprint.h
@@ -607,6 +607,12 @@ private:
 	void HandlePackageDone(const FEndLoadPackageContext& Context);
 	// ControlRigBP, once end-loaded, will inform other ControlRig-Dependent systems that ControlRig instances are ready.
 	void BroadcastControlRigPackageDone();
+
+	// Previously some memory classes were parented to the package
+	// which violated the "single asset per package" rule
+	// this function removes those deprecated class.
+	// new classes should be created by RecompileVM and parented to the asset object
+	void RemoveDeprecatedVMMemoryClass() const;
 #endif
 
 	// Class used to temporarily cache all 
diff --git a/Engine/Source/Runtime/RigVM/Private/RigVMCore/RigVMMemoryStorage.cpp b/Engine/Source/Runtime/RigVM/Private/RigVMCore/RigVMMemoryStorage.cpp
--- a/Engine/Source/Runtime/RigVM/Private/RigVMCore/RigVMMemoryStorage.cpp
+++ b/Engine/Source/Runtime/RigVM/Private/RigVMCore/RigVMMemoryStorage.cpp
@@ -4,6 +4,7 @@
 #include "RigVMModule.h"
 #include "RigVMTypeUtils.h"
 #include "UObject/Interface.h"
+#include "AssetRegistry/AssetData.h"
 
 ////////////////////////////////////////////////////////////////////////////////////////////////////////////////////////
 
@@ -13,6 +14,44 @@ const FString FRigVMPropertyDescription::ContainerSuffix = TEXT(">");
 
 ////////////////////////////////////////////////////////////////////////////////////////////////////////////////////////
 
+// Generator class should be parented to the asset object, instead of the package
+// because the engine no longer supports multiple 'assets' per package
+static UObject* GetGeneratorClassOuter(UPackage* InPackage)
+{
+	if (!InPackage)
+	{
+		return nullptr;
+	}
+
+	// special case handling for transient package
+	if (InPackage == GetTransientPackage())
+	{
+		return InPackage;
+	}
+
+	// find the asset object to use as the class outer
+	// not using FindAssetInPackage here because
+	// in ControlRigEditor::UpdateControlRig
+	// we set the BP object to transient before VM recompilation
+	// so FindAssetInPackage would fail in finding the BP object
+	// the transient flag was needed such that recompile VM does not dirty BP
+	TArray<UObject*> TopLevelObjects;
+	GetObjectsWithOuter(InPackage, TopLevelObjects, false);
+	
+	UObject* AssetObject = nullptr;
+	for (UObject* Object : TopLevelObjects)
+	{
+		if (FAssetData::IsUAsset(Object))
+		{
+			AssetObject = Object;
+			break;
+		}
+	}
+
+	return AssetObject;	
+}
+
+
 FRigVMPropertyDescription::FRigVMPropertyDescription(const FProperty* InProperty, const FString& InDefaultValue, const FName& InName)
 	: Name(InName)
 	, Property(InProperty)
@@ -369,9 +408,15 @@ URigVMMemoryStorageGeneratorClass* URigVMMemoryStorageGeneratorClass::GetStorage
 	for(UObject* PotentialClassContainer : PotentialClassContainers)
 	{
 		UPackage* Package = PotentialClassContainer->GetOutermost();
-		if(URigVMMemoryStorageGeneratorClass* Class = FindObject<URigVMMemoryStorageGeneratorClass>(Package, *ClassName))
+
+		UObject* ClassOuter = GetGeneratorClassOuter(Package);
+
+		if (ClassOuter)
 		{
-			return Class;
+			if(URigVMMemoryStorageGeneratorClass* Class = FindObject<URigVMMemoryStorageGeneratorClass>(ClassOuter, *ClassName))
+			{
+				return Class;
+			}
 		}
 	}
 
@@ -386,12 +431,21 @@ URigVMMemoryStorageGeneratorClass* URigVMMemoryStorageGeneratorClass::CreateStor
 {
 	check(InOuter);
 	UPackage* Package = InOuter->GetOutermost();
+
+	// This class should be parented to the asset object, instead of the package
+	// because the engine no longer supports multiple 'assets' per package
+	UObject* ClassOuter = GetGeneratorClassOuter(Package);
+	if (!ensure(ClassOuter))
+	{
+		return nullptr;
+	}
+	
 	UClass *SuperClass = URigVMMemoryStorage::StaticClass();
 
 	const FString& ClassName = GetClassName(InMemoryType);
 
-	// if there's an old class - remove it from the package and mark it to destroy
-	URigVMMemoryStorageGeneratorClass* OldClass = FindObject<URigVMMemoryStorageGeneratorClass>(Package, *ClassName);
+	// if there's an old class - remove it from the asset and mark it to destroy
+	URigVMMemoryStorageGeneratorClass* OldClass = FindObject<URigVMMemoryStorageGeneratorClass>(ClassOuter, *ClassName);
 	if(OldClass)
 	{
 		FString DiscardedMemoryClassName;
@@ -412,9 +466,9 @@ URigVMMemoryStorageGeneratorClass* URigVMMemoryStorageGeneratorClass::CreateStor
 
 	// create the new class
 	URigVMMemoryStorageGeneratorClass* Class = NewObject<URigVMMemoryStorageGeneratorClass>(
-		Package,
+		ClassOuter,
 		*ClassName,
-		RF_Standalone | RF_Public
+		RF_Standalone
 	);
 
 	// clear the class (sets relevant flags)
@@ -476,11 +530,15 @@ bool URigVMMemoryStorageGeneratorClass::RemoveStorageClass(UObject* InOuter, ERi
 {
 	check(InOuter);
 	UPackage* Package = InOuter->GetOutermost();
-
+	
+	// This class should have been parented to the asset object, instead of the package
+	// because the engine no longer supports multiple 'assets' per package
+	UObject* ClassOuter = GetGeneratorClassOuter(Package);
+	
 	const FString ClassName = GetClassName(InMemoryType);
 
-	// if there's an old class - remove it from the package and mark it to destroy
-	URigVMMemoryStorageGeneratorClass* OldClass = FindObject<URigVMMemoryStorageGeneratorClass>(Package, *ClassName);
+	// if there's an old class - remove it from the asset and mark it to destroy
+	URigVMMemoryStorageGeneratorClass* OldClass = FindObject<URigVMMemoryStorageGeneratorClass>(ClassOuter, *ClassName);
 	if(OldClass)
 	{
 		OldClass->Rename(nullptr, GetTransientPackage(), REN_ForceNoResetLoaders | REN_DoNotDirty | REN_DontCreateRedirectors | REN_NonTransactional);
diff --git a/Engine/Source/Runtime/RigVM/Public/RigVMCore/RigVMMemoryStorage.h b/Engine/Source/Runtime/RigVM/Public/RigVMCore/RigVMMemoryStorage.h
--- a/Engine/Source/Runtime/RigVM/Public/RigVMCore/RigVMMemoryStorage.h
+++ b/Engine/Source/Runtime/RigVM/Public/RigVMCore/RigVMMemoryStorage.h
@@ -209,7 +209,8 @@ class RIGVM_API URigVMMemoryStorageGeneratorClass :
 	public UClass
 {
 	GENERATED_BODY()
-
+	DECLARE_WITHIN(UObject);
+	
 public:
 
 	// Default constructor
