ControlRig: ensure that the VM is recompiled immediately during ControlRig Blueprint postload such that other systems using control rig can use the most up-to-date version of VM instead of the serialized VM. And Because we now recompile VM eariler, memory Classes can get filled before their own loading took place. To avoid mem class serialization adding duplicated properties to the class, we have to throw away serialzed data and keep the freshly generated data.

diff --git a/Engine/Source/Runtime/RigVM/Private/RigVMCore/RigVMMemoryStorage.cpp b/Engine/Source/Runtime/RigVM/Private/RigVMCore/RigVMMemoryStorage.cpp
--- a/Engine/Source/Runtime/RigVM/Private/RigVMCore/RigVMMemoryStorage.cpp
+++ b/Engine/Source/Runtime/RigVM/Private/RigVMCore/RigVMMemoryStorage.cpp
@@ -350,6 +350,21 @@ void URigVMMemoryStorageGeneratorClass::Link(FArchive& Ar, bool bRelinkExistingP
 
 void URigVMMemoryStorageGeneratorClass::Serialize(FArchive& Ar)
 {
+	if (Ar.IsLoading() && ChildProperties != nullptr)
+	{
+		// if there are already properties in this class
+		// it means that the VM has regenerated the class before
+		// the class' own deserialization took place
+		// in that case, the VM generated class should take priority
+		// so a dummy class is used here to consume/throw away outdated serialized data
+		URigVMMemoryStorageGeneratorClass* AuxStorageClass = NewObject<URigVMMemoryStorageGeneratorClass>(
+			GetTransientPackage(),
+			TEXT("URigVMMemoryStorageGeneratorClass_Auxiliary"));
+		AuxStorageClass->ClassAddReferencedObjects = &this->AddReferencedObjects;
+		AuxStorageClass->Serialize(Ar);
+		return;
+	}
+
 	Super::Serialize(Ar);
 
 	if(Ar.IsLoading() || Ar.IsSaving())
