ControlRig: update vm mem class serialize() to use a new api for adding addrefobject func ptr

#horde https://horde.devtools.epicgames.com/log/62853e5c890768f0e937c3c1?lineindex=722

diff --git a/Engine/Source/Runtime/RigVM/Private/RigVMCore/RigVMMemoryStorage.cpp b/Engine/Source/Runtime/RigVM/Private/RigVMCore/RigVMMemoryStorage.cpp
--- a/Engine/Source/Runtime/RigVM/Private/RigVMCore/RigVMMemoryStorage.cpp
+++ b/Engine/Source/Runtime/RigVM/Private/RigVMCore/RigVMMemoryStorage.cpp
@@ -360,7 +360,7 @@ void URigVMMemoryStorageGeneratorClass::Serialize(FArchive& Ar)
 		URigVMMemoryStorageGeneratorClass* AuxStorageClass = NewObject<URigVMMemoryStorageGeneratorClass>(
 			GetTransientPackage(),
 			TEXT("URigVMMemoryStorageGeneratorClass_Auxiliary"));
-		AuxStorageClass->ClassAddReferencedObjects = &this->AddReferencedObjects;
+		AuxStorageClass->CppClassStaticFunctions.SetAddReferencedObjects(&this->AddReferencedObjects);
 		AuxStorageClass->Serialize(Ar);
 		return;
 	}
