Control Rig: Fix select nodes not supporting arrays of arrays

diff --git a/Engine/Plugins/Animation/ControlRig/Source/ControlRigDeveloper/Private/ControlRigBlueprint.cpp b/Engine/Plugins/Animation/ControlRig/Source/ControlRigDeveloper/Private/ControlRigBlueprint.cpp
--- a/Engine/Plugins/Animation/ControlRig/Source/ControlRigDeveloper/Private/ControlRigBlueprint.cpp
+++ b/Engine/Plugins/Animation/ControlRig/Source/ControlRigDeveloper/Private/ControlRigBlueprint.cpp
@@ -938,10 +938,10 @@ void UControlRigBlueprint::RefreshAllModels()
 			}
 		}
 #if WITH_EDITOR
-		// if (Controller->RecomputeAllTemplateFilteredTypes(false))
-		// {
-		// 	ensureMsgf(false, TEXT("Pin type changed during load %s"), *GetPackage()->GetPathName());
-		// }
+		if (Controller->RecomputeAllTemplateFilteredTypes(false))
+		{
+			ensureMsgf(false, TEXT("Pin type changed during load %s"), *GetPackage()->GetPathName());
+		}
 #endif
 	}
 }
diff --git a/Engine/Source/Developer/RigVMDeveloper/Private/RigVMModel/Nodes/RigVMArrayNode.cpp b/Engine/Source/Developer/RigVMDeveloper/Private/RigVMModel/Nodes/RigVMArrayNode.cpp
--- a/Engine/Source/Developer/RigVMDeveloper/Private/RigVMModel/Nodes/RigVMArrayNode.cpp
+++ b/Engine/Source/Developer/RigVMDeveloper/Private/RigVMModel/Nodes/RigVMArrayNode.cpp
@@ -254,6 +254,12 @@ const FRigVMTemplate* URigVMArrayNode::GetTemplate() const
 	{		
 		const TArray<FRigVMTemplateArgument::FType>& SingleTypes = FRigVMTemplateArgument::GetCompatibleTypes(FRigVMTemplateArgument::ETypeCategory_SingleAnyValue);
 		const TArray<FRigVMTemplateArgument::FType>& ArrayTypes = FRigVMTemplateArgument::GetCompatibleTypes(FRigVMTemplateArgument::ETypeCategory_ArrayAnyValue);
+		const TArray<FRigVMTemplateArgument::FType>& ArrayArrayTypes = FRigVMTemplateArgument::GetCompatibleTypes(FRigVMTemplateArgument::ETypeCategory_ArrayArrayAnyValue);
+
+		TArray<FRigVMTemplateArgument::FType> ArrayValueTypes = ArrayTypes;
+		ArrayValueTypes.Append(ArrayArrayTypes);
+		TArray<FRigVMTemplateArgument::FType> ElementTypes = SingleTypes;
+		ElementTypes.Append(ArrayTypes);
 
 		UScriptStruct* ExecuteStruct = GetGraph()->GetExecuteContextStruct();
 		FRigVMTemplateArgument::FType ExecuteType(ExecuteStruct->GetStructCPPName(), ExecuteStruct);
@@ -266,57 +272,57 @@ const FRigVMTemplate* URigVMArrayNode::GetTemplate() const
 			case ERigVMOpCode::ArrayReverse:
     		{
     			Arguments.Emplace(*FRigVMStruct::ExecuteContextName.ToString(), ERigVMPinDirection::IO, ExecuteType);
-    			Arguments.Emplace(*ArrayName, ERigVMPinDirection::IO, ArrayTypes);
+    			Arguments.Emplace(*ArrayName, ERigVMPinDirection::IO, ArrayValueTypes);
     			break;
     		}
     		case ERigVMOpCode::ArrayGetNum:
     		{
-    			Arguments.Emplace(*ArrayName, ERigVMPinDirection::Input, ArrayTypes);
+    			Arguments.Emplace(*ArrayName, ERigVMPinDirection::Input, ArrayValueTypes);
     			Arguments.Emplace(*NumName, ERigVMPinDirection::Output, FRigVMTemplateArgument::FType(RigVMTypeUtils::Int32Type, nullptr));
     			break;
     		} 
     		case ERigVMOpCode::ArraySetNum:
     		{
     			Arguments.Emplace(*FRigVMStruct::ExecuteContextName.ToString(), ERigVMPinDirection::IO, ExecuteType);
-    			Arguments.Emplace(*ArrayName, ERigVMPinDirection::IO, ArrayTypes);
+    			Arguments.Emplace(*ArrayName, ERigVMPinDirection::IO, ArrayValueTypes);
     			Arguments.Emplace(*NumName, ERigVMPinDirection::Input, FRigVMTemplateArgument::FType(RigVMTypeUtils::Int32Type, nullptr));
     			break;
     		}
     		case ERigVMOpCode::ArrayGetAtIndex:
     		{
-    			Arguments.Emplace(*ArrayName, ERigVMPinDirection::Input, ArrayTypes);
+    			Arguments.Emplace(*ArrayName, ERigVMPinDirection::Input, ArrayValueTypes);
     			Arguments.Emplace(*IndexName, ERigVMPinDirection::Input, FRigVMTemplateArgument::FType(RigVMTypeUtils::Int32Type, nullptr));
-    			Arguments.Emplace(*ElementName, ERigVMPinDirection::Output, SingleTypes);
+    			Arguments.Emplace(*ElementName, ERigVMPinDirection::Output, ElementTypes);
     			break;
     		}  
     		case ERigVMOpCode::ArraySetAtIndex:
 			case ERigVMOpCode::ArrayInsert:
     		{
     			Arguments.Emplace(*FRigVMStruct::ExecuteContextName.ToString(), ERigVMPinDirection::IO, ExecuteType);
-    			Arguments.Emplace(*ArrayName, ERigVMPinDirection::IO, ArrayTypes);
+    			Arguments.Emplace(*ArrayName, ERigVMPinDirection::IO, ArrayValueTypes);
     			Arguments.Emplace(*IndexName, ERigVMPinDirection::Input, FRigVMTemplateArgument::FType(RigVMTypeUtils::Int32Type, nullptr));
-    			Arguments.Emplace(*ElementName, ERigVMPinDirection::Input, SingleTypes);
+    			Arguments.Emplace(*ElementName, ERigVMPinDirection::Input, ElementTypes);
     			break;
     		}
     		case ERigVMOpCode::ArrayAdd:
     		{
     			Arguments.Emplace(*FRigVMStruct::ExecuteContextName.ToString(), ERigVMPinDirection::IO, ExecuteType);
-    			Arguments.Emplace(*ArrayName, ERigVMPinDirection::IO, ArrayTypes);
-    			Arguments.Emplace(*ElementName, ERigVMPinDirection::Input, SingleTypes);
+    			Arguments.Emplace(*ArrayName, ERigVMPinDirection::IO, ArrayValueTypes);
+    			Arguments.Emplace(*ElementName, ERigVMPinDirection::Input, ElementTypes);
     			Arguments.Emplace(*IndexName, ERigVMPinDirection::Output, FRigVMTemplateArgument::FType(RigVMTypeUtils::Int32Type, nullptr));
     			break;
     		}
     		case ERigVMOpCode::ArrayRemove:
     		{
     			Arguments.Emplace(*FRigVMStruct::ExecuteContextName.ToString(), ERigVMPinDirection::IO, ExecuteType);
-    			Arguments.Emplace(*ArrayName, ERigVMPinDirection::IO, ArrayTypes);
+    			Arguments.Emplace(*ArrayName, ERigVMPinDirection::IO, ArrayValueTypes);
     			Arguments.Emplace(*IndexName, ERigVMPinDirection::Input, FRigVMTemplateArgument::FType(RigVMTypeUtils::Int32Type, nullptr));
     			break;
     		}
     		case ERigVMOpCode::ArrayFind:
     		{
-    			Arguments.Emplace(*ArrayName, ERigVMPinDirection::Input, ArrayTypes);
-    			Arguments.Emplace(*ElementName, ERigVMPinDirection::Input, SingleTypes);
+    			Arguments.Emplace(*ArrayName, ERigVMPinDirection::Input, ArrayValueTypes);
+    			Arguments.Emplace(*ElementName, ERigVMPinDirection::Input, ElementTypes);
     			Arguments.Emplace(*IndexName, ERigVMPinDirection::Output, FRigVMTemplateArgument::FType(RigVMTypeUtils::Int32Type, nullptr));
     			Arguments.Emplace(*SuccessName, ERigVMPinDirection::Output, FRigVMTemplateArgument::FType(RigVMTypeUtils::BoolType, nullptr));
     			break;
@@ -325,21 +331,21 @@ const FRigVMTemplate* URigVMArrayNode::GetTemplate() const
 			case ERigVMOpCode::ArrayUnion:
     		{
     			Arguments.Emplace(*FRigVMStruct::ExecuteContextName.ToString(), ERigVMPinDirection::IO, ExecuteType);
-    			Arguments.Emplace(*ArrayName, ERigVMPinDirection::IO, ArrayTypes);
-    			Arguments.Emplace(*OtherName, ERigVMPinDirection::Input, ArrayTypes);
+    			Arguments.Emplace(*ArrayName, ERigVMPinDirection::IO, ArrayValueTypes);
+    			Arguments.Emplace(*OtherName, ERigVMPinDirection::Input, ArrayValueTypes);
     			break;
     		}
     		case ERigVMOpCode::ArrayClone:
     		{
-    			Arguments.Emplace(*ArrayName, ERigVMPinDirection::Input, ArrayTypes);
-    			Arguments.Emplace(*CloneName, ERigVMPinDirection::Output, ArrayTypes);
+    			Arguments.Emplace(*ArrayName, ERigVMPinDirection::Input, ArrayValueTypes);
+    			Arguments.Emplace(*CloneName, ERigVMPinDirection::Output, ArrayValueTypes);
     			break;    				
     		}
     		case ERigVMOpCode::ArrayIterator:
     		{
     			Arguments.Emplace(*FRigVMStruct::ExecuteContextName.ToString(), ERigVMPinDirection::IO, ExecuteType);
-    			Arguments.Emplace(*ArrayName, ERigVMPinDirection::Input, ArrayTypes);
-    			Arguments.Emplace(*ElementName, ERigVMPinDirection::Output, SingleTypes);
+    			Arguments.Emplace(*ArrayName, ERigVMPinDirection::Input, ArrayValueTypes);
+    			Arguments.Emplace(*ElementName, ERigVMPinDirection::Output, ElementTypes);
     			Arguments.Emplace(*IndexName, ERigVMPinDirection::Output, FRigVMTemplateArgument::FType(RigVMTypeUtils::Int32Type, nullptr));
     			Arguments.Emplace(*CountName, ERigVMPinDirection::Output, FRigVMTemplateArgument::FType(RigVMTypeUtils::Int32Type, nullptr));
     			Arguments.Emplace(*RatioName, ERigVMPinDirection::Output, FRigVMTemplateArgument::FType(RigVMTypeUtils::FloatType, nullptr));
@@ -350,9 +356,9 @@ const FRigVMTemplate* URigVMArrayNode::GetTemplate() const
     		case ERigVMOpCode::ArrayDifference:
     		case ERigVMOpCode::ArrayIntersection:
     		{
-    			Arguments.Emplace(*ArrayName, ERigVMPinDirection::Input, ArrayTypes);
-    			Arguments.Emplace(*OtherName, ERigVMPinDirection::Input, ArrayTypes);
-    			Arguments.Emplace(*ResultName, ERigVMPinDirection::Output, ArrayTypes);
+    			Arguments.Emplace(*ArrayName, ERigVMPinDirection::Input, ArrayValueTypes);
+    			Arguments.Emplace(*OtherName, ERigVMPinDirection::Input, ArrayValueTypes);
+    			Arguments.Emplace(*ResultName, ERigVMPinDirection::Output, ArrayValueTypes);
     			break;
     		}
     		default:
diff --git a/Engine/Source/Developer/RigVMDeveloper/Private/RigVMModel/Nodes/RigVMIfNode.cpp b/Engine/Source/Developer/RigVMDeveloper/Private/RigVMModel/Nodes/RigVMIfNode.cpp
--- a/Engine/Source/Developer/RigVMDeveloper/Private/RigVMModel/Nodes/RigVMIfNode.cpp
+++ b/Engine/Source/Developer/RigVMDeveloper/Private/RigVMModel/Nodes/RigVMIfNode.cpp
@@ -28,6 +28,7 @@ const FRigVMTemplate* URigVMIfNode::GetTemplate() const
 		TArray<FRigVMTemplateArgument::FType> Types;
 		Types.Append(FRigVMTemplateArgument::GetCompatibleTypes(FRigVMTemplateArgument::ETypeCategory_SingleAnyValue));
 		Types.Append(FRigVMTemplateArgument::GetCompatibleTypes(FRigVMTemplateArgument::ETypeCategory_ArrayAnyValue));
+		Types.Append(FRigVMTemplateArgument::GetCompatibleTypes(FRigVMTemplateArgument::ETypeCategory_ArrayArrayAnyValue));
 		
 		TArray<FRigVMTemplateArgument> Arguments;
 		Arguments.Emplace(*ConditionName, ERigVMPinDirection::Input, FRigVMTemplateArgument::FType(RigVMTypeUtils::BoolType, nullptr));
diff --git a/Engine/Source/Developer/RigVMDeveloper/Private/RigVMModel/Nodes/RigVMRerouteNode.cpp b/Engine/Source/Developer/RigVMDeveloper/Private/RigVMModel/Nodes/RigVMRerouteNode.cpp
--- a/Engine/Source/Developer/RigVMDeveloper/Private/RigVMModel/Nodes/RigVMRerouteNode.cpp
+++ b/Engine/Source/Developer/RigVMDeveloper/Private/RigVMModel/Nodes/RigVMRerouteNode.cpp
@@ -83,6 +83,7 @@ const FRigVMTemplate* URigVMRerouteNode::GetTemplate() const
 		TArray<FRigVMTemplateArgument::FType> Types;
 		Types.Append(FRigVMTemplateArgument::GetCompatibleTypes(FRigVMTemplateArgument::ETypeCategory_SingleAnyValue));
 		Types.Append(FRigVMTemplateArgument::GetCompatibleTypes(FRigVMTemplateArgument::ETypeCategory_ArrayAnyValue));
+		Types.Append(FRigVMTemplateArgument::GetCompatibleTypes(FRigVMTemplateArgument::ETypeCategory_ArrayArrayAnyValue));
 
 		UScriptStruct* ExecuteStruct = GetGraph()->GetExecuteContextStruct();
 		Types.Add(FRigVMTemplateArgument::FType(ExecuteStruct->GetStructCPPName(), ExecuteStruct));
diff --git a/Engine/Source/Developer/RigVMDeveloper/Private/RigVMModel/Nodes/RigVMSelectNode.cpp b/Engine/Source/Developer/RigVMDeveloper/Private/RigVMModel/Nodes/RigVMSelectNode.cpp
--- a/Engine/Source/Developer/RigVMDeveloper/Private/RigVMModel/Nodes/RigVMSelectNode.cpp
+++ b/Engine/Source/Developer/RigVMDeveloper/Private/RigVMModel/Nodes/RigVMSelectNode.cpp
@@ -38,11 +38,17 @@ const FRigVMTemplate* URigVMSelectNode::GetTemplate() const
 	{
 		const TArray<FRigVMTemplateArgument::FType>& SingleTypes = FRigVMTemplateArgument::GetCompatibleTypes(FRigVMTemplateArgument::ETypeCategory_SingleAnyValue);
 		const TArray<FRigVMTemplateArgument::FType>& ArrayTypes = FRigVMTemplateArgument::GetCompatibleTypes(FRigVMTemplateArgument::ETypeCategory_ArrayAnyValue);
+		const TArray<FRigVMTemplateArgument::FType>& ArrayArrayTypes = FRigVMTemplateArgument::GetCompatibleTypes(FRigVMTemplateArgument::ETypeCategory_ArrayArrayAnyValue);
+
+		TArray<FRigVMTemplateArgument::FType> ResultTypes = SingleTypes;
+		ResultTypes.Append(ArrayTypes);
+		TArray<FRigVMTemplateArgument::FType> ValueTypes = ArrayTypes;
+		ValueTypes.Append(ArrayArrayTypes);		
 		
 		TArray<FRigVMTemplateArgument> Arguments;
 		Arguments.Emplace(*IndexName, ERigVMPinDirection::Input, FRigVMTemplateArgument::FType(RigVMTypeUtils::Int32Type, nullptr));
-		Arguments.Emplace(*ValueName, ERigVMPinDirection::Input, ArrayTypes);
-		Arguments.Emplace(*ResultName, ERigVMPinDirection::Output, SingleTypes);
+		Arguments.Emplace(*ValueName, ERigVMPinDirection::Input, ValueTypes);
+		Arguments.Emplace(*ResultName, ERigVMPinDirection::Output, ResultTypes);
 		
 		CachedTemplate = FRigVMRegistry::Get().GetOrAddTemplateFromArguments(*SelectName, Arguments);
 	}
diff --git a/Engine/Source/Runtime/RigVM/Private/RigVMCore/RigVMTemplate.cpp b/Engine/Source/Runtime/RigVM/Private/RigVMCore/RigVMTemplate.cpp
--- a/Engine/Source/Runtime/RigVM/Private/RigVMCore/RigVMTemplate.cpp
+++ b/Engine/Source/Runtime/RigVM/Private/RigVMCore/RigVMTemplate.cpp
@@ -180,29 +180,41 @@ const TArray<FRigVMTemplateArgument::FType> FRigVMTemplateArgument::GetCompatibl
 	{
 		CompatibleTypes.Add(ETypeCategory_SingleAnyValue, TArray<FType>());
 		CompatibleTypes.Add(ETypeCategory_ArrayAnyValue, TArray<FType>());
+		CompatibleTypes.Add(ETypeCategory_ArrayArrayAnyValue, TArray<FType>());
 		CompatibleTypes.Add(ETypeCategory_SingleSimpleValue, TArray<FType>());
 		CompatibleTypes.Add(ETypeCategory_ArraySimpleValue, TArray<FType>());
+		CompatibleTypes.Add(ETypeCategory_ArrayArraySimpleValue, TArray<FType>());
 		CompatibleTypes.Add(ETypeCategory_SingleMathStructValue, TArray<FType>());
 		CompatibleTypes.Add(ETypeCategory_ArrayMathStructValue, TArray<FType>());
+		CompatibleTypes.Add(ETypeCategory_ArrayArrayMathStructValue, TArray<FType>());
 		CompatibleTypes.Add(ETypeCategory_SingleScriptStructValue, TArray<FType>());
 		CompatibleTypes.Add(ETypeCategory_ArrayScriptStructValue, TArray<FType>());
+		CompatibleTypes.Add(ETypeCategory_ArrayArrayScriptStructValue, TArray<FType>());
 		CompatibleTypes.Add(ETypeCategory_SingleEnumValue, TArray<FType>());
 		CompatibleTypes.Add(ETypeCategory_ArrayEnumValue, TArray<FType>());
+		CompatibleTypes.Add(ETypeCategory_ArrayArrayEnumValue, TArray<FType>());
 		CompatibleTypes.Add(ETypeCategory_SingleObjectValue, TArray<FType>());
 		CompatibleTypes.Add(ETypeCategory_ArrayObjectValue, TArray<FType>());
+		CompatibleTypes.Add(ETypeCategory_ArrayArrayObjectValue, TArray<FType>());
 		
 		TArray<FType>& SingleAnyValueTypes = CompatibleTypes.FindChecked(ETypeCategory_SingleAnyValue);
 		TArray<FType>& ArrayAnyValueTypes = CompatibleTypes.FindChecked(ETypeCategory_ArrayAnyValue);
+		TArray<FType>& ArrayArrayAnyValueTypes = CompatibleTypes.FindChecked(ETypeCategory_ArrayArrayAnyValue);
 		TArray<FType>& SingleSimpleValueTypes = CompatibleTypes.FindChecked(ETypeCategory_SingleSimpleValue);
 		TArray<FType>& ArraySimpleValueTypes = CompatibleTypes.FindChecked(ETypeCategory_ArraySimpleValue);
+		TArray<FType>& ArrayArraySimpleValueTypes = CompatibleTypes.FindChecked(ETypeCategory_ArrayArraySimpleValue);
 		TArray<FType>& SingleMathStructValueTypes = CompatibleTypes.FindChecked(ETypeCategory_SingleMathStructValue);
 		TArray<FType>& ArrayMathStructValueTypes = CompatibleTypes.FindChecked(ETypeCategory_ArrayMathStructValue);
+		TArray<FType>& ArrayArrayMathStructValueTypes = CompatibleTypes.FindChecked(ETypeCategory_ArrayArrayMathStructValue);
 		TArray<FType>& SingleScriptStructValueTypes = CompatibleTypes.FindChecked(ETypeCategory_SingleScriptStructValue);
 		TArray<FType>& ArrayScriptStructValueTypes = CompatibleTypes.FindChecked(ETypeCategory_ArrayScriptStructValue);
+		TArray<FType>& ArrayArrayScriptStructValueTypes = CompatibleTypes.FindChecked(ETypeCategory_ArrayArrayScriptStructValue);
 		TArray<FType>& SingleEnumValueTypes = CompatibleTypes.FindChecked(ETypeCategory_SingleEnumValue);
 		TArray<FType>& ArrayEnumValueTypes = CompatibleTypes.FindChecked(ETypeCategory_ArrayEnumValue);
+		TArray<FType>& ArrayArrayEnumValueTypes = CompatibleTypes.FindChecked(ETypeCategory_ArrayArrayEnumValue);
 		TArray<FType>& SingleObjectValueTypes = CompatibleTypes.FindChecked(ETypeCategory_SingleObjectValue);
 		TArray<FType>& ArrayObjectValueTypes = CompatibleTypes.FindChecked(ETypeCategory_ArrayObjectValue);
+		TArray<FType>& ArrayArrayObjectValueTypes = CompatibleTypes.FindChecked(ETypeCategory_ArrayArrayObjectValue);
 		
 		SingleSimpleValueTypes.Add(FType(RigVMTypeUtils::BoolType));
 		SingleSimpleValueTypes.Add(FType(RigVMTypeUtils::Int32Type));
@@ -215,10 +227,12 @@ const TArray<FRigVMTemplateArgument::FType> FRigVMTemplateArgument::GetCompatibl
 		for(const FType& Type : SingleSimpleValueTypes)
 		{
 			ArraySimpleValueTypes.Add(FType(RigVMTypeUtils::ArrayTypeFromBaseType(Type.CPPType)));
+			ArrayArraySimpleValueTypes.Add(FType(RigVMTypeUtils::ArrayTypeFromBaseType(ArraySimpleValueTypes.Last().CPPType)));
 		}
 
 		SingleAnyValueTypes = SingleSimpleValueTypes;
 		ArrayAnyValueTypes = ArraySimpleValueTypes;
+		ArrayArrayAnyValueTypes = ArrayArraySimpleValueTypes;
 
 		static const TArray<UScriptStruct*> MathTypes = { 
 			TBaseStructure<FRotator>::Get(),
@@ -237,6 +251,7 @@ const TArray<FRigVMTemplateArgument::FType> FRigVMTemplateArgument::GetCompatibl
 		{
 			SingleMathStructValueTypes.Add(FType(MathType->GetStructCPPName(), MathType));
 			ArrayMathStructValueTypes.Add(FType(RigVMTypeUtils::ArrayTypeFromBaseType(MathType->GetStructCPPName()), MathType));
+			ArrayArrayMathStructValueTypes.Add(FType(RigVMTypeUtils::ArrayTypeFromBaseType(ArrayMathStructValueTypes.Last().CPPType), MathType));
 		}
 
 		struct FTypeTraverser
@@ -363,8 +378,10 @@ const TArray<FRigVMTemplateArgument::FType> FRigVMTemplateArgument::GetCompatibl
 			{
 				SingleAnyValueTypes.Add(FType(CPPType, ScriptStruct));
 				ArrayAnyValueTypes.Add(FType(RigVMTypeUtils::ArrayTypeFromBaseType(CPPType), ScriptStruct));
+				ArrayArrayAnyValueTypes.Add(FType(RigVMTypeUtils::ArrayTypeFromBaseType(ArrayAnyValueTypes.Last().CPPType), ScriptStruct));
 				SingleScriptStructValueTypes.Add(FType(CPPType, ScriptStruct));
 				ArrayScriptStructValueTypes.Add(FType(RigVMTypeUtils::ArrayTypeFromBaseType(CPPType), ScriptStruct));
+				ArrayArrayScriptStructValueTypes.Add(FType(RigVMTypeUtils::ArrayTypeFromBaseType(ArrayScriptStructValueTypes.Last().CPPType), ScriptStruct));
 			}
 		}
 
@@ -379,8 +396,10 @@ const TArray<FRigVMTemplateArgument::FType> FRigVMTemplateArgument::GetCompatibl
 			const FString CPPType = Enum->CppType.IsEmpty() ? Enum->GetName() : Enum->CppType;
 			SingleAnyValueTypes.Add(FType(CPPType, Enum));
 			ArrayAnyValueTypes.Add(FType(RigVMTypeUtils::ArrayTypeFromBaseType(CPPType), Enum));
+			ArrayArrayAnyValueTypes.Add(FType(RigVMTypeUtils::ArrayTypeFromBaseType(ArrayAnyValueTypes.Last().CPPType), Enum));
 			SingleEnumValueTypes.Add(FType(CPPType, Enum));
 			ArrayEnumValueTypes.Add(FType(RigVMTypeUtils::ArrayTypeFromBaseType(CPPType), Enum));
+			ArrayArrayEnumValueTypes.Add(FType(RigVMTypeUtils::ArrayTypeFromBaseType(ArrayEnumValueTypes.Last().CPPType), Enum));
 		}
 
 		// add all classes
@@ -397,17 +416,25 @@ const TArray<FRigVMTemplateArgument::FType> FRigVMTemplateArgument::GetCompatibl
 			const FString CPPType = Class->GetPrefixCPP() + Class->GetName();
 			SingleAnyValueTypes.Add(FType(CPPType, Class));
 			ArrayAnyValueTypes.Add(FType(RigVMTypeUtils::ArrayTypeFromBaseType(CPPType), Class));
+			ArrayArrayAnyValueTypes.Add(FType(RigVMTypeUtils::ArrayTypeFromBaseType(ArrayAnyValueTypes.Last().CPPType), Class));
 			SingleObjectValueTypes.Add(FType(CPPType, Class));
 			ArrayObjectValueTypes.Add(FType(RigVMTypeUtils::ArrayTypeFromBaseType(CPPType), Class));
+			ArrayArrayObjectValueTypes.Add(FType(RigVMTypeUtils::ArrayTypeFromBaseType(ArrayObjectValueTypes.Last().CPPType), Class));
 		}
 
 		// check that matching categories have the same size
 		check(CompatibleTypes.FindChecked(ETypeCategory_SingleAnyValue).Num() == CompatibleTypes.FindChecked(ETypeCategory_ArrayAnyValue).Num());
+		check(CompatibleTypes.FindChecked(ETypeCategory_SingleAnyValue).Num() == CompatibleTypes.FindChecked(ETypeCategory_ArrayArrayAnyValue).Num());
 		check(CompatibleTypes.FindChecked(ETypeCategory_SingleSimpleValue).Num() == CompatibleTypes.FindChecked(ETypeCategory_ArraySimpleValue).Num());
+		check(CompatibleTypes.FindChecked(ETypeCategory_SingleSimpleValue).Num() == CompatibleTypes.FindChecked(ETypeCategory_ArrayArraySimpleValue).Num());
 		check(CompatibleTypes.FindChecked(ETypeCategory_SingleMathStructValue).Num() == CompatibleTypes.FindChecked(ETypeCategory_ArrayMathStructValue).Num());
+		check(CompatibleTypes.FindChecked(ETypeCategory_SingleMathStructValue).Num() == CompatibleTypes.FindChecked(ETypeCategory_ArrayArrayMathStructValue).Num());
 		check(CompatibleTypes.FindChecked(ETypeCategory_SingleScriptStructValue).Num() == CompatibleTypes.FindChecked(ETypeCategory_ArrayScriptStructValue).Num());
+		check(CompatibleTypes.FindChecked(ETypeCategory_SingleScriptStructValue).Num() == CompatibleTypes.FindChecked(ETypeCategory_ArrayArrayScriptStructValue).Num());
 		check(CompatibleTypes.FindChecked(ETypeCategory_SingleEnumValue).Num() == CompatibleTypes.FindChecked(ETypeCategory_ArrayEnumValue).Num());
+		check(CompatibleTypes.FindChecked(ETypeCategory_SingleEnumValue).Num() == CompatibleTypes.FindChecked(ETypeCategory_ArrayArrayEnumValue).Num());
 		check(CompatibleTypes.FindChecked(ETypeCategory_SingleObjectValue).Num() == CompatibleTypes.FindChecked(ETypeCategory_ArrayObjectValue).Num());
+		check(CompatibleTypes.FindChecked(ETypeCategory_SingleObjectValue).Num() == CompatibleTypes.FindChecked(ETypeCategory_ArrayArrayObjectValue).Num());
 	}
 
 	return CompatibleTypes.FindChecked(InCategory);
diff --git a/Engine/Source/Runtime/RigVM/Public/RigVMCore/RigVMTemplate.h b/Engine/Source/Runtime/RigVM/Public/RigVMCore/RigVMTemplate.h
--- a/Engine/Source/Runtime/RigVM/Public/RigVMCore/RigVMTemplate.h
+++ b/Engine/Source/Runtime/RigVM/Public/RigVMCore/RigVMTemplate.h
@@ -147,6 +147,7 @@ struct RIGVM_API FRigVMTemplateArgument
 	{
 		EArrayType_SingleValue,
 		EArrayType_ArrayValue,
+		EArrayType_ArrayArrayValue,
 		EArrayType_Mixed,
 		EArrayType_Invalid
 	};
@@ -155,16 +156,22 @@ struct RIGVM_API FRigVMTemplateArgument
 	{
 		ETypeCategory_SingleAnyValue,
 		ETypeCategory_ArrayAnyValue,
+		ETypeCategory_ArrayArrayAnyValue,
 		ETypeCategory_SingleSimpleValue,
 		ETypeCategory_ArraySimpleValue,
+		ETypeCategory_ArrayArraySimpleValue,
 		ETypeCategory_SingleMathStructValue,
 		ETypeCategory_ArrayMathStructValue,
+		ETypeCategory_ArrayArrayMathStructValue,
 		ETypeCategory_SingleScriptStructValue,
 		ETypeCategory_ArrayScriptStructValue,
+		ETypeCategory_ArrayArrayScriptStructValue,
 		ETypeCategory_SingleEnumValue,
 		ETypeCategory_ArrayEnumValue,
+		ETypeCategory_ArrayArrayEnumValue,
 		ETypeCategory_SingleObjectValue,
 		ETypeCategory_ArrayObjectValue,
+		ETypeCategory_ArrayArrayObjectValue,
 		ETypeCategory_Invalid
 	};
 
