ControlRig: Fix rig vm validation for external typed operands

#preflgiht https://horde.devtools.epicgames.com/job/628dcdf12f2409bc1e1141bd

diff --git a/Engine/Source/Runtime/RigVM/Private/RigVMCore/RigVM.cpp b/Engine/Source/Runtime/RigVM/Private/RigVMCore/RigVM.cpp
--- a/Engine/Source/Runtime/RigVM/Private/RigVMCore/RigVM.cpp
+++ b/Engine/Source/Runtime/RigVM/Private/RigVMCore/RigVM.cpp
@@ -318,7 +318,7 @@
 
 	TArray<URigVMMemoryStorage*> LocalMemory = { GetWorkMemory(), GetLiteralMemory(), GetDebugMemory() };
 	
-	auto CheckOperandValidity = [LocalMemory, &bAllOperandsValid](const FRigVMOperand& InOperand) -> bool
+	auto CheckOperandValidity = [LocalMemory, &bAllOperandsValid, this](const FRigVMOperand& InOperand) -> bool
 	{
 		if(InOperand.GetContainerIndex() < 0 || InOperand.GetContainerIndex() >= (int32)ERigVMMemoryType::Invalid)
 		{
@@ -328,7 +328,6 @@
 
 
 		const URigVMMemoryStorage* MemoryForOperand = LocalMemory[InOperand.GetContainerIndex()];
-
 		if(InOperand.GetMemoryType() != ERigVMMemoryType::External)
 		{
 			if(!MemoryForOperand->IsValidIndex(InOperand.GetRegisterIndex()))
@@ -336,17 +335,29 @@
 				bAllOperandsValid = false;
 				return false;
 			}
-		}
 
-		if(InOperand.GetRegisterOffset() != INDEX_NONE)
-		{
-			if(!MemoryForOperand->GetPropertyPaths().IsValidIndex(InOperand.GetRegisterOffset()))
+			if(InOperand.GetRegisterOffset() != INDEX_NONE)
 			{
-				bAllOperandsValid = false;
-				return false;
+				if(!MemoryForOperand->GetPropertyPaths().IsValidIndex(InOperand.GetRegisterOffset()))
+				{
+					bAllOperandsValid = false;
+					return false;
+				}
 			}
 		}
-
+		else if (InOperand.GetMemoryType() == ERigVMMemoryType::External)
+		{
+			// given that external variables array is populated at runtime
+			// checking for property path is the best we can do
+			if(InOperand.GetRegisterOffset() != INDEX_NONE)
+			{
+				if (!ExternalPropertyPathDescriptions.IsValidIndex(InOperand.GetRegisterOffset()))
+				{
+					bAllOperandsValid = false;
+					return false;
+				}
+			}
+		}
 		return true;
 	};
 	
