RigVM: Event invocation first pass

diff --git a/Engine/Source/Runtime/RigVM/Private/RigVMCore/RigVM.cpp b/Engine/Source/Runtime/RigVM/Private/RigVMCore/RigVM.cpp
--- a/Engine/Source/Runtime/RigVM/Private/RigVMCore/RigVM.cpp
+++ b/Engine/Source/Runtime/RigVM/Private/RigVMCore/RigVM.cpp
@@ -72,6 +72,7 @@
 	, DebugMemoryStorageObject(nullptr)
 	, ByteCodePtr(&ByteCodeStorage)
 	, NumExecutions(0)
+	, bHitErrorDuringExecution(false)
 #if WITH_EDITOR
 	, DebugInfo(nullptr)
 	, HaltedAtBreakpoint(nullptr)
@@ -1140,6 +1141,10 @@
 				{
 					break;
 				}
+			case ERigVMOpCode::InvokeEntry:
+				{
+					break;
+				}
 			case ERigVMOpCode::Invalid:
 			default:
 				{
@@ -1565,7 +1570,7 @@
 			case ERigVMOpCode::ArrayDifference:
 			case ERigVMOpCode::ArrayIntersection:
 			case ERigVMOpCode::ArrayReverse:
-
+			case ERigVMOpCode::InvokeEntry:
 			{
 				break;
 			}
@@ -1583,60 +1588,77 @@
 
 bool URigVM::Execute(TArrayView<URigVMMemoryStorage*> Memory, TArrayView<void*> AdditionalArguments, const FName& InEntryName)
 {
-	if (ExecutingThreadId != INDEX_NONE)
+	// if this the first entry being executed - get ready for execution
+	const bool bIsRootEntry = EntriesBeingExecuted.IsEmpty(); 
+	if(bIsRootEntry)
 	{
-		ensureMsgf(ExecutingThreadId == FPlatformTLS::GetCurrentThreadId(), TEXT("RigVM::Execute from multiple threads (%d and %d)"), ExecutingThreadId, (int32)FPlatformTLS::GetCurrentThreadId());
-	}
-	CopyDeferredVMIfRequired();
-	TGuardValue<int32> GuardThreadId(ExecutingThreadId, FPlatformTLS::GetCurrentThreadId());
+		bHitErrorDuringExecution = false;
+		
+		if (ExecutingThreadId != INDEX_NONE)
+		{
+			ensureMsgf(ExecutingThreadId == FPlatformTLS::GetCurrentThreadId(), TEXT("RigVM::Execute from multiple threads (%d and %d)"), ExecutingThreadId, (int32)FPlatformTLS::GetCurrentThreadId());
+		}
+		CopyDeferredVMIfRequired();
+		TGuardValue<int32> GuardThreadId(ExecutingThreadId, FPlatformTLS::GetCurrentThreadId());
 
-	ResolveFunctionsIfRequired();
-	RefreshInstructionsIfRequired();
+		ResolveFunctionsIfRequired();
+		RefreshInstructionsIfRequired();
 
-	if (Instructions.Num() == 0)
-	{
-		return true;
-	}
+		if (Instructions.Num() == 0)
+		{
+			return true;
+		}
 
-	// changes to the layout of memory array should be reflected in GetContainerIndex()
-	TArray<URigVMMemoryStorage*> LocalMemory;
-	if (Memory.Num() == 0)
-	{
-		LocalMemory = GetLocalMemoryArray();
-		Memory = LocalMemory;
+		// changes to the layout of memory array should be reflected in GetContainerIndex()
+		TArray<URigVMMemoryStorage*> LocalMemory;
+		if (Memory.Num() == 0)
+		{
+			LocalMemory = GetLocalMemoryArray();
+			Memory = LocalMemory;
+		}
+	
+		CacheMemoryHandlesIfRequired(Memory);
 	}
 	
-	CacheMemoryHandlesIfRequired(Memory);
 	FRigVMByteCode& ByteCode = GetByteCode();
 	TArray<FRigVMFunctionPtr>& Functions = GetFunctions();
 
 #if WITH_EDITOR
 	TArray<FName>& FunctionNames = GetFunctionNames();
 
-	if (FirstEntryEventInQueue == NAME_None || FirstEntryEventInQueue == InEntryName)
+	if(bIsRootEntry)
 	{
-		InstructionVisitedDuringLastRun.Reset();
-		InstructionVisitOrder.Reset();
-		InstructionVisitedDuringLastRun.SetNumZeroed(Instructions.Num());
-		InstructionCyclesDuringLastRun.Reset();
-		if(Context.PublicData.RuntimeSettings.bEnableProfiling)
+		if (FirstEntryEventInQueue == NAME_None || FirstEntryEventInQueue == InEntryName)
 		{
-			InstructionCyclesDuringLastRun.SetNumUninitialized(Instructions.Num());
-			for(int32 DurationIndex=0;DurationIndex<InstructionCyclesDuringLastRun.Num();DurationIndex++)
+			InstructionVisitedDuringLastRun.Reset();
+			InstructionVisitOrder.Reset();
+			InstructionVisitedDuringLastRun.SetNumZeroed(Instructions.Num());
+			InstructionCyclesDuringLastRun.Reset();
+			if(Context.PublicData.RuntimeSettings.bEnableProfiling)
 			{
-				InstructionCyclesDuringLastRun[DurationIndex] = UINT64_MAX;
+				InstructionCyclesDuringLastRun.SetNumUninitialized(Instructions.Num());
+				for(int32 DurationIndex=0;DurationIndex<InstructionCyclesDuringLastRun.Num();DurationIndex++)
+				{
+					InstructionCyclesDuringLastRun[DurationIndex] = UINT64_MAX;
+				}
 			}
 		}
 	}
 #endif
 
-	Context.Reset();
-	Context.SliceOffsets.AddZeroed(Instructions.Num());
-	Context.OpaqueArguments = AdditionalArguments;
+	if(bIsRootEntry)
+	{
+		Context.Reset();
+		Context.SliceOffsets.AddZeroed(Instructions.Num());
+		Context.OpaqueArguments = AdditionalArguments;
+	}
 
 	TGuardValue<URigVM*> VMInContext(Context.VM, this);
 
-	ClearDebugMemory();
+	if(bIsRootEntry)
+	{
+		ClearDebugMemory();
+	}
 
 	if (!InEntryName.IsNone())
 	{
@@ -1646,16 +1668,51 @@
 			return false;
 		}
 		SetInstructionIndex((uint16)ByteCode.GetEntry(EntryIndex).InstructionIndex);
+		EntriesBeingExecuted.Push(EntryIndex);
+
+		if(bIsRootEntry)
+		{
+			Context.PublicData.EventName = InEntryName;
+		}
+	}
+	else
+	{
+		int32 FirstEntryIndex = INDEX_NONE;
+		for(int32 EntryIndex = 0; EntryIndex < ByteCode.NumEntries(); EntryIndex++)
+		{
+			if(ByteCode.GetEntry(EntryIndex).InstructionIndex == 0)
+			{
+				FirstEntryIndex = EntryIndex;
+				break;
+			}
+		}
+		
+		EntriesBeingExecuted.Push(FirstEntryIndex);
+
+		if(bIsRootEntry)
+		{
+			if(ByteCode.Entries.IsValidIndex(FirstEntryIndex))
+			{
+				Context.PublicData.EventName = ByteCode.GetEntry(FirstEntryIndex).Name;
+			}
+			else
+			{
+				Context.PublicData.EventName = NAME_None;
+			}
+		}
 	}
 
 #if WITH_EDITOR
-	if (DebugInfo)
+	if (DebugInfo && bIsRootEntry)
 	{
 		DebugInfo->StartExecution();
 	}
 #endif
 
-	NumExecutions++;
+	if(bIsRootEntry)
+	{
+		NumExecutions++;
+	}
 
 #if WITH_EDITOR
 	uint64 StartCycles = 0;
@@ -1668,9 +1725,18 @@
 
 	while (Instructions.IsValidIndex(Context.PublicData.InstructionIndex))
 	{
+		if(bHitErrorDuringExecution)
+		{
+			EntriesBeingExecuted.Pop();
+			return false;
+		}
+		
 #if WITH_EDITOR
 		if (DebugInfo && ShouldHaltAtInstruction(InEntryName, Context.PublicData.InstructionIndex))
 		{
+			// we'll recursively exit all invoked
+			// entries here.
+			EntriesBeingExecuted.Pop();
 			return true;
 		}
 
@@ -1923,18 +1989,22 @@
 			}
 			case ERigVMOpCode::Exit:
 			{
+				if(bIsRootEntry)
+				{
 #if WITH_EDITOR
-				Context.LastExecutionMicroSeconds = OverallCycles * FPlatformTime::GetSecondsPerCycle() * 1000.0 * 1000.0;
+					Context.LastExecutionMicroSeconds = OverallCycles * FPlatformTime::GetSecondsPerCycle() * 1000.0 * 1000.0;
 #endif
-				ExecutionReachedExit().Broadcast(InEntryName);
+					ExecutionReachedExit().Broadcast(InEntryName);
 #if WITH_EDITOR					
-				if (HaltedAtBreakpoint != nullptr)
-				{
-					HaltedAtBreakpoint = nullptr;
-					DebugInfo->SetCurrentActiveBreakpoint(nullptr);
-					ExecutionHalted().Broadcast(INDEX_NONE, nullptr, InEntryName);
+					if (HaltedAtBreakpoint != nullptr)
+					{
+						HaltedAtBreakpoint = nullptr;
+						DebugInfo->SetCurrentActiveBreakpoint(nullptr);
+						ExecutionHalted().Broadcast(INDEX_NONE, nullptr, InEntryName);
+					}
+#endif
 				}
-#endif
+				EntriesBeingExecuted.Pop();
 				return true;
 			}
 			case ERigVMOpCode::BeginBlock:
@@ -2490,9 +2560,53 @@
 				Context.PublicData.InstructionIndex++;
 				break;
 			}
+			case ERigVMOpCode::InvokeEntry:
+			{
+				const FRigVMInvokeEntryOp& Op = ByteCode.GetOpAt<FRigVMInvokeEntryOp>(Instruction);
+				const int32 EntryIndex = ByteCode.FindEntryIndex(Op.EntryName);
+
+				if(EntryIndex == INDEX_NONE)
+				{
+					static constexpr TCHAR MissingEntry[] = TEXT("Entry('%s') cannot be found.");
+					Context.PublicData.Logf(EMessageSeverity::Error, MissingEntry, *Op.EntryName.ToString());
+					bHitErrorDuringExecution = true;
+				}
+				else if(EntriesBeingExecuted.Contains(EntryIndex)) // avoid cycles
+				{
+					TArray<FString> EntryNames;
+					for(int32 EntryBeingExecuted : EntriesBeingExecuted)
+					{
+						EntryNames.Add(ByteCode.GetEntry(EntryBeingExecuted).Name.ToString());
+					}
+					EntryNames.Add(Op.EntryName.ToString());
+
+					static constexpr TCHAR RecursiveEntry[] = TEXT("Entry('%s') is being invoked recursively (%s).");
+					Context.PublicData.Logf(EMessageSeverity::Error, RecursiveEntry, *Op.EntryName.ToString(), *FString::Join(EntryNames, TEXT(" -> ")));
+					bHitErrorDuringExecution = true;
+				}
+				else
+				{
+					// this will restore the public data after invoking the entry
+					TGuardValue<FRigVMExecuteContext> PublicDataGuard(Context.PublicData, Context.PublicData);
+					Execute(Memory, AdditionalArguments, Op.EntryName);
+
+#if WITH_EDITOR
+					// if we are halted at a break point we need to exit here
+					if (DebugInfo && ShouldHaltAtInstruction(InEntryName, Context.PublicData.InstructionIndex))
+					{
+						EntriesBeingExecuted.Pop();
+						return true;
+					}
+#endif
+				}
+					
+				Context.PublicData.InstructionIndex++;
+				break;
+			}
 			case ERigVMOpCode::Invalid:
 			{
 				ensure(false);
+				EntriesBeingExecuted.Pop();
 				return false;
 			}
 		}
@@ -2518,14 +2632,18 @@
 	}
 
 #if WITH_EDITOR
-	if (HaltedAtBreakpoint != nullptr)
+	if(bIsRootEntry)
 	{
-		DebugInfo->SetCurrentActiveBreakpoint(nullptr);
-		HaltedAtBreakpoint = nullptr;
-		ExecutionHalted().Broadcast(INDEX_NONE, nullptr, InEntryName);
+		if (HaltedAtBreakpoint != nullptr)
+		{
+			DebugInfo->SetCurrentActiveBreakpoint(nullptr);
+			HaltedAtBreakpoint = nullptr;
+			ExecutionHalted().Broadcast(INDEX_NONE, nullptr, InEntryName);
+		}
 	}
 #endif
 
+	EntriesBeingExecuted.Pop();
 	return true;
 }
 
@@ -2902,6 +3020,12 @@
 				ResultLine = FString::Printf(TEXT("Reverse array %s"), *GetOperandLabel(Op.Arg, OperandFormatFunction));
 				break;
 			}
+			case ERigVMOpCode::InvokeEntry:
+			{
+				const FRigVMInvokeEntryOp& Op = ByteCode.GetOpAt<FRigVMInvokeEntryOp>(Instructions[InstructionIndex]);
+				ResultLine = FString::Printf(TEXT("Invoke entry %s"), *Op.EntryName.ToString());
+				break;
+			}
 			default:
 			{
 				ensure(false);
diff --git a/Engine/Source/Runtime/RigVM/Private/RigVMCore/RigVMByteCode.cpp b/Engine/Source/Runtime/RigVM/Private/RigVMCore/RigVMByteCode.cpp
--- a/Engine/Source/Runtime/RigVM/Private/RigVMCore/RigVMByteCode.cpp
+++ b/Engine/Source/Runtime/RigVM/Private/RigVMCore/RigVMByteCode.cpp
@@ -117,6 +117,21 @@ void FRigVMChangeTypeOp::Serialize(FArchive& Ar)
 	ensure(false);
 }
 
+void FRigVMInvokeEntryOp::Serialize(FArchive& Ar)
+{
+	if(Ar.IsLoading())
+	{
+		FString EntryNameString;
+		Ar << EntryNameString;
+		EntryName = *EntryNameString;
+	}
+	else
+	{
+		FString EntryNameString = EntryName.ToString();
+		Ar << EntryNameString;
+	}
+}
+
 FRigVMInstructionArray::FRigVMInstructionArray()
 {
 }
@@ -400,6 +415,12 @@ void FRigVMByteCode::Save(FArchive& Ar) const
 				// so all we need is the previously saved OpCode.
 				break;
 			}
+			case ERigVMOpCode::InvokeEntry:
+			{
+				FRigVMInvokeEntryOp Op = GetOpAt<FRigVMInvokeEntryOp>(Instruction.ByteCodeIndex);
+				Ar << Op;
+				break;
+			}
 			default:
 			{
 				ensure(false);
@@ -624,6 +645,13 @@ void FRigVMByteCode::Load(FArchive& Ar)
 				AddEndBlockOp();
 				break;
 			}
+			case ERigVMOpCode::InvokeEntry:
+			{
+				FRigVMInvokeEntryOp Op;
+				Ar << Op;
+				AddOp<FRigVMInvokeEntryOp>(Op);
+				break;
+			}
 			default:
 			{
 				ensure(false);
@@ -860,6 +888,10 @@ uint32 FRigVMByteCode::GetOperatorHash(const FRigVMInstruction& InInstruction) c
 		{
 			return GetTypeHash(GetOpAt<FRigVMBaseOp>(InInstruction));
 		}
+		case ERigVMOpCode::InvokeEntry:
+		{
+			return GetTypeHash(GetOpAt<FRigVMInvokeEntryOp>(InInstruction));
+		}
 		case ERigVMOpCode::Invalid:
 		{
 			ensure(false);
@@ -1063,6 +1095,10 @@ uint64 FRigVMByteCode::GetOpNumBytesAt(uint64 InByteCodeIndex, bool bIncludeOper
 		{
 			return (uint64)sizeof(FRigVMBaseOp);
 		}
+		case ERigVMOpCode::InvokeEntry:
+		{
+			return (uint64)sizeof(FRigVMInvokeEntryOp);
+		}
 		case ERigVMOpCode::Invalid:
 		{
 			ensure(false);
@@ -1415,6 +1451,12 @@ FString FRigVMByteCode::DumpToText() const
 				Line += FString::Printf(TEXT(", ArgF %s"), *ArgF);
 				break;
 			}
+			case ERigVMOpCode::InvokeEntry:
+			{
+				const FRigVMInvokeEntryOp& Op = GetOpAt<FRigVMInvokeEntryOp>(Instruction.ByteCodeIndex);
+				Line += FString::Printf(TEXT(", Entry '%s'"), *Op.EntryName.ToString());
+				break;
+			}
 			case ERigVMOpCode::Invalid:
 			{
 				ensure(false);
@@ -1528,6 +1570,11 @@ uint64 FRigVMByteCode::AddArrayReverseOp(FRigVMOperand InArrayArg)
 	return AddOp(FRigVMUnaryOp(ERigVMOpCode::ArrayReverse, InArrayArg));
 }
 
+uint64 FRigVMByteCode::AddInvokeEntryOp(const FName& InEntryName)
+{
+	return AddOp(FRigVMInvokeEntryOp(InEntryName));
+}
+
 FRigVMOperandArray FRigVMByteCode::GetOperandsForOp(const FRigVMInstruction& InInstruction) const
 {
 	switch(InInstruction.OpCode)
@@ -1831,6 +1878,11 @@ uint64 FRigVMByteCode::GetOpAlignment(ERigVMOpCode InOpCode) const
 			static const uint64 Alignment = FRigVMBaseOp::StaticStruct()->GetCppStructOps()->GetAlignment();
 			return Alignment;
 		}
+		case ERigVMOpCode::InvokeEntry:
+		{
+			static const uint64 Alignment = FRigVMInvokeEntryOp::StaticStruct()->GetCppStructOps()->GetAlignment();
+			return Alignment;
+		}
 		case ERigVMOpCode::Invalid:
 		{
 			ensure(false);
diff --git a/Engine/Source/Runtime/RigVM/Public/RigVMCore/RigVM.h b/Engine/Source/Runtime/RigVM/Public/RigVMCore/RigVM.h
--- a/Engine/Source/Runtime/RigVM/Public/RigVMCore/RigVM.h
+++ b/Engine/Source/Runtime/RigVM/Public/RigVMCore/RigVM.h
@@ -653,6 +653,8 @@ private:
 	UPROPERTY(transient)
 	uint32 NumExecutions;
 
+	bool bHitErrorDuringExecution;
+
 #if WITH_EDITOR
 	FRigVMDebugInfo* DebugInfo;
 	TSharedPtr<FRigVMBreakpoint> HaltedAtBreakpoint;
@@ -743,6 +745,7 @@ private:
 	TMap<FRigVMOperand, TArray<FRigVMOperand>> OperandToDebugRegisters;
 
 	int32 ExecutingThreadId;
+	TArray<int32> EntriesBeingExecuted;
 
 	UPROPERTY(transient)
 	TObjectPtr<URigVM> DeferredVMToCopy;
diff --git a/Engine/Source/Runtime/RigVM/Public/RigVMCore/RigVMByteCode.h b/Engine/Source/Runtime/RigVM/Public/RigVMCore/RigVMByteCode.h
--- a/Engine/Source/Runtime/RigVM/Public/RigVMCore/RigVMByteCode.h
+++ b/Engine/Source/Runtime/RigVM/Public/RigVMCore/RigVMByteCode.h
@@ -115,6 +115,7 @@ enum class ERigVMOpCode : uint8
 	ArrayDifference, // returns a new array containing elements only found in one array (ternary op, in array, in array, out result)
 	ArrayIntersection, // returns a new array containing elements found in both of the input arrays (ternary op, in array, in array, out result)
 	ArrayReverse, // returns the reverse of the input array (unary op, in out array)
+	InvokeEntry, // invokes an entry from the entry list
 	Invalid,
 	FirstArrayOpCode = ArrayReset,
 	LastArrayOpCode = ArrayReverse,
@@ -701,6 +702,37 @@ struct RIGVM_API FRigVMChangeTypeOp : public FRigVMUnaryOp
 	}
 };
 
+// invoke another entry
+USTRUCT()
+struct RIGVM_API FRigVMInvokeEntryOp : public FRigVMBaseOp
+{
+	GENERATED_USTRUCT_BODY()
+
+	FRigVMInvokeEntryOp()
+		: FRigVMBaseOp(ERigVMOpCode::InvokeEntry)
+		, EntryName(NAME_None)
+	{}
+
+	FRigVMInvokeEntryOp(FName InEntryName)
+		: FRigVMBaseOp(ERigVMOpCode::InvokeEntry)
+		, EntryName(InEntryName)
+	{}
+
+	FName EntryName;
+
+	friend FORCEINLINE uint32 GetTypeHash(const FRigVMInvokeEntryOp& Op)
+	{
+		return HashCombine(GetTypeHash((const FRigVMBaseOp&)Op), GetTypeHash(Op.EntryName));
+	}
+
+	void Serialize(FArchive& Ar);
+	FORCEINLINE friend FArchive& operator<<(FArchive& Ar, FRigVMInvokeEntryOp& P)
+	{
+		P.Serialize(Ar);
+		return Ar;
+	}
+};
+
 /**
  * The FRigVMInstruction represents
  * a single instruction within the VM.
@@ -937,6 +969,9 @@ public:
 	// adds an array reverse operator
 	uint64 AddArrayReverseOp(FRigVMOperand InArrayArg);
 
+	// adds an invoke entry operator
+	uint64 AddInvokeEntryOp(const FName& InEntryName);
+
 	// returns an instruction array for iterating over all operators
 	FORCEINLINE FRigVMInstructionArray GetInstructions() const
 	{
@@ -1150,5 +1185,6 @@ private:
 
 	friend class URigVMCompiler;
 	friend class FRigVMByteCodeTest;
+	friend class FRigVMInvokeEntryTest;
 	friend class URigVM;
 };
