ControlRig: Reverted a previous change that parented rigvm memory objects to the ControlRigBP asset object because the engine no longer supports uobject as uclass outer, only upackage can be the outer of an uclass.

Also added UControlRigBlueprint::GetPreloadDependencies to ensure memory classes are loaded before ControlRigBP during cooking.

diff --git a/Engine/Plugins/Animation/ControlRig/Source/ControlRigDeveloper/Private/ControlRigBlueprint.cpp b/Engine/Plugins/Animation/ControlRig/Source/ControlRigDeveloper/Private/ControlRigBlueprint.cpp
--- a/Engine/Plugins/Animation/ControlRig/Source/ControlRigDeveloper/Private/ControlRigBlueprint.cpp
+++ b/Engine/Plugins/Animation/ControlRig/Source/ControlRigDeveloper/Private/ControlRigBlueprint.cpp
@@ -50,6 +50,23 @@
 
 #define LOCTEXT_NAMESPACE "ControlRigBlueprint"
 
+static TArray<UClass*> GetClassObjectsInPackage(UPackage* InPackage)
+{
+	TArray<UObject*> Objects;
+	GetObjectsWithOuter(InPackage, Objects, false);
+
+	TArray<UClass*> ClassObjects;
+	for (UObject* Object : Objects)
+	{
+		if (UClass* Class = Cast<UClass>(Object))
+		{
+			ClassObjects.Add(Class);
+		}
+	}
+
+	return ClassObjects;
+}
+
 FEdGraphPinType FControlRigPublicFunctionArg::GetPinType() const
 {
 	FRigVMExternalVariable Variable;
@@ -248,6 +265,42 @@ void UControlRigBlueprint::PostEditChangeChainProperty(FPropertyChangedChainEven
 	PostEditChangeChainPropertyEvent.Broadcast(PropertyChangedEvent);
 }
 
+void UControlRigBlueprint::PostRename(UObject* OldOuter, const FName OldName)
+{
+	Super::PostRename(OldOuter, OldName);
+
+	// Whenever the asset is renamed/moved, generated classes parented to the old package
+	// are not moved to the new package automatically (see FAssetRenameManager), so we
+	// have to manually perform the move/rename, to avoid invalid reference to the old package
+	
+	// Note: while asset duplication doesn't duplicate the classes either, it is not a problem there
+	// because we always recompile in post duplicate.
+	TArray<UClass*> ClassObjects = GetClassObjectsInPackage(OldOuter->GetPackage());
+
+	for (UClass* ClassObject : ClassObjects)
+	{
+		if (URigVMMemoryStorageGeneratorClass* MemoryClass = Cast<URigVMMemoryStorageGeneratorClass>(ClassObject))
+		{
+			MemoryClass->Rename(nullptr, GetPackage(), REN_ForceNoResetLoaders | REN_DoNotDirty | REN_DontCreateRedirectors | REN_NonTransactional);
+		}
+	}
+}
+
+void UControlRigBlueprint::GetPreloadDependencies(TArray<UObject*>& OutDeps)
+{
+	Super::GetPreloadDependencies(OutDeps);
+
+	TArray<UClass*> ClassObjects = GetClassObjectsInPackage(GetPackage());
+
+	for (UClass* ClassObject : ClassObjects)
+	{
+		if (URigVMMemoryStorageGeneratorClass* MemoryClass = Cast<URigVMMemoryStorageGeneratorClass>(ClassObject))
+		{
+			OutDeps.Add(MemoryClass);
+		}
+	}
+}
+
 FRigVMClient* UControlRigBlueprint::GetRigVMClient()
 {
 	return &RigVMClient;
@@ -1098,7 +1151,7 @@ void UControlRigBlueprint::BroadcastControlRigPackageDone()
 void UControlRigBlueprint::RemoveDeprecatedVMMemoryClass() const
 {
 	TArray<UObject*> Objects;
-	GetObjectsWithPackage(GetPackage(), Objects, false);
+	GetObjectsWithOuter(this, Objects, false);
 
 	for (UObject* Object : Objects)
 	{
diff --git a/Engine/Plugins/Animation/ControlRig/Source/ControlRigDeveloper/Public/ControlRigBlueprint.h b/Engine/Plugins/Animation/ControlRig/Source/ControlRigDeveloper/Public/ControlRigBlueprint.h
--- a/Engine/Plugins/Animation/ControlRig/Source/ControlRigDeveloper/Public/ControlRigBlueprint.h
+++ b/Engine/Plugins/Animation/ControlRig/Source/ControlRigDeveloper/Public/ControlRigBlueprint.h
@@ -248,6 +248,9 @@ public:
 
 	// UObject interface
 	virtual void PostEditChangeChainProperty(struct FPropertyChangedChainEvent& PropertyChangedEvent) override;
+	virtual void PostRename(UObject* OldOuter, const FName OldName) override;
+	/** Called during cooking. Must return all objects that will be Preload()ed when this is serialized at load time. */
+	void GetPreloadDependencies(TArray<UObject*>& OutDeps) override;
 
 	// IRigVMClientHost interface
 	virtual FRigVMClient* GetRigVMClient() override;
@@ -617,10 +620,12 @@ private:
 	// ControlRigBP, once end-loaded, will inform other ControlRig-Dependent systems that ControlRig instances are ready.
 	void BroadcastControlRigPackageDone();
 
-	// Previously some memory classes were parented to the package
-	// which violated the "single asset per package" rule
+	// Previously some memory classes were parented to the asset object
+	// however it is no longer supported since classes are now identified 
+	// with only package name + class name, see FTopLevelAssetPath
 	// this function removes those deprecated class.
-	// new classes should be created by RecompileVM and parented to the asset object
+	// new classes should be created by RecompileVM and parented to the Package
+	// during PostLoad
 	void RemoveDeprecatedVMMemoryClass() const;
 #endif
 
diff --git a/Engine/Source/Runtime/RigVM/Private/RigVMCore/RigVMMemoryStorage.cpp b/Engine/Source/Runtime/RigVM/Private/RigVMCore/RigVMMemoryStorage.cpp
--- a/Engine/Source/Runtime/RigVM/Private/RigVMCore/RigVMMemoryStorage.cpp
+++ b/Engine/Source/Runtime/RigVM/Private/RigVMCore/RigVMMemoryStorage.cpp
@@ -423,15 +423,9 @@ URigVMMemoryStorageGeneratorClass* URigVMMemoryStorageGeneratorClass::GetStorage
 	for(UObject* PotentialClassContainer : PotentialClassContainers)
 	{
 		UPackage* Package = PotentialClassContainer->GetOutermost();
-
-		UObject* ClassOuter = GetGeneratorClassOuter(Package);
-
-		if (ClassOuter)
+		if(URigVMMemoryStorageGeneratorClass* Class = FindObject<URigVMMemoryStorageGeneratorClass>(Package, *ClassName))
 		{
-			if(URigVMMemoryStorageGeneratorClass* Class = FindObject<URigVMMemoryStorageGeneratorClass>(ClassOuter, *ClassName))
-			{
-				return Class;
-			}
+			return Class;
 		}
 	}
 
@@ -446,21 +440,12 @@ URigVMMemoryStorageGeneratorClass* URigVMMemoryStorageGeneratorClass::CreateStor
 {
 	check(InOuter);
 	UPackage* Package = InOuter->GetOutermost();
-
-	// This class should be parented to the asset object, instead of the package
-	// because the engine no longer supports multiple 'assets' per package
-	UObject* ClassOuter = GetGeneratorClassOuter(Package);
-	if (!ensure(ClassOuter))
-	{
-		return nullptr;
-	}
-	
 	UClass *SuperClass = URigVMMemoryStorage::StaticClass();
 
 	const FString& ClassName = GetClassName(InMemoryType);
 
-	// if there's an old class - remove it from the asset and mark it to destroy
-	URigVMMemoryStorageGeneratorClass* OldClass = FindObject<URigVMMemoryStorageGeneratorClass>(ClassOuter, *ClassName);
+	// if there's an old class - remove it from the package and mark it to destroy
+	URigVMMemoryStorageGeneratorClass* OldClass = FindObject<URigVMMemoryStorageGeneratorClass>(Package, *ClassName);
 	if(OldClass)
 	{
 		FString DiscardedMemoryClassName;
@@ -481,7 +466,7 @@ URigVMMemoryStorageGeneratorClass* URigVMMemoryStorageGeneratorClass::CreateStor
 
 	// create the new class
 	URigVMMemoryStorageGeneratorClass* Class = NewObject<URigVMMemoryStorageGeneratorClass>(
-		ClassOuter,
+		Package,
 		*ClassName,
 		RF_Standalone
 	);
@@ -545,15 +530,11 @@ bool URigVMMemoryStorageGeneratorClass::RemoveStorageClass(UObject* InOuter, ERi
 {
 	check(InOuter);
 	UPackage* Package = InOuter->GetOutermost();
-	
-	// This class should have been parented to the asset object, instead of the package
-	// because the engine no longer supports multiple 'assets' per package
-	UObject* ClassOuter = GetGeneratorClassOuter(Package);
-	
+
 	const FString ClassName = GetClassName(InMemoryType);
 
-	// if there's an old class - remove it from the asset and mark it to destroy
-	URigVMMemoryStorageGeneratorClass* OldClass = FindObject<URigVMMemoryStorageGeneratorClass>(ClassOuter, *ClassName);
+	// if there's an old class - remove it from the package and mark it to destroy
+	URigVMMemoryStorageGeneratorClass* OldClass = FindObject<URigVMMemoryStorageGeneratorClass>(Package, *ClassName);
 	if(OldClass)
 	{
 		OldClass->Rename(nullptr, GetTransientPackage(), REN_ForceNoResetLoaders | REN_DoNotDirty | REN_DontCreateRedirectors | REN_NonTransactional);
diff --git a/Engine/Source/Runtime/RigVM/Public/RigVMCore/RigVMMemoryStorage.h b/Engine/Source/Runtime/RigVM/Public/RigVMCore/RigVMMemoryStorage.h
--- a/Engine/Source/Runtime/RigVM/Public/RigVMCore/RigVMMemoryStorage.h
+++ b/Engine/Source/Runtime/RigVM/Public/RigVMCore/RigVMMemoryStorage.h
@@ -209,8 +209,16 @@ class RIGVM_API URigVMMemoryStorageGeneratorClass :
 	public UClass
 {
 	GENERATED_BODY()
-	DECLARE_WITHIN(UObject);
-	
+
+	// DECLARE_WITHIN(UObject) is only kept for back-compat, please don't parent the class
+	// to the asset object.
+	// This class should be parented to the package, instead of the asset object
+	// because the engine no longer supports asset object as UClass outer
+	// Given in the past we have parented the class to the asset object,
+	// this flag has to be kept such that we can load the old asset in the first place and
+	// re-parent it back to the package in post load
+	DECLARE_WITHIN(UObject)
+
 public:
 
 	// Default constructor
