Control Rig: Change FRigVMTemplate to USTRUCT

diff --git a/Engine/Plugins/Animation/ControlRig/Source/ControlRigEditor/Private/Graph/NodeSpawners/ControlRigTemplateNodeSpawner.cpp b/Engine/Plugins/Animation/ControlRig/Source/ControlRigEditor/Private/Graph/NodeSpawners/ControlRigTemplateNodeSpawner.cpp
--- a/Engine/Plugins/Animation/ControlRig/Source/ControlRigEditor/Private/Graph/NodeSpawners/ControlRigTemplateNodeSpawner.cpp
+++ b/Engine/Plugins/Animation/ControlRig/Source/ControlRigEditor/Private/Graph/NodeSpawners/ControlRigTemplateNodeSpawner.cpp
@@ -96,7 +96,7 @@ bool UControlRigTemplateNodeSpawner::IsTemplateNodeFilteredOut(FBlueprintActionF
 							{
 								if (const FRigVMTemplate* Template = FRigVMRegistry::Get().FindTemplate(TemplateNotation))
 								{
-									FRigVMTemplateArgument::FType Type(ModelPin->GetCPPType(), ModelPin->GetCPPTypeObject());
+									FRigVMTemplateArgumentType Type(ModelPin->GetCPPType(), ModelPin->GetCPPTypeObject());
 									for (int32 i=0; i<Template->NumArguments(); ++i)
 									{
 										const FRigVMTemplateArgument* Argument = Template->GetArgument(i);
diff --git a/Engine/Plugins/Animation/ControlRig/Source/ControlRigEditor/Private/Graph/SControlRigGraphChangePinType.cpp b/Engine/Plugins/Animation/ControlRig/Source/ControlRigEditor/Private/Graph/SControlRigGraphChangePinType.cpp
--- a/Engine/Plugins/Animation/ControlRig/Source/ControlRigEditor/Private/Graph/SControlRigGraphChangePinType.cpp
+++ b/Engine/Plugins/Animation/ControlRig/Source/ControlRigEditor/Private/Graph/SControlRigGraphChangePinType.cpp
@@ -49,7 +49,7 @@ void SControlRigChangePinType::Construct(const FArguments& InArgs)
 	];
 }
 
-FText SControlRigChangePinType::GetBindingText(const FRigVMTemplateArgument::FType& InType)
+FText SControlRigChangePinType::GetBindingText(const FRigVMTemplateArgumentType& InType)
 {
 	if(UObject* CPPTypeObject = InType.CPPTypeObject)
 	{
@@ -147,7 +147,7 @@ FLinearColor SControlRigChangePinType::GetBindingColor() const
 {
 	if(ModelPins.Num() > 0)
 	{
-		const FRigVMTemplateArgument::FType Type = ModelPins[0]->GetTemplateArgumentType();
+		const FRigVMTemplateArgumentType Type = ModelPins[0]->GetTemplateArgumentType();
 		const FEdGraphPinType PinType = RigVMTypeUtils::PinTypeFromCPPType(Type.CPPType, Type.CPPTypeObject);
 		const UControlRigGraphSchema* Schema = GetDefault<UControlRigGraphSchema>();
 		return Schema->GetPinTypeColor(PinType);
@@ -178,7 +178,7 @@ void SControlRigChangePinType::FillPinTypeMenu(FMenuBuilder& MenuBuilder)
 			return Info;
 		}
 
-		static FArgumentInfo Make(const FRigVMTemplateArgument::FType& InType, const bool bInIsFilteredOut)
+		static FArgumentInfo Make(const FRigVMTemplateArgumentType& InType, const bool bInIsFilteredOut)
 		{
 			FArgumentInfo Info;
 			Info.Property = nullptr;
@@ -189,11 +189,11 @@ void SControlRigChangePinType::FillPinTypeMenu(FMenuBuilder& MenuBuilder)
 
 		const FProperty* Property;
 		bool bIsFilteredOut;
-		FRigVMTemplateArgument::FType Type;
+		FRigVMTemplateArgumentType Type;
 	};
 
-	typedef TPair<FRigVMTemplateArgument::FType, FArgumentInfo> FTypePair;
-	TMap<FRigVMTemplateArgument::FType, FArgumentInfo> Types;
+	typedef TPair<FRigVMTemplateArgumentType, FArgumentInfo> FTypePair;
+	TMap<FRigVMTemplateArgumentType, FArgumentInfo> Types;
 	for(URigVMPin* ModelPin : ModelPins)
 	{
 		if(!ModelPin->IsRootPin())
@@ -208,12 +208,12 @@ void SControlRigChangePinType::FillPinTypeMenu(FMenuBuilder& MenuBuilder)
 				const FName ArgumentName = ModelPin->GetFName();
 				if(const FRigVMTemplateArgument* Argument = Template->FindArgument(ArgumentName))
 				{
-					const TArray<FRigVMTemplateArgument::FType>& AllArgumentTypes = Argument->GetTypes();
-					const TArray<FRigVMTemplateArgument::FType>& FilteredArgumentTypes = Argument->GetSupportedTypes(TemplateNode->GetFilteredPermutationsIndices());
+					const TArray<FRigVMTemplateArgumentType>& AllArgumentTypes = Argument->GetTypes();
+					const TArray<FRigVMTemplateArgumentType>& FilteredArgumentTypes = Argument->GetSupportedTypes(TemplateNode->GetFilteredPermutationsIndices());
 					
 					for(int32 PermutationIndex = 0; PermutationIndex < Template->NumPermutations(); PermutationIndex++)
 					{
-						FRigVMTemplateArgument::FType ArgumentType = AllArgumentTypes[PermutationIndex];
+						FRigVMTemplateArgumentType ArgumentType = AllArgumentTypes[PermutationIndex];
 					
 						bool bIsFilteredOut = false;
 						if(!FilteredArgumentTypes.Contains(ArgumentType))
@@ -305,7 +305,7 @@ void SControlRigChangePinType::FillPinTypeMenu(FMenuBuilder& MenuBuilder)
 		
 		for(int32 TypeIndex=0; TypeIndex < SortedTypes.Num(); TypeIndex++)
 		{
-			const FRigVMTemplateArgument::FType& Type = SortedTypes[TypeIndex].Key;
+			const FRigVMTemplateArgumentType& Type = SortedTypes[TypeIndex].Key;
 			const bool bIsFilteredOut = SortedTypes[TypeIndex].Value.bIsFilteredOut;
 			if (Type.CPPTypeObject != nullptr && !IsValid(Type.CPPTypeObject))
 			{
@@ -324,7 +324,7 @@ void SControlRigChangePinType::FillPinTypeMenu(FMenuBuilder& MenuBuilder)
 				if(UScriptStruct* ScriptStruct = Cast<UScriptStruct>(Type.CPPTypeObject))
 				{
 					if(!FRigVMTemplateArgument::GetCompatibleTypes(FRigVMTemplateArgument::ETypeCategory_SingleMathStructValue).Contains(
-						FRigVMTemplateArgument::FType(ScriptStruct->GetStructCPPName(), ScriptStruct)))
+						FRigVMTemplateArgumentType(ScriptStruct->GetStructCPPName(), ScriptStruct)))
 					{
 						continue;
 					}
@@ -373,7 +373,7 @@ void SControlRigChangePinType::FillPinTypeMenu(FMenuBuilder& MenuBuilder)
 	MenuBuilder.EndSection(); // Local Variables
 }
 
-void SControlRigChangePinType::HandlePinTypeChanged(FRigVMTemplateArgument::FType InType)
+void SControlRigChangePinType::HandlePinTypeChanged(FRigVMTemplateArgumentType InType)
 {
 	if(ModelPins.IsEmpty() || (Blueprint == nullptr))
 	{
diff --git a/Engine/Plugins/Animation/ControlRig/Source/ControlRigEditor/Private/Graph/SControlRigGraphChangePinType.h b/Engine/Plugins/Animation/ControlRig/Source/ControlRigEditor/Private/Graph/SControlRigGraphChangePinType.h
--- a/Engine/Plugins/Animation/ControlRig/Source/ControlRigEditor/Private/Graph/SControlRigGraphChangePinType.h
+++ b/Engine/Plugins/Animation/ControlRig/Source/ControlRigEditor/Private/Graph/SControlRigGraphChangePinType.h
@@ -28,7 +28,7 @@ public:
 
 protected:
 
-	static FText GetBindingText(const FRigVMTemplateArgument::FType& InType);
+	static FText GetBindingText(const FRigVMTemplateArgumentType& InType);
 	FText GetBindingText(URigVMPin* ModelPin) const;
 	FText GetBindingText() const;
 	const FSlateBrush* GetBindingImage() const;
@@ -37,7 +37,7 @@ protected:
 	bool OnCanBindToClass(UClass* InClass) const;
 	void OnAddBinding(FName InPropertyName, const TArray<FBindingChainElement>& InBindingChain);
 	void FillPinTypeMenu( FMenuBuilder& MenuBuilder );
-	void HandlePinTypeChanged(FRigVMTemplateArgument::FType InType);
+	void HandlePinTypeChanged(FRigVMTemplateArgumentType InType);
 
 	TArray<URigVMPin*> ModelPins;
 	UControlRigBlueprint* Blueprint;
diff --git a/Engine/Source/Developer/RigVMDeveloper/Private/RigVMModel/Nodes/RigVMArrayNode.cpp b/Engine/Source/Developer/RigVMDeveloper/Private/RigVMModel/Nodes/RigVMArrayNode.cpp
--- a/Engine/Source/Developer/RigVMDeveloper/Private/RigVMModel/Nodes/RigVMArrayNode.cpp
+++ b/Engine/Source/Developer/RigVMDeveloper/Private/RigVMModel/Nodes/RigVMArrayNode.cpp
@@ -252,17 +252,17 @@ const FRigVMTemplate* URigVMArrayNode::GetTemplate() const
 	
 	if(CachedTemplate == nullptr)
 	{		
-		const TArray<FRigVMTemplateArgument::FType>& SingleTypes = FRigVMTemplateArgument::GetCompatibleTypes(FRigVMTemplateArgument::ETypeCategory_SingleAnyValue);
-		const TArray<FRigVMTemplateArgument::FType>& ArrayTypes = FRigVMTemplateArgument::GetCompatibleTypes(FRigVMTemplateArgument::ETypeCategory_ArrayAnyValue);
-		const TArray<FRigVMTemplateArgument::FType>& ArrayArrayTypes = FRigVMTemplateArgument::GetCompatibleTypes(FRigVMTemplateArgument::ETypeCategory_ArrayArrayAnyValue);
+		const TArray<FRigVMTemplateArgumentType>& SingleTypes = FRigVMTemplateArgument::GetCompatibleTypes(FRigVMTemplateArgument::ETypeCategory_SingleAnyValue);
+		const TArray<FRigVMTemplateArgumentType>& ArrayTypes = FRigVMTemplateArgument::GetCompatibleTypes(FRigVMTemplateArgument::ETypeCategory_ArrayAnyValue);
+		const TArray<FRigVMTemplateArgumentType>& ArrayArrayTypes = FRigVMTemplateArgument::GetCompatibleTypes(FRigVMTemplateArgument::ETypeCategory_ArrayArrayAnyValue);
 
-		TArray<FRigVMTemplateArgument::FType> ArrayValueTypes = ArrayTypes;
+		TArray<FRigVMTemplateArgumentType> ArrayValueTypes = ArrayTypes;
 		ArrayValueTypes.Append(ArrayArrayTypes);
-		TArray<FRigVMTemplateArgument::FType> ElementTypes = SingleTypes;
+		TArray<FRigVMTemplateArgumentType> ElementTypes = SingleTypes;
 		ElementTypes.Append(ArrayTypes);
 
 		UScriptStruct* ExecuteStruct = GetGraph()->GetExecuteContextStruct();
-		FRigVMTemplateArgument::FType ExecuteType(ExecuteStruct->GetStructCPPName(), ExecuteStruct);
+		FRigVMTemplateArgumentType ExecuteType(ExecuteStruct->GetStructCPPName(), ExecuteStruct);
 		
 		TArray<FRigVMTemplateArgument> Arguments;
 		
@@ -278,20 +278,20 @@ const FRigVMTemplate* URigVMArrayNode::GetTemplate() const
     		case ERigVMOpCode::ArrayGetNum:
     		{
     			Arguments.Emplace(*ArrayName, ERigVMPinDirection::Input, ArrayValueTypes);
-    			Arguments.Emplace(*NumName, ERigVMPinDirection::Output, FRigVMTemplateArgument::FType(RigVMTypeUtils::Int32Type, nullptr));
+    			Arguments.Emplace(*NumName, ERigVMPinDirection::Output, FRigVMTemplateArgumentType(RigVMTypeUtils::Int32Type, nullptr));
     			break;
     		} 
     		case ERigVMOpCode::ArraySetNum:
     		{
     			Arguments.Emplace(*FRigVMStruct::ExecuteContextName.ToString(), ERigVMPinDirection::IO, ExecuteType);
     			Arguments.Emplace(*ArrayName, ERigVMPinDirection::IO, ArrayValueTypes);
-    			Arguments.Emplace(*NumName, ERigVMPinDirection::Input, FRigVMTemplateArgument::FType(RigVMTypeUtils::Int32Type, nullptr));
+    			Arguments.Emplace(*NumName, ERigVMPinDirection::Input, FRigVMTemplateArgumentType(RigVMTypeUtils::Int32Type, nullptr));
     			break;
     		}
     		case ERigVMOpCode::ArrayGetAtIndex:
     		{
     			Arguments.Emplace(*ArrayName, ERigVMPinDirection::Input, ArrayValueTypes);
-    			Arguments.Emplace(*IndexName, ERigVMPinDirection::Input, FRigVMTemplateArgument::FType(RigVMTypeUtils::Int32Type, nullptr));
+    			Arguments.Emplace(*IndexName, ERigVMPinDirection::Input, FRigVMTemplateArgumentType(RigVMTypeUtils::Int32Type, nullptr));
     			Arguments.Emplace(*ElementName, ERigVMPinDirection::Output, ElementTypes);
     			break;
     		}  
@@ -300,7 +300,7 @@ const FRigVMTemplate* URigVMArrayNode::GetTemplate() const
     		{
     			Arguments.Emplace(*FRigVMStruct::ExecuteContextName.ToString(), ERigVMPinDirection::IO, ExecuteType);
     			Arguments.Emplace(*ArrayName, ERigVMPinDirection::IO, ArrayValueTypes);
-    			Arguments.Emplace(*IndexName, ERigVMPinDirection::Input, FRigVMTemplateArgument::FType(RigVMTypeUtils::Int32Type, nullptr));
+    			Arguments.Emplace(*IndexName, ERigVMPinDirection::Input, FRigVMTemplateArgumentType(RigVMTypeUtils::Int32Type, nullptr));
     			Arguments.Emplace(*ElementName, ERigVMPinDirection::Input, ElementTypes);
     			break;
     		}
@@ -309,22 +309,22 @@ const FRigVMTemplate* URigVMArrayNode::GetTemplate() const
     			Arguments.Emplace(*FRigVMStruct::ExecuteContextName.ToString(), ERigVMPinDirection::IO, ExecuteType);
     			Arguments.Emplace(*ArrayName, ERigVMPinDirection::IO, ArrayValueTypes);
     			Arguments.Emplace(*ElementName, ERigVMPinDirection::Input, ElementTypes);
-    			Arguments.Emplace(*IndexName, ERigVMPinDirection::Output, FRigVMTemplateArgument::FType(RigVMTypeUtils::Int32Type, nullptr));
+    			Arguments.Emplace(*IndexName, ERigVMPinDirection::Output, FRigVMTemplateArgumentType(RigVMTypeUtils::Int32Type, nullptr));
     			break;
     		}
     		case ERigVMOpCode::ArrayRemove:
     		{
     			Arguments.Emplace(*FRigVMStruct::ExecuteContextName.ToString(), ERigVMPinDirection::IO, ExecuteType);
     			Arguments.Emplace(*ArrayName, ERigVMPinDirection::IO, ArrayValueTypes);
-    			Arguments.Emplace(*IndexName, ERigVMPinDirection::Input, FRigVMTemplateArgument::FType(RigVMTypeUtils::Int32Type, nullptr));
+    			Arguments.Emplace(*IndexName, ERigVMPinDirection::Input, FRigVMTemplateArgumentType(RigVMTypeUtils::Int32Type, nullptr));
     			break;
     		}
     		case ERigVMOpCode::ArrayFind:
     		{
     			Arguments.Emplace(*ArrayName, ERigVMPinDirection::Input, ArrayValueTypes);
     			Arguments.Emplace(*ElementName, ERigVMPinDirection::Input, ElementTypes);
-    			Arguments.Emplace(*IndexName, ERigVMPinDirection::Output, FRigVMTemplateArgument::FType(RigVMTypeUtils::Int32Type, nullptr));
-    			Arguments.Emplace(*SuccessName, ERigVMPinDirection::Output, FRigVMTemplateArgument::FType(RigVMTypeUtils::BoolType, nullptr));
+    			Arguments.Emplace(*IndexName, ERigVMPinDirection::Output, FRigVMTemplateArgumentType(RigVMTypeUtils::Int32Type, nullptr));
+    			Arguments.Emplace(*SuccessName, ERigVMPinDirection::Output, FRigVMTemplateArgumentType(RigVMTypeUtils::BoolType, nullptr));
     			break;
     		}
     		case ERigVMOpCode::ArrayAppend:
@@ -346,9 +346,9 @@ const FRigVMTemplate* URigVMArrayNode::GetTemplate() const
     			Arguments.Emplace(*FRigVMStruct::ExecuteContextName.ToString(), ERigVMPinDirection::IO, ExecuteType);
     			Arguments.Emplace(*ArrayName, ERigVMPinDirection::Input, ArrayValueTypes);
     			Arguments.Emplace(*ElementName, ERigVMPinDirection::Output, ElementTypes);
-    			Arguments.Emplace(*IndexName, ERigVMPinDirection::Output, FRigVMTemplateArgument::FType(RigVMTypeUtils::Int32Type, nullptr));
-    			Arguments.Emplace(*CountName, ERigVMPinDirection::Output, FRigVMTemplateArgument::FType(RigVMTypeUtils::Int32Type, nullptr));
-    			Arguments.Emplace(*RatioName, ERigVMPinDirection::Output, FRigVMTemplateArgument::FType(RigVMTypeUtils::FloatType, nullptr));
+    			Arguments.Emplace(*IndexName, ERigVMPinDirection::Output, FRigVMTemplateArgumentType(RigVMTypeUtils::Int32Type, nullptr));
+    			Arguments.Emplace(*CountName, ERigVMPinDirection::Output, FRigVMTemplateArgumentType(RigVMTypeUtils::Int32Type, nullptr));
+    			Arguments.Emplace(*RatioName, ERigVMPinDirection::Output, FRigVMTemplateArgumentType(RigVMTypeUtils::FloatType, nullptr));
     			Arguments.Emplace(*ContinueName, ERigVMPinDirection::Hidden, ExecuteType);
     			Arguments.Emplace(*CompletedName, ERigVMPinDirection::Output, ExecuteType);
     			break;
diff --git a/Engine/Source/Developer/RigVMDeveloper/Private/RigVMModel/Nodes/RigVMIfNode.cpp b/Engine/Source/Developer/RigVMDeveloper/Private/RigVMModel/Nodes/RigVMIfNode.cpp
--- a/Engine/Source/Developer/RigVMDeveloper/Private/RigVMModel/Nodes/RigVMIfNode.cpp
+++ b/Engine/Source/Developer/RigVMDeveloper/Private/RigVMModel/Nodes/RigVMIfNode.cpp
@@ -25,13 +25,13 @@ const FRigVMTemplate* URigVMIfNode::GetTemplate() const
 	
 	if(CachedTemplate == nullptr)
 	{
-		TArray<FRigVMTemplateArgument::FType> Types;
+		TArray<FRigVMTemplateArgumentType> Types;
 		Types.Append(FRigVMTemplateArgument::GetCompatibleTypes(FRigVMTemplateArgument::ETypeCategory_SingleAnyValue));
 		Types.Append(FRigVMTemplateArgument::GetCompatibleTypes(FRigVMTemplateArgument::ETypeCategory_ArrayAnyValue));
 		Types.Append(FRigVMTemplateArgument::GetCompatibleTypes(FRigVMTemplateArgument::ETypeCategory_ArrayArrayAnyValue));
 		
 		TArray<FRigVMTemplateArgument> Arguments;
-		Arguments.Emplace(*ConditionName, ERigVMPinDirection::Input, FRigVMTemplateArgument::FType(RigVMTypeUtils::BoolType, nullptr));
+		Arguments.Emplace(*ConditionName, ERigVMPinDirection::Input, FRigVMTemplateArgumentType(RigVMTypeUtils::BoolType, nullptr));
 		Arguments.Emplace(*TrueName, ERigVMPinDirection::Input, Types);
 		Arguments.Emplace(*FalseName, ERigVMPinDirection::Input, Types);
 		Arguments.Emplace(*ResultName, ERigVMPinDirection::Output, Types);
diff --git a/Engine/Source/Developer/RigVMDeveloper/Private/RigVMModel/Nodes/RigVMRerouteNode.cpp b/Engine/Source/Developer/RigVMDeveloper/Private/RigVMModel/Nodes/RigVMRerouteNode.cpp
--- a/Engine/Source/Developer/RigVMDeveloper/Private/RigVMModel/Nodes/RigVMRerouteNode.cpp
+++ b/Engine/Source/Developer/RigVMDeveloper/Private/RigVMModel/Nodes/RigVMRerouteNode.cpp
@@ -80,13 +80,13 @@ const FRigVMTemplate* URigVMRerouteNode::GetTemplate() const
 	
 	if(CachedTemplate == nullptr)
 	{
-		TArray<FRigVMTemplateArgument::FType> Types;
+		TArray<FRigVMTemplateArgumentType> Types;
 		Types.Append(FRigVMTemplateArgument::GetCompatibleTypes(FRigVMTemplateArgument::ETypeCategory_SingleAnyValue));
 		Types.Append(FRigVMTemplateArgument::GetCompatibleTypes(FRigVMTemplateArgument::ETypeCategory_ArrayAnyValue));
 		Types.Append(FRigVMTemplateArgument::GetCompatibleTypes(FRigVMTemplateArgument::ETypeCategory_ArrayArrayAnyValue));
 
 		UScriptStruct* ExecuteStruct = GetGraph()->GetExecuteContextStruct();
-		Types.Add(FRigVMTemplateArgument::FType(ExecuteStruct->GetStructCPPName(), ExecuteStruct));
+		Types.Add(FRigVMTemplateArgumentType(ExecuteStruct->GetStructCPPName(), ExecuteStruct));
 		
 		TArray<FRigVMTemplateArgument> Arguments;
 		Arguments.Emplace(TEXT("Value"), ERigVMPinDirection::IO, Types);
diff --git a/Engine/Source/Developer/RigVMDeveloper/Private/RigVMModel/Nodes/RigVMSelectNode.cpp b/Engine/Source/Developer/RigVMDeveloper/Private/RigVMModel/Nodes/RigVMSelectNode.cpp
--- a/Engine/Source/Developer/RigVMDeveloper/Private/RigVMModel/Nodes/RigVMSelectNode.cpp
+++ b/Engine/Source/Developer/RigVMDeveloper/Private/RigVMModel/Nodes/RigVMSelectNode.cpp
@@ -36,17 +36,17 @@ const FRigVMTemplate* URigVMSelectNode::GetTemplate() const
 	
 	if(CachedTemplate == nullptr)
 	{
-		const TArray<FRigVMTemplateArgument::FType>& SingleTypes = FRigVMTemplateArgument::GetCompatibleTypes(FRigVMTemplateArgument::ETypeCategory_SingleAnyValue);
-		const TArray<FRigVMTemplateArgument::FType>& ArrayTypes = FRigVMTemplateArgument::GetCompatibleTypes(FRigVMTemplateArgument::ETypeCategory_ArrayAnyValue);
-		const TArray<FRigVMTemplateArgument::FType>& ArrayArrayTypes = FRigVMTemplateArgument::GetCompatibleTypes(FRigVMTemplateArgument::ETypeCategory_ArrayArrayAnyValue);
+		const TArray<FRigVMTemplateArgumentType>& SingleTypes = FRigVMTemplateArgument::GetCompatibleTypes(FRigVMTemplateArgument::ETypeCategory_SingleAnyValue);
+		const TArray<FRigVMTemplateArgumentType>& ArrayTypes = FRigVMTemplateArgument::GetCompatibleTypes(FRigVMTemplateArgument::ETypeCategory_ArrayAnyValue);
+		const TArray<FRigVMTemplateArgumentType>& ArrayArrayTypes = FRigVMTemplateArgument::GetCompatibleTypes(FRigVMTemplateArgument::ETypeCategory_ArrayArrayAnyValue);
 
-		TArray<FRigVMTemplateArgument::FType> ResultTypes = SingleTypes;
+		TArray<FRigVMTemplateArgumentType> ResultTypes = SingleTypes;
 		ResultTypes.Append(ArrayTypes);
-		TArray<FRigVMTemplateArgument::FType> ValueTypes = ArrayTypes;
+		TArray<FRigVMTemplateArgumentType> ValueTypes = ArrayTypes;
 		ValueTypes.Append(ArrayArrayTypes);		
 		
 		TArray<FRigVMTemplateArgument> Arguments;
-		Arguments.Emplace(*IndexName, ERigVMPinDirection::Input, FRigVMTemplateArgument::FType(RigVMTypeUtils::Int32Type, nullptr));
+		Arguments.Emplace(*IndexName, ERigVMPinDirection::Input, FRigVMTemplateArgumentType(RigVMTypeUtils::Int32Type, nullptr));
 		Arguments.Emplace(*ValueName, ERigVMPinDirection::Input, ValueTypes);
 		Arguments.Emplace(*ResultName, ERigVMPinDirection::Output, ResultTypes);
 		
diff --git a/Engine/Source/Developer/RigVMDeveloper/Private/RigVMModel/Nodes/RigVMTemplateNode.cpp b/Engine/Source/Developer/RigVMDeveloper/Private/RigVMModel/Nodes/RigVMTemplateNode.cpp
--- a/Engine/Source/Developer/RigVMDeveloper/Private/RigVMModel/Nodes/RigVMTemplateNode.cpp
+++ b/Engine/Source/Developer/RigVMDeveloper/Private/RigVMModel/Nodes/RigVMTemplateNode.cpp
@@ -190,7 +190,7 @@ bool URigVMTemplateNode::SupportsType(const URigVMPin* InPin, const FString& InC
 	if (const FRigVMTemplate* Template = GetTemplate())
 	{
 		const FString CacheKey = RootPin->GetName() + TEXT("|") + CPPType;
-		if (const TPair<bool, FRigVMTemplateArgument::FType>* CachedResult = SupportedTypesCache.Find(CacheKey))
+		if (const TPair<bool, FRigVMTemplateArgumentType>* CachedResult = SupportedTypesCache.Find(CacheKey))
 		{
 			if(OutCPPType)
 			{
@@ -199,7 +199,7 @@ bool URigVMTemplateNode::SupportsType(const URigVMPin* InPin, const FString& InC
 			return CachedResult->Key;
 		}
 
-		FRigVMTemplateArgument::FType OutType;
+		FRigVMTemplateArgumentType OutType;
 		if (Template->ArgumentSupportsType(RootPin->GetFName(), CPPType, &OutType))
 		{
 			if (OutCPPType)
@@ -267,10 +267,10 @@ bool URigVMTemplateNode::FilteredSupportsType(const URigVMPin* InPin, const FStr
 		return GetTemplate()->ArgumentSupportsType(RootPin->GetFName(), InCPPType);
 	}
 
-	const TArray<FRigVMTemplateArgument::FType>& Types = Argument->GetTypes();
+	const TArray<FRigVMTemplateArgumentType>& Types = Argument->GetTypes();
 	for (int32 PermutationIndex : FilteredPermutations)
 	{
-		const FRigVMTemplateArgument::FType& FilteredType = Types[PermutationIndex];
+		const FRigVMTemplateArgumentType& FilteredType = Types[PermutationIndex];
 		if (FilteredType.Matches(RootCPPType, bAllowFloatingPointCasts))
 		{
 			return true;
@@ -365,7 +365,7 @@ FString URigVMTemplateNode::GetInitialDefaultValueForPin(const FName& InRootPinN
 		{
 			const FRigVMTemplate* Template = GetTemplate();
 			const FRigVMTemplateArgument* Argument = Template->FindArgument(InRootPinName);
-			const FRigVMTemplateArgument::FType Type = Argument->GetTypes()[PermutationIndex];
+			const FRigVMTemplateArgumentType Type = Argument->GetTypes()[PermutationIndex];
 
 			if (Type.IsArray())
 			{
@@ -470,10 +470,10 @@ const TArray<int32>& URigVMTemplateNode::GetFilteredPermutationsIndices() const
 	return FilteredPermutations;
 }
 
-TArray<FRigVMTemplateArgument::FType> URigVMTemplateNode::GetFilteredTypesForPin(URigVMPin* InPin) const
+TArray<FRigVMTemplateArgumentType> URigVMTemplateNode::GetFilteredTypesForPin(URigVMPin* InPin) const
 {
 	ensureMsgf(InPin->GetNode() == this, TEXT("GetFilteredTypesForPin of %s with pin from another node %s"), *GetNodePath(), *InPin->GetPinPath(true));
-	TArray<FRigVMTemplateArgument::FType> FilteredTypes;
+	TArray<FRigVMTemplateArgumentType> FilteredTypes;
 
 	if (FilteredPermutations.IsEmpty())
 	{
@@ -489,7 +489,7 @@ TArray<FRigVMTemplateArgument::FType> URigVMTemplateNode::GetFilteredTypesForPin
 			if (InPin->GetName() == ArgName)
 			{
 				const FRigVMTemplateArgument* Argument = GetTemplate()->FindArgument(*ArgName);
-				for (const FRigVMTemplateArgument::FType& Type : Argument->GetTypes())
+				for (const FRigVMTemplateArgumentType& Type : Argument->GetTypes())
 				{
 					if (Type.CPPType == CPPType)
 					{
@@ -512,7 +512,7 @@ TArray<FRigVMTemplateArgument::FType> URigVMTemplateNode::GetFilteredTypesForPin
 		FilteredTypes = Argument->GetSupportedTypes(FilteredPermutations);
 		if (bIsArrayElement)
 		{
-			for (FRigVMTemplateArgument::FType& ArrayType : FilteredTypes)
+			for (FRigVMTemplateArgumentType& ArrayType : FilteredTypes)
 			{
 				ArrayType.ConvertToBaseElement();
 			}
@@ -565,7 +565,7 @@ TArray<int32> URigVMTemplateNode::GetNewFilteredPermutations(URigVMPin* InPin, U
 			{
 				for(int32 PermutationIndex : PermutationsToTry)
 				{
-					FRigVMTemplateArgument::FType Type = Argument->GetTypes()[PermutationIndex];
+					FRigVMTemplateArgumentType Type = Argument->GetTypes()[PermutationIndex];
 					if (bIsArrayElement)
 					{
 						Type.ConvertToBaseElement();
@@ -586,7 +586,7 @@ TArray<int32> URigVMTemplateNode::GetNewFilteredPermutations(URigVMPin* InPin, U
 			const FString LinkedCPPType = LinkedPin->GetCPPType();
 			for(int32 PermIndex : PermutationsToTry)
 			{
-				FRigVMTemplateArgument::FType Type = Argument->GetTypes()[PermIndex];
+				FRigVMTemplateArgumentType Type = Argument->GetTypes()[PermIndex];
 				if (bIsArrayElement)
 				{
 					Type.ConvertToBaseElement();
@@ -601,17 +601,17 @@ TArray<int32> URigVMTemplateNode::GetNewFilteredPermutations(URigVMPin* InPin, U
 	return NewFilteredPermutations;
 }
 
-TArray<int32> URigVMTemplateNode::GetNewFilteredPermutations(URigVMPin* InPin, const TArray<FRigVMTemplateArgument::FType>& InTypes)
+TArray<int32> URigVMTemplateNode::GetNewFilteredPermutations(URigVMPin* InPin, const TArray<FRigVMTemplateArgumentType>& InTypes)
 {
 	TArray<int32> NewFilteredPermutations;
 	NewFilteredPermutations.Reserve(FilteredPermutations.Num());
 
 	URigVMPin* RootPin = InPin;
-	TArray<FRigVMTemplateArgument::FType> RootTypes = InTypes;
+	TArray<FRigVMTemplateArgumentType> RootTypes = InTypes;
 	if (URigVMPin* ParentPin = InPin->GetParentPin())
 	{
 		RootPin = ParentPin;
-		for (FRigVMTemplateArgument::FType& Type : RootTypes)
+		for (FRigVMTemplateArgumentType& Type : RootTypes)
 		{
 			Type.ConvertToArray();
 		}
@@ -625,10 +625,10 @@ TArray<int32> URigVMTemplateNode::GetNewFilteredPermutations(URigVMPin* InPin, c
 	
 	if (const FRigVMTemplateArgument* Argument = GetTemplate()->FindArgument(RootPin->GetFName()))
 	{
-		const TArray<FRigVMTemplateArgument::FType>& Types = Argument->GetTypes();
+		const TArray<FRigVMTemplateArgumentType>& Types = Argument->GetTypes();
 		for(int32 PermIndex : PermutationsToTry)
 		{
-			for (FRigVMTemplateArgument::FType& RootType : RootTypes)
+			for (FRigVMTemplateArgumentType& RootType : RootTypes)
 			{				
 				if (Types[PermIndex].Matches(RootType.CPPType))
 				{
@@ -715,7 +715,7 @@ TArray<FString> URigVMTemplateNode::GetArgumentTypesForPermutation(const int32 I
 	return ArgTypes;
 }
 
-bool URigVMTemplateNode::PinNeedsFilteredTypesUpdate(URigVMPin* InPin, const TArray<FRigVMTemplateArgument::FType>& InTypes)
+bool URigVMTemplateNode::PinNeedsFilteredTypesUpdate(URigVMPin* InPin, const TArray<FRigVMTemplateArgumentType>& InTypes)
 {
 	TArray<int32> NewFilteredPermutations = GetNewFilteredPermutations(InPin, InTypes);
 	if (NewFilteredPermutations.Num() == FilteredPermutations.Num())
@@ -757,7 +757,7 @@ bool URigVMTemplateNode::UpdateFilteredPermutations(URigVMPin* InPin, URigVMPin*
 	return true;
 }
 
-bool URigVMTemplateNode::UpdateFilteredPermutations(URigVMPin* InPin, const TArray<FRigVMTemplateArgument::FType>& InTypes)
+bool URigVMTemplateNode::UpdateFilteredPermutations(URigVMPin* InPin, const TArray<FRigVMTemplateArgumentType>& InTypes)
 {
 	TArray<int32> NewFilteredPermutations = GetNewFilteredPermutations(InPin, InTypes);
 	if (NewFilteredPermutations.Num() == FilteredPermutations.Num())
diff --git a/Engine/Source/Developer/RigVMDeveloper/Private/RigVMModel/RigVMController.cpp b/Engine/Source/Developer/RigVMDeveloper/Private/RigVMModel/RigVMController.cpp
--- a/Engine/Source/Developer/RigVMDeveloper/Private/RigVMModel/RigVMController.cpp
+++ b/Engine/Source/Developer/RigVMDeveloper/Private/RigVMModel/RigVMController.cpp
@@ -473,7 +473,7 @@ TArray<FString> URigVMController::GetAddNodePythonCommands(URigVMNode* Node) con
 									*Pin->GetCPPType(),
 									*Pin->GetCPPTypeObject()->GetPathName()));
 
-							TypeMap.Add(Argument->GetName(), FRigVMTemplateArgument::FType(Pin->GetCPPType(), Pin->GetCPPTypeObject()));
+							TypeMap.Add(Argument->GetName(), FRigVMTemplateArgumentType(Pin->GetCPPType(), Pin->GetCPPTypeObject()));
 							Template->Resolve(TypeMap, Permutations, false);
 						}
 					}
@@ -1816,7 +1816,7 @@ bool URigVMController::UnresolveTemplateNodes(const TArray<URigVMTemplateNode*>&
 	TArray<TPair<FString, FString>> LinkedPaths = GetLinkedPinPaths(Nodes);
 
 	// Find pins outside our set of nodes which were resolved to a type
-	TMap<URigVMPin*, FRigVMTemplateArgument::FType> PinsToResolve;
+	TMap<URigVMPin*, FRigVMTemplateArgumentType> PinsToResolve;
 	for (TPair<FString, FString>& LinkPath : LinkedPaths)
 	{
 		const FString& SourcePath = LinkPath.Key;
@@ -1861,14 +1861,14 @@ bool URigVMController::UnresolveTemplateNodes(const TArray<URigVMTemplateNode*>&
 		{
 			if (!SourcePin->IsWildCard())
 			{
-				PinsToResolve.Add(SourcePin, FRigVMTemplateArgument::FType(SourcePin->GetCPPType(), SourcePin->GetCPPTypeObject()));
+				PinsToResolve.Add(SourcePin, FRigVMTemplateArgumentType(SourcePin->GetCPPType(), SourcePin->GetCPPTypeObject()));
 			}
 		}
 		else if(bTargetOutside && !bSourceOutside)
 		{
 			if (!TargetPin->IsWildCard())
 			{
-				PinsToResolve.Add(TargetPin, FRigVMTemplateArgument::FType(TargetPin->GetCPPType(), TargetPin->GetCPPTypeObject()));
+				PinsToResolve.Add(TargetPin, FRigVMTemplateArgumentType(TargetPin->GetCPPType(), TargetPin->GetCPPTypeObject()));
 			}
 		}			
 	}
@@ -1915,7 +1915,7 @@ bool URigVMController::UnresolveTemplateNodes(const TArray<URigVMTemplateNode*>&
 		// if the types' containers don't match the current pin containers try again
 		bool bTypesWereAdjusted = false;
 		FRigVMTemplate::FTypeMap AdjustedTypes = Types;
-		for(TPair<FName, FRigVMTemplateArgument::FType>& TypePair : AdjustedTypes)
+		for(TPair<FName, FRigVMTemplateArgumentType>& TypePair : AdjustedTypes)
 		{
 			if(URigVMPin* Pin = Node->FindPin(TypePair.Key.ToString()))
 			{
@@ -1943,10 +1943,10 @@ bool URigVMController::UnresolveTemplateNodes(const TArray<URigVMTemplateNode*>&
 			}
 		}
 
-		for(const TPair<FName, FRigVMTemplateArgument::FType>& TypePair : Types)
+		for(const TPair<FName, FRigVMTemplateArgumentType>& TypePair : Types)
 		{
 			const FName& PinName = TypePair.Key;
-			const FRigVMTemplateArgument::FType& ExpectedType = TypePair.Value;
+			const FRigVMTemplateArgumentType& ExpectedType = TypePair.Value;
 			
 			URigVMPin* Pin = Node->FindPin(PinName.ToString());
 			if(Pin == nullptr)
@@ -1965,7 +1965,7 @@ bool URigVMController::UnresolveTemplateNodes(const TArray<URigVMTemplateNode*>&
 
 				if(Pin->IsArray())
 				{
-					FRigVMTemplateArgument::FType ElementType = ExpectedType;
+					FRigVMTemplateArgumentType ElementType = ExpectedType;
 					ElementType.CPPType = RigVMTypeUtils::BaseTypeFromArrayType(ElementType.CPPType);
 					
 					TArray<URigVMPin*> SubPins = Pin->GetSubPins();
@@ -1986,7 +1986,7 @@ bool URigVMController::UnresolveTemplateNodes(const TArray<URigVMTemplateNode*>&
 			return A->IsWildCard() == B->IsWildCard();
 		}), bSetupUndoRedo);
 	
-	for (TPair<URigVMPin*, FRigVMTemplateArgument::FType>& Pin : PinsToResolve)
+	for (TPair<URigVMPin*, FRigVMTemplateArgumentType>& Pin : PinsToResolve)
 	{
 		TArray<URigVMLink*> LinksToBreak;
 		for (URigVMLink* Link : Pin.Key->GetLinks())
@@ -3583,16 +3583,8 @@ TArray<FName> URigVMController::ImportNodesFromText(const FString& InText, bool
 			{
 				if (!TemplateNode->IsSingleton())
 				{
-					{
-						TGuardValue<TArray<FString>> GuardPreferredPermutation(TemplateNode->PreferredPermutationTypes, {});
-						TemplateNode->InitializeFilteredPermutations();
-					}
 					TemplateNode->InitializeFilteredPermutationsFromTypes();
-					if (!TemplateNode->PreferredPermutationTypes.IsEmpty())
-					{
-						TemplateNode->FilteredPermutations = TemplateNode->FindPermutationsForTypes(TemplateNode->PreferredPermutationTypes, true);
-					}
-					UpdateTemplateNodePinTypes(TemplateNode, bSetupUndoRedo);
+					UpdateTemplateNodePinTypes(TemplateNode, false);
 				}
 			}
 
@@ -3703,7 +3695,7 @@ TArray<FName> URigVMController::ImportNodesFromText(const FString& InText, bool
 							{
 								if (!FirstTemplateNode->IsSingleton())
 								{
-									TArray<FRigVMTemplateArgument::FType> InputTypes = GetWildcardFilteredTypes(SourcePin);
+									TArray<FRigVMTemplateArgumentType> InputTypes = GetWildcardFilteredTypes(SourcePin);
 									if (InputTypes.Num() > 0)
 									{
 										PrepareTemplatePinForType(TargetPin, InputTypes, bSetupUndoRedo);										
@@ -3714,7 +3706,7 @@ TArray<FName> URigVMController::ImportNodesFromText(const FString& InText, bool
 							{
 								if (!SecondTemplateNode->IsSingleton())
 								{
-									TArray<FRigVMTemplateArgument::FType> OutTypes = GetWildcardFilteredTypes(TargetPin);
+									TArray<FRigVMTemplateArgumentType> OutTypes = GetWildcardFilteredTypes(TargetPin);
 									if (OutTypes.Num() > 0)
 									{
 										PrepareTemplatePinForType(SourcePin, OutTypes, bSetupUndoRedo);										
@@ -9256,7 +9248,7 @@ bool URigVMController::AddLink(URigVMPin* OutputPin, URigVMPin* InputPin, bool b
 		{
 			if (!FirstTemplateNode->IsSingleton())
 			{
-				TArray<FRigVMTemplateArgument::FType> InputTypes = GetWildcardFilteredTypes(SecondToResolve);
+				TArray<FRigVMTemplateArgumentType> InputTypes = GetWildcardFilteredTypes(SecondToResolve);
 				if (InputTypes.Num() > 0)
 				{
 					if (!PrepareTemplatePinForType(FirstToResolve, InputTypes, bSetupUndoRedo))
@@ -9271,7 +9263,7 @@ bool URigVMController::AddLink(URigVMPin* OutputPin, URigVMPin* InputPin, bool b
 		{
 			if (!SecondTemplateNode->IsSingleton())
 			{
-				TArray<FRigVMTemplateArgument::FType> OutTypes = GetWildcardFilteredTypes(FirstToResolve);
+				TArray<FRigVMTemplateArgumentType> OutTypes = GetWildcardFilteredTypes(FirstToResolve);
 				if (OutTypes.Num() > 0)
 				{
 					if (!PrepareTemplatePinForType(SecondToResolve, OutTypes, bSetupUndoRedo))
@@ -11645,7 +11637,7 @@ URigVMRerouteNode* URigVMController::AddFreeRerouteNode(bool bShowAsFullNode, co
 	Node->InitializeFilteredPermutations();
 	if (InCPPType != RigVMTypeUtils::GetWildCardCPPType() && InCPPType != RigVMTypeUtils::GetWildCardArrayCPPType())
 	{
-		PrepareTemplatePinForType(ValuePin, {FRigVMTemplateArgument::FType(ValuePin->GetCPPType(), ValuePin->GetCPPTypeObject())}, bSetupUndoRedo);
+		PrepareTemplatePinForType(ValuePin, {FRigVMTemplateArgumentType(ValuePin->GetCPPType(), ValuePin->GetCPPTypeObject())}, bSetupUndoRedo);
 		TArray<int32> FilterPermutations = Node->GetFilteredPermutationsIndices();
 		if (FilterPermutations.Num() == 1)
 		{
@@ -11837,7 +11829,7 @@ URigVMIfNode* URigVMController::AddIfNode(const FString& InCPPType, const FName&
 	Node->InitializeFilteredPermutations();
 	if (InCPPType != RigVMTypeUtils::GetWildCardCPPType() && InCPPType != RigVMTypeUtils::GetWildCardArrayCPPType())
 	{
-		PrepareTemplatePinForType(TruePin, {FRigVMTemplateArgument::FType(TruePin->GetCPPType(), TruePin->GetCPPTypeObject())}, bSetupUndoRedo);
+		PrepareTemplatePinForType(TruePin, {FRigVMTemplateArgumentType(TruePin->GetCPPType(), TruePin->GetCPPTypeObject())}, bSetupUndoRedo);
 		TArray<int32> FilterPermutations = Node->GetFilteredPermutationsIndices();
 		if (FilterPermutations.Num() == 1)
 		{
@@ -11958,7 +11950,7 @@ URigVMSelectNode* URigVMController::AddSelectNode(const FString& InCPPType, cons
 	Node->InitializeFilteredPermutations();
 	if (InCPPType != RigVMTypeUtils::GetWildCardCPPType() && InCPPType != RigVMTypeUtils::GetWildCardArrayCPPType())
 	{
-		PrepareTemplatePinForType(ResultPin, {FRigVMTemplateArgument::FType(ResultPin->GetCPPType(), ResultPin->GetCPPTypeObject())}, bSetupUndoRedo);
+		PrepareTemplatePinForType(ResultPin, {FRigVMTemplateArgumentType(ResultPin->GetCPPType(), ResultPin->GetCPPTypeObject())}, bSetupUndoRedo);
 		TArray<int32> FilterPermutations = Node->GetFilteredPermutationsIndices();
 		if (FilterPermutations.Num() == 1)
 		{
@@ -12060,7 +12052,7 @@ URigVMTemplateNode* URigVMController::AddTemplateNode(const FName& InNotation, c
 		const FRigVMTemplateArgument* Arg = Template->GetArgument(ArgIndex);
 
 		URigVMPin* Pin = NewObject<URigVMPin>(Node, Arg->GetName());
-		const FRigVMTemplateArgument::FType& Type = Types.FindChecked(Arg->GetName());
+		const FRigVMTemplateArgumentType& Type = Types.FindChecked(Arg->GetName());
 		Pin->CPPType = Type.CPPType;
 		Pin->CPPTypeObject = Type.CPPTypeObject;
 		if (Pin->CPPTypeObject)
@@ -12543,7 +12535,7 @@ URigVMArrayNode* URigVMController::AddArrayNode(ERigVMOpCode InOpCode, const FSt
 	if (InCPPType != RigVMTypeUtils::GetWildCardCPPType() && InCPPType != RigVMTypeUtils::GetWildCardArrayCPPType())
 	{
 		URigVMPin* ArrayPin = Node->FindPin(URigVMArrayNode::ArrayName);
-		PrepareTemplatePinForType(ArrayPin, {FRigVMTemplateArgument::FType(ArrayPin->GetCPPType(), ArrayPin->GetCPPTypeObject())}, bSetupUndoRedo);
+		PrepareTemplatePinForType(ArrayPin, {FRigVMTemplateArgumentType(ArrayPin->GetCPPType(), ArrayPin->GetCPPTypeObject())}, bSetupUndoRedo);
 		TArray<int32> FilterPermutations = Node->GetFilteredPermutationsIndices();
 		if (FilterPermutations.Num() == 1)
 		{
@@ -14944,7 +14936,7 @@ bool URigVMController::FullyResolveTemplateNode(URigVMTemplateNode* InNode, int3
 		for (URigVMPin* Pin : InNode->GetPins())
 		{
 			check(!Pin->IsWildCard());
-			TypeMap.Add(Pin->GetFName(), FRigVMTemplateArgument::FType(Pin->GetCPPType(), Pin->GetCPPTypeObject()));
+			TypeMap.Add(Pin->GetFName(), FRigVMTemplateArgumentType(Pin->GetCPPType(), Pin->GetCPPTypeObject()));
 		}
 
 		TArray<int32> Permutations;
@@ -14961,11 +14953,11 @@ bool URigVMController::FullyResolveTemplateNode(URigVMTemplateNode* InNode, int3
 	// find all existing pins that we may need to change
 	TArray<FRigVMTemplateArgument> MissingPins;
 	TArray<URigVMPin*> PinsToRemove;
-	TMap<URigVMPin*, FRigVMTemplateArgument::FType> PinTypesToChange;
+	TMap<URigVMPin*, FRigVMTemplateArgumentType> PinTypesToChange;
 	for(int32 ArgIndex = 0; ArgIndex < Template->NumArguments(); ArgIndex++)
 	{
 		const FRigVMTemplateArgument* Argument = Template->GetArgument(ArgIndex);
-		const FRigVMTemplateArgument::FType ResolvedType = Argument->GetSupportedTypes(PermutationIndices)[0];
+		const FRigVMTemplateArgumentType ResolvedType = Argument->GetSupportedTypes(PermutationIndices)[0];
 
 		URigVMPin* Pin = InNode->FindPin(Argument->GetName().ToString());
 		if(Pin == nullptr)
@@ -14989,7 +14981,7 @@ bool URigVMController::FullyResolveTemplateNode(URigVMTemplateNode* InNode, int3
 		for (TFieldIterator<FProperty> It(ResolvedFunction->Struct); It; ++It)
 		{
 			const FRigVMTemplateArgument ExpectedArgument(*It);
-			const FRigVMTemplateArgument::FType ExpectedType = ExpectedArgument.GetSupportedTypes()[0];
+			const FRigVMTemplateArgumentType ExpectedType = ExpectedArgument.GetSupportedTypes()[0];
 
 			if(URigVMPin* Pin = InNode->FindPin(It->GetFName().ToString()))
 			{
@@ -15034,10 +15026,10 @@ bool URigVMController::FullyResolveTemplateNode(URigVMTemplateNode* InNode, int3
 
 	// update the incorrectly typed pins
 	bool bNeedsTemplateUpdate = false;
-	for(const TPair<URigVMPin*, FRigVMTemplateArgument::FType>& Pair : PinTypesToChange)
+	for(const TPair<URigVMPin*, FRigVMTemplateArgumentType>& Pair : PinTypesToChange)
 	{
 		URigVMPin* Pin = Pair.Key;
-		const FRigVMTemplateArgument::FType& ExpectedType = Pair.Value;
+		const FRigVMTemplateArgumentType& ExpectedType = Pair.Value;
 		
 		if(!Pin->IsWildCard())
 		{
@@ -15111,7 +15103,7 @@ bool URigVMController::FullyResolveTemplateNode(URigVMTemplateNode* InNode, int3
 	return true;
 }
 
-bool URigVMController::PrepareTemplatePinForType(URigVMPin* InPin, const TArray<FRigVMTemplateArgument::FType>& InTypes, bool bSetupUndoRedo)
+bool URigVMController::PrepareTemplatePinForType(URigVMPin* InPin, const TArray<FRigVMTemplateArgumentType>& InTypes, bool bSetupUndoRedo)
 {
 	if (!IsValidGraph())
 	{
@@ -15136,11 +15128,11 @@ bool URigVMController::PrepareTemplatePinForType(URigVMPin* InPin, const TArray<
 
 	// If the pin is an element of an array, get the array pin and convert the types
 	URigVMPin* RootPin = InPin;
-	TArray<FRigVMTemplateArgument::FType> Types = InTypes;
+	TArray<FRigVMTemplateArgumentType> Types = InTypes;
 	if (InPin->IsArrayElement())
 	{
 		RootPin = InPin->GetParentPin();
-		for (FRigVMTemplateArgument::FType& Type : Types)
+		for (FRigVMTemplateArgumentType& Type : Types)
 		{
 			Type.ConvertToArray();
 		}
@@ -15158,7 +15150,7 @@ bool URigVMController::PrepareTemplatePinForType(URigVMPin* InPin, const TArray<
 	}
 
 	bool bFilteredSupportsType = false;
-	for (const FRigVMTemplateArgument::FType& Type : Types)
+	for (const FRigVMTemplateArgumentType& Type : Types)
 	{
 		if (TemplateNode->FilteredSupportsType(RootPin, Type.CPPType))
 		{
@@ -15179,7 +15171,7 @@ bool URigVMController::PrepareTemplatePinForType(URigVMPin* InPin, const TArray<
 			TemplateNode->PreferredPermutationTypes = {};
 		}
 		
-		for (const FRigVMTemplateArgument::FType& Type : Types)
+		for (const FRigVMTemplateArgumentType& Type : Types)
 		{
 			if (TemplateNode->SupportsType(RootPin, Type.CPPType))
 			{
@@ -15302,12 +15294,12 @@ bool URigVMController::PrepareTemplatePinForType(URigVMPin* InPin, const TArray<
 	return false;
 }
 
-TArray<FRigVMTemplateArgument::FType> URigVMController::GetWildcardFilteredTypes(URigVMPin* InPin)
+TArray<FRigVMTemplateArgumentType> URigVMController::GetWildcardFilteredTypes(URigVMPin* InPin)
 {
-	TArray<FRigVMTemplateArgument::FType> Types;
+	TArray<FRigVMTemplateArgumentType> Types;
 	if (!InPin->IsWildCard())
 	{
-		Types.Add(FRigVMTemplateArgument::FType(InPin->GetCPPType(), InPin->GetCPPTypeObject()));
+		Types.Add(FRigVMTemplateArgumentType(InPin->GetCPPType(), InPin->GetCPPTypeObject()));
 		return Types;
 	}
 
@@ -15352,7 +15344,7 @@ bool URigVMController::ResolveWildCardPin(const FString& InPinPath, const FStrin
 
 	if (URigVMPin* Pin = Graph->FindPin(InPinPath))
 	{
-		if (ResolveWildCardPin(Pin, FRigVMTemplateArgument::FType(CPPType, CPPTypeObject), bSetupUndoRedo, bPrintPythonCommand))
+		if (ResolveWildCardPin(Pin, FRigVMTemplateArgumentType(CPPType, CPPTypeObject), bSetupUndoRedo, bPrintPythonCommand))
 		{
 			if (bPrintPythonCommand)
 			{
@@ -15374,7 +15366,7 @@ bool URigVMController::ResolveWildCardPin(const FString& InPinPath, const FStrin
 	return false;
 }
 
-bool URigVMController::ResolveWildCardPin(URigVMPin* InPin, const FRigVMTemplateArgument::FType& InType, bool bSetupUndoRedo, bool bPrintPythonCommand)
+bool URigVMController::ResolveWildCardPin(URigVMPin* InPin, const FRigVMTemplateArgumentType& InType, bool bSetupUndoRedo, bool bPrintPythonCommand)
 {
 	if (!IsValidGraph())
 	{
@@ -15469,7 +15461,7 @@ bool URigVMController::UpdateFilteredPermutations(URigVMPin* InPin, URigVMPin* I
 	return false;
 }
 
-bool URigVMController::UpdateFilteredPermutations(URigVMPin* InPin, const TArray<FRigVMTemplateArgument::FType>& InTypes, bool bSetupUndoRedo)
+bool URigVMController::UpdateFilteredPermutations(URigVMPin* InPin, const TArray<FRigVMTemplateArgumentType>& InTypes, bool bSetupUndoRedo)
 {
 	URigVMTemplateNode* Node = Cast<URigVMTemplateNode>(InPin->GetNode());
 	if (!Node)
@@ -15510,7 +15502,7 @@ bool URigVMController::UpdateTemplateNodePinTypes(URigVMTemplateNode* InNode, bo
 			continue;
 		}
 		
-		TArray<FRigVMTemplateArgument::FType> Types = InNode->GetFilteredTypesForPin(Pin);
+		TArray<FRigVMTemplateArgumentType> Types = InNode->GetFilteredTypesForPin(Pin);
 		if (Types.IsEmpty())
 		{
 			continue;
@@ -15688,7 +15680,7 @@ bool URigVMController::RecomputeAllTemplateFilteredTypes(bool bSetupUndoRedo)
 	check(Graph);
 
 	// Save all template pin types, in case we need them after initializing
-	TMap<URigVMPin*, FRigVMTemplateArgument::FType> TypesBeforeRecomputing;
+	TMap<URigVMPin*, FRigVMTemplateArgumentType> TypesBeforeRecomputing;
 	for (URigVMNode* Node : Graph->GetNodes())
 	{
 		if (URigVMTemplateNode* TemplateNode = Cast<URigVMTemplateNode>(Node))
@@ -15714,7 +15706,7 @@ bool URigVMController::RecomputeAllTemplateFilteredTypes(bool bSetupUndoRedo)
 						continue;
 					}
 
-					TypesBeforeRecomputing.Add(Pin, FRigVMTemplateArgument::FType(Pin->GetCPPType(), Pin->GetCPPTypeObject()));
+					TypesBeforeRecomputing.Add(Pin, FRigVMTemplateArgumentType(Pin->GetCPPType(), Pin->GetCPPTypeObject()));
 				}
 			}
 		}
@@ -15763,7 +15755,7 @@ bool URigVMController::RecomputeAllTemplateFilteredTypes(bool bSetupUndoRedo)
 				if (OutputPin->IsStructMember())
 				{
 					URigVMPin* RootPin = OutputPin->GetRootPin();
-					const FRigVMTemplateArgument::FType Type = TypesBeforeRecomputing.FindChecked(RootPin);
+					const FRigVMTemplateArgumentType Type = TypesBeforeRecomputing.FindChecked(RootPin);
 					if (UpdateFilteredPermutations(RootPin, {Type}, bSetupUndoRedo))
 					{
 						PropagateTemplateFilteredTypes(OutputNode, bSetupUndoRedo);
@@ -15780,7 +15772,7 @@ bool URigVMController::RecomputeAllTemplateFilteredTypes(bool bSetupUndoRedo)
 				if (InputPin->IsStructMember() && InputNode->PreferredPermutationTypes.IsEmpty())
 				{
 					URigVMPin* RootPin = InputPin->GetRootPin();
-					const FRigVMTemplateArgument::FType Type = TypesBeforeRecomputing.FindChecked(RootPin);
+					const FRigVMTemplateArgumentType Type = TypesBeforeRecomputing.FindChecked(RootPin);
 					if (UpdateFilteredPermutations(RootPin, {Type}, bSetupUndoRedo))
 					{
 						PropagateTemplateFilteredTypes(InputNode, bSetupUndoRedo);
diff --git a/Engine/Source/Developer/RigVMDeveloper/Private/RigVMModel/RigVMPin.cpp b/Engine/Source/Developer/RigVMDeveloper/Private/RigVMModel/RigVMPin.cpp
--- a/Engine/Source/Developer/RigVMDeveloper/Private/RigVMModel/RigVMPin.cpp
+++ b/Engine/Source/Developer/RigVMDeveloper/Private/RigVMModel/RigVMPin.cpp
@@ -610,9 +610,9 @@ FString URigVMPin::GetArrayElementCppType() const
 	return RigVMTypeUtils::BaseTypeFromArrayType(ResolvedType);
 }
 
-FRigVMTemplateArgument::FType URigVMPin::GetTemplateArgumentType() const
+FRigVMTemplateArgumentType URigVMPin::GetTemplateArgumentType() const
 {
-	return FRigVMTemplateArgument::FType(GetCPPType(), GetCPPTypeObject());
+	return FRigVMTemplateArgumentType(GetCPPType(), GetCPPTypeObject());
 }
 
 bool URigVMPin::IsStringType() const
diff --git a/Engine/Source/Developer/RigVMDeveloper/Public/RigVMModel/Nodes/RigVMTemplateNode.h b/Engine/Source/Developer/RigVMDeveloper/Public/RigVMModel/Nodes/RigVMTemplateNode.h
--- a/Engine/Source/Developer/RigVMDeveloper/Public/RigVMModel/Nodes/RigVMTemplateNode.h
+++ b/Engine/Source/Developer/RigVMDeveloper/Public/RigVMModel/Nodes/RigVMTemplateNode.h
@@ -75,14 +75,14 @@ public:
 	const TArray<int32>& GetFilteredPermutationsIndices() const;
 
 	// returns the filtered types of this pin
-	TArray<FRigVMTemplateArgument::FType> GetFilteredTypesForPin(URigVMPin* InPin) const;
+	TArray<FRigVMTemplateArgumentType> GetFilteredTypesForPin(URigVMPin* InPin) const;
 
 	// returns true if updating pin filters with InTypes would result in different filters 
-	bool PinNeedsFilteredTypesUpdate(URigVMPin* InPin, const TArray<FRigVMTemplateArgument::FType>& InTypes);
+	bool PinNeedsFilteredTypesUpdate(URigVMPin* InPin, const TArray<FRigVMTemplateArgumentType>& InTypes);
 	bool PinNeedsFilteredTypesUpdate(URigVMPin* InPin, URigVMPin* LinkedPin);
 
 	// updates the filtered permutations given a link or the types for a pin
-	bool UpdateFilteredPermutations(URigVMPin* InPin, const TArray<FRigVMTemplateArgument::FType>& InTypes);
+	bool UpdateFilteredPermutations(URigVMPin* InPin, const TArray<FRigVMTemplateArgumentType>& InTypes);
 	bool UpdateFilteredPermutations(URigVMPin* InPin, URigVMPin* LinkedPin);
 
 	// initializes the filtered permutations to all possible permutations
@@ -96,7 +96,7 @@ protected:
 	virtual void InvalidateCache() override;
 	
 	TArray<int32> GetNewFilteredPermutations(URigVMPin* InPin, URigVMPin* LinkedPin);
-	TArray<int32> GetNewFilteredPermutations(URigVMPin* InPin, const TArray<FRigVMTemplateArgument::FType>& InTypes);
+	TArray<int32> GetNewFilteredPermutations(URigVMPin* InPin, const TArray<FRigVMTemplateArgumentType>& InTypes);
 
 	TArray<int32> FindPermutationsForTypes(const TArray<FString>& ArgumentTypes, bool bAllowCasting = false);
 	TArray<FString> GetArgumentTypesForPermutation(const int32 InPermutationIndex);
@@ -113,7 +113,7 @@ protected:
 	TArray<FString> PreferredPermutationTypes;
 	
 	TArray<int32> FilteredPermutations;
-	TMap<FString, TPair<bool, FRigVMTemplateArgument::FType>> SupportedTypesCache;
+	TMap<FString, TPair<bool, FRigVMTemplateArgumentType>> SupportedTypesCache;
 
 	mutable const FRigVMTemplate* CachedTemplate;
 	mutable const FRigVMFunction* CachedFunction;
diff --git a/Engine/Source/Developer/RigVMDeveloper/Public/RigVMModel/RigVMController.h b/Engine/Source/Developer/RigVMDeveloper/Public/RigVMModel/RigVMController.h
--- a/Engine/Source/Developer/RigVMDeveloper/Public/RigVMModel/RigVMController.h
+++ b/Engine/Source/Developer/RigVMDeveloper/Public/RigVMModel/RigVMController.h
@@ -378,7 +378,7 @@ public:
 	// Resolves a wildcard pin on any node
 	UFUNCTION(BlueprintCallable, Category = RigVMController)
 	bool ResolveWildCardPin(const FString& InPinPath, const FString& InCPPType, const FName& InCPPTypeObjectPath, bool bSetupUndoRedo = true, bool bPrintPythonCommand = false);
-	bool ResolveWildCardPin(URigVMPin* InPin, const FRigVMTemplateArgument::FType& InType, bool bSetupUndoRedo = true, bool bPrintPythonCommand = false);
+	bool ResolveWildCardPin(URigVMPin* InPin, const FRigVMTemplateArgumentType& InType, bool bSetupUndoRedo = true, bool bPrintPythonCommand = false);
 		
 	// Adds a Function / Struct Node to the edited Graph as an injected node
 	// UnitNode represent a RIGVM_METHOD declaration on a USTRUCT.
@@ -1050,15 +1050,15 @@ private:
 	// If any of the types is supported (without breaking any links), then the filtered permutations will be updated and the change will
 	// propagate to other nodes in the graph.
 	// If it is not supported, we will attempt to find and break any links that do not support this change
-	bool PrepareTemplatePinForType(URigVMPin* InPin, const TArray<FRigVMTemplateArgument::FType>& InTypes, bool bSetupUndoRedo);
+	bool PrepareTemplatePinForType(URigVMPin* InPin, const TArray<FRigVMTemplateArgumentType>& InTypes, bool bSetupUndoRedo);
 
 	// Get filtered types for a wildcard node. If template node, that means just returning its filtered wildcard types, but if it's another type of node (select, if, rereoute), iterate
 	// its connections to figure out the filtered types
-	TArray<FRigVMTemplateArgument::FType> GetWildcardFilteredTypes(URigVMPin* InPin);
+	TArray<FRigVMTemplateArgumentType> GetWildcardFilteredTypes(URigVMPin* InPin);
 	
 	// Updates the permutations allowed without having to break any links
 	bool UpdateFilteredPermutations(URigVMPin* InPin, URigVMPin* InLinkedPin, bool bSetupUndoRedo);
-	bool UpdateFilteredPermutations(URigVMPin* InPin, const TArray<FRigVMTemplateArgument::FType>& InTypes, bool bSetupUndoRedo);
+	bool UpdateFilteredPermutations(URigVMPin* InPin, const TArray<FRigVMTemplateArgumentType>& InTypes, bool bSetupUndoRedo);
 
 	// Changes Pin types if filtered types of a pin are unique
 	bool UpdateTemplateNodePinTypes(URigVMTemplateNode* InNode, bool bSetupUndoRedo);
diff --git a/Engine/Source/Developer/RigVMDeveloper/Public/RigVMModel/RigVMPin.h b/Engine/Source/Developer/RigVMDeveloper/Public/RigVMModel/RigVMPin.h
--- a/Engine/Source/Developer/RigVMDeveloper/Public/RigVMModel/RigVMPin.h
+++ b/Engine/Source/Developer/RigVMDeveloper/Public/RigVMModel/RigVMPin.h
@@ -234,7 +234,7 @@ public:
 	FString GetArrayElementCppType() const;
 
 	// Returns the argument type this pin would represent within a template
-	FRigVMTemplateArgument::FType GetTemplateArgumentType() const;
+	FRigVMTemplateArgumentType GetTemplateArgumentType() const;
 
 	// Returns true if the C++ data type is FString or FName
 	UFUNCTION(BlueprintCallable, Category = RigVMPin)
diff --git a/Engine/Source/Runtime/RigVM/Private/RigVMCore/RigVMRegistry.cpp b/Engine/Source/Runtime/RigVM/Private/RigVMCore/RigVMRegistry.cpp
--- a/Engine/Source/Runtime/RigVM/Private/RigVMCore/RigVMRegistry.cpp
+++ b/Engine/Source/Runtime/RigVM/Private/RigVMCore/RigVMRegistry.cpp
@@ -131,7 +131,7 @@ const FRigVMTemplate* FRigVMRegistry::GetOrAddTemplateFromArguments(const FName&
 	int32 NumPermutations = 1;
 	for(const FRigVMTemplateArgument& Argument : InArguments)
 	{
-		if(!Argument.bSingleton && NumPermutations > 1)
+		if(!Argument.IsSingleton() && NumPermutations > 1)
 		{
 			if(Argument.Types.Num() != NumPermutations)
 			{
@@ -161,7 +161,6 @@ const FRigVMTemplate* FRigVMRegistry::GetOrAddTemplateFromArguments(const FName&
 				Argument.TypeToPermutations.Add(Argument.Types[i].CPPType, {i});				
 			}
 			
-			Argument.bSingleton = false;
 			NumPermutations = FMath::Max(NumPermutations, Argument.Types.Num()); 
 		}
 	}
@@ -173,7 +172,7 @@ const FRigVMTemplate* FRigVMRegistry::GetOrAddTemplateFromArguments(const FName&
 		{
 			if(Argument.Types.Num() == 1)
 			{
-				const FRigVMTemplateArgument::FType Type = Argument.Types[0];
+				const FRigVMTemplateArgumentType Type = Argument.Types[0];
 				Argument.Types.SetNum(NumPermutations);
 				TArray<int32> ArgTypePermutations;
 				ArgTypePermutations.SetNum(NumPermutations);
@@ -183,7 +182,6 @@ const FRigVMTemplate* FRigVMRegistry::GetOrAddTemplateFromArguments(const FName&
 					ArgTypePermutations[Index] = Index;
 				}
 				Argument.TypeToPermutations.Add(Type.CPPType, ArgTypePermutations);
-				Argument.bSingleton = true;
 			}
 		}
 	}
diff --git a/Engine/Source/Runtime/RigVM/Private/RigVMCore/RigVMTemplate.cpp b/Engine/Source/Runtime/RigVM/Private/RigVMCore/RigVMTemplate.cpp
--- a/Engine/Source/Runtime/RigVM/Private/RigVMCore/RigVMTemplate.cpp
+++ b/Engine/Source/Runtime/RigVM/Private/RigVMCore/RigVMTemplate.cpp
@@ -13,7 +13,6 @@ FRigVMTemplateArgument::FRigVMTemplateArgument()
 	: Index(INDEX_NONE)
 	, Name(NAME_None)
 	, Direction(ERigVMPinDirection::IO)
-	, bSingleton(true)
 {
 }
 
@@ -21,7 +20,6 @@ FRigVMTemplateArgument::FRigVMTemplateArgument(FProperty* InProperty)
 	: Index(INDEX_NONE)
 	, Name(InProperty->GetFName())
 	, Direction(ERigVMPinDirection::IO)
-	, bSingleton(true)
 {
 #if WITH_EDITOR
 	Direction = FRigVMStruct::GetPinDirectionFromProperty(InProperty);
@@ -29,7 +27,7 @@ FRigVMTemplateArgument::FRigVMTemplateArgument(FProperty* InProperty)
 
 	FString ExtendedType;
 	const FString CPPType = InProperty->GetCPPType(&ExtendedType);
-	FType Type(CPPType + ExtendedType);
+	FRigVMTemplateArgumentType Type(CPPType + ExtendedType);
 
 	if (FArrayProperty* ArrayProperty = CastField<FArrayProperty>(InProperty))
 	{
@@ -54,34 +52,23 @@ FRigVMTemplateArgument::FRigVMTemplateArgument(FProperty* InProperty)
 	TypeToPermutations.Add(Type.CPPType, {0});
 }
 
-FRigVMTemplateArgument::FRigVMTemplateArgument(const FName& InName, ERigVMPinDirection InDirection, const FType& InType)
+FRigVMTemplateArgument::FRigVMTemplateArgument(const FName& InName, ERigVMPinDirection InDirection, const FRigVMTemplateArgumentType& InType)
 : Index(INDEX_NONE)
 , Name(InName)
 , Direction(InDirection)
-, bSingleton(true)
 , Types({InType})
 {
 	TypeToPermutations.Add(InType.CPPType, {0});	
 }
 
-FRigVMTemplateArgument::FRigVMTemplateArgument(const FName& InName, ERigVMPinDirection InDirection, const TArray<FType>& InTypes)
+FRigVMTemplateArgument::FRigVMTemplateArgument(const FName& InName, ERigVMPinDirection InDirection, const TArray<FRigVMTemplateArgumentType>& InTypes)
 : Index(INDEX_NONE)
 , Name(InName)
 , Direction(InDirection)
-, bSingleton(true)
 , Types(InTypes)
 {
 	check(Types.Num() > 0);
 
-	for(int32 TypeIndex=1;TypeIndex<Types.Num();TypeIndex++)
-	{
-		if(Types[TypeIndex] != Types[0])
-		{
-			bSingleton = false;
-			break;
-		}
-	}
-
 	for(int32 TypeIndex=0;TypeIndex<Types.Num();TypeIndex++)
 	{
 		if (TArray<int32>* Permutations = TypeToPermutations.Find(Types[TypeIndex].CPPType))
@@ -95,7 +82,7 @@ FRigVMTemplateArgument::FRigVMTemplateArgument(const FName& InName, ERigVMPinDir
 	}
 }
 
-bool FRigVMTemplateArgument::SupportsType(const FString& InCPPType, FType* OutType) const
+bool FRigVMTemplateArgument::SupportsType(const FString& InCPPType, FRigVMTemplateArgumentType* OutType) const
 {
 	if (const TArray<int32>* Permutations = TypeToPermutations.Find(InCPPType))
 	{
@@ -107,7 +94,7 @@ bool FRigVMTemplateArgument::SupportsType(const FString& InCPPType, FType* OutTy
 	}
 
 	// Try to find compatible type
-	TArray<FString> CompatibleTypes = FType::GetCompatibleTypes(InCPPType);
+	TArray<FString> CompatibleTypes = FRigVMTemplateArgumentType::GetCompatibleTypes(InCPPType);
 	for (const FString& Compatible : CompatibleTypes)
 	{
 		if (const TArray<int32>* Permutations = TypeToPermutations.Find(Compatible))
@@ -125,7 +112,7 @@ bool FRigVMTemplateArgument::SupportsType(const FString& InCPPType, FType* OutTy
 
 bool FRigVMTemplateArgument::IsSingleton(const TArray<int32>& InPermutationIndices) const
 {
-	if (bSingleton)
+	if (TypeToPermutations.Num() == 1)
 	{
 		return true;
 	}
@@ -134,7 +121,7 @@ bool FRigVMTemplateArgument::IsSingleton(const TArray<int32>& InPermutationIndic
 		return false;
 	}
 
-	const FType TypeToCheck = Types[InPermutationIndices[0]];
+	const FRigVMTemplateArgumentType TypeToCheck = Types[InPermutationIndices[0]];
 	for (int32 PermutationIndex = 1; PermutationIndex < InPermutationIndices.Num(); PermutationIndex++)
 	{
 		if (Types[InPermutationIndices[PermutationIndex]] != TypeToCheck)
@@ -171,63 +158,63 @@ FRigVMTemplateArgument::EArrayType FRigVMTemplateArgument::GetArrayType() const
 	return EArrayType_Invalid;
 }
 
-const TArray<FRigVMTemplateArgument::FType> FRigVMTemplateArgument::GetCompatibleTypes(ETypeCategory InCategory)
+const TArray<FRigVMTemplateArgumentType> FRigVMTemplateArgument::GetCompatibleTypes(ETypeCategory InCategory)
 {
 	check(InCategory != ETypeCategory_Invalid);
 	
-	static TMap<ETypeCategory, TArray<FType>> CompatibleTypes;
+	static TMap<ETypeCategory, TArray<FRigVMTemplateArgumentType>> CompatibleTypes;
 	if(CompatibleTypes.IsEmpty())
 	{
-		CompatibleTypes.Add(ETypeCategory_SingleAnyValue, TArray<FType>());
-		CompatibleTypes.Add(ETypeCategory_ArrayAnyValue, TArray<FType>());
-		CompatibleTypes.Add(ETypeCategory_ArrayArrayAnyValue, TArray<FType>());
-		CompatibleTypes.Add(ETypeCategory_SingleSimpleValue, TArray<FType>());
-		CompatibleTypes.Add(ETypeCategory_ArraySimpleValue, TArray<FType>());
-		CompatibleTypes.Add(ETypeCategory_ArrayArraySimpleValue, TArray<FType>());
-		CompatibleTypes.Add(ETypeCategory_SingleMathStructValue, TArray<FType>());
-		CompatibleTypes.Add(ETypeCategory_ArrayMathStructValue, TArray<FType>());
-		CompatibleTypes.Add(ETypeCategory_ArrayArrayMathStructValue, TArray<FType>());
-		CompatibleTypes.Add(ETypeCategory_SingleScriptStructValue, TArray<FType>());
-		CompatibleTypes.Add(ETypeCategory_ArrayScriptStructValue, TArray<FType>());
-		CompatibleTypes.Add(ETypeCategory_ArrayArrayScriptStructValue, TArray<FType>());
-		CompatibleTypes.Add(ETypeCategory_SingleEnumValue, TArray<FType>());
-		CompatibleTypes.Add(ETypeCategory_ArrayEnumValue, TArray<FType>());
-		CompatibleTypes.Add(ETypeCategory_ArrayArrayEnumValue, TArray<FType>());
-		CompatibleTypes.Add(ETypeCategory_SingleObjectValue, TArray<FType>());
-		CompatibleTypes.Add(ETypeCategory_ArrayObjectValue, TArray<FType>());
-		CompatibleTypes.Add(ETypeCategory_ArrayArrayObjectValue, TArray<FType>());
+		CompatibleTypes.Add(ETypeCategory_SingleAnyValue, TArray<FRigVMTemplateArgumentType>());
+		CompatibleTypes.Add(ETypeCategory_ArrayAnyValue, TArray<FRigVMTemplateArgumentType>());
+		CompatibleTypes.Add(ETypeCategory_ArrayArrayAnyValue, TArray<FRigVMTemplateArgumentType>());
+		CompatibleTypes.Add(ETypeCategory_SingleSimpleValue, TArray<FRigVMTemplateArgumentType>());
+		CompatibleTypes.Add(ETypeCategory_ArraySimpleValue, TArray<FRigVMTemplateArgumentType>());
+		CompatibleTypes.Add(ETypeCategory_ArrayArraySimpleValue, TArray<FRigVMTemplateArgumentType>());
+		CompatibleTypes.Add(ETypeCategory_SingleMathStructValue, TArray<FRigVMTemplateArgumentType>());
+		CompatibleTypes.Add(ETypeCategory_ArrayMathStructValue, TArray<FRigVMTemplateArgumentType>());
+		CompatibleTypes.Add(ETypeCategory_ArrayArrayMathStructValue, TArray<FRigVMTemplateArgumentType>());
+		CompatibleTypes.Add(ETypeCategory_SingleScriptStructValue, TArray<FRigVMTemplateArgumentType>());
+		CompatibleTypes.Add(ETypeCategory_ArrayScriptStructValue, TArray<FRigVMTemplateArgumentType>());
+		CompatibleTypes.Add(ETypeCategory_ArrayArrayScriptStructValue, TArray<FRigVMTemplateArgumentType>());
+		CompatibleTypes.Add(ETypeCategory_SingleEnumValue, TArray<FRigVMTemplateArgumentType>());
+		CompatibleTypes.Add(ETypeCategory_ArrayEnumValue, TArray<FRigVMTemplateArgumentType>());
+		CompatibleTypes.Add(ETypeCategory_ArrayArrayEnumValue, TArray<FRigVMTemplateArgumentType>());
+		CompatibleTypes.Add(ETypeCategory_SingleObjectValue, TArray<FRigVMTemplateArgumentType>());
+		CompatibleTypes.Add(ETypeCategory_ArrayObjectValue, TArray<FRigVMTemplateArgumentType>());
+		CompatibleTypes.Add(ETypeCategory_ArrayArrayObjectValue, TArray<FRigVMTemplateArgumentType>());
 		
-		TArray<FType>& SingleAnyValueTypes = CompatibleTypes.FindChecked(ETypeCategory_SingleAnyValue);
-		TArray<FType>& ArrayAnyValueTypes = CompatibleTypes.FindChecked(ETypeCategory_ArrayAnyValue);
-		TArray<FType>& ArrayArrayAnyValueTypes = CompatibleTypes.FindChecked(ETypeCategory_ArrayArrayAnyValue);
-		TArray<FType>& SingleSimpleValueTypes = CompatibleTypes.FindChecked(ETypeCategory_SingleSimpleValue);
-		TArray<FType>& ArraySimpleValueTypes = CompatibleTypes.FindChecked(ETypeCategory_ArraySimpleValue);
-		TArray<FType>& ArrayArraySimpleValueTypes = CompatibleTypes.FindChecked(ETypeCategory_ArrayArraySimpleValue);
-		TArray<FType>& SingleMathStructValueTypes = CompatibleTypes.FindChecked(ETypeCategory_SingleMathStructValue);
-		TArray<FType>& ArrayMathStructValueTypes = CompatibleTypes.FindChecked(ETypeCategory_ArrayMathStructValue);
-		TArray<FType>& ArrayArrayMathStructValueTypes = CompatibleTypes.FindChecked(ETypeCategory_ArrayArrayMathStructValue);
-		TArray<FType>& SingleScriptStructValueTypes = CompatibleTypes.FindChecked(ETypeCategory_SingleScriptStructValue);
-		TArray<FType>& ArrayScriptStructValueTypes = CompatibleTypes.FindChecked(ETypeCategory_ArrayScriptStructValue);
-		TArray<FType>& ArrayArrayScriptStructValueTypes = CompatibleTypes.FindChecked(ETypeCategory_ArrayArrayScriptStructValue);
-		TArray<FType>& SingleEnumValueTypes = CompatibleTypes.FindChecked(ETypeCategory_SingleEnumValue);
-		TArray<FType>& ArrayEnumValueTypes = CompatibleTypes.FindChecked(ETypeCategory_ArrayEnumValue);
-		TArray<FType>& ArrayArrayEnumValueTypes = CompatibleTypes.FindChecked(ETypeCategory_ArrayArrayEnumValue);
-		TArray<FType>& SingleObjectValueTypes = CompatibleTypes.FindChecked(ETypeCategory_SingleObjectValue);
-		TArray<FType>& ArrayObjectValueTypes = CompatibleTypes.FindChecked(ETypeCategory_ArrayObjectValue);
-		TArray<FType>& ArrayArrayObjectValueTypes = CompatibleTypes.FindChecked(ETypeCategory_ArrayArrayObjectValue);
+		TArray<FRigVMTemplateArgumentType>& SingleAnyValueTypes = CompatibleTypes.FindChecked(ETypeCategory_SingleAnyValue);
+		TArray<FRigVMTemplateArgumentType>& ArrayAnyValueTypes = CompatibleTypes.FindChecked(ETypeCategory_ArrayAnyValue);
+		TArray<FRigVMTemplateArgumentType>& ArrayArrayAnyValueTypes = CompatibleTypes.FindChecked(ETypeCategory_ArrayArrayAnyValue);
+		TArray<FRigVMTemplateArgumentType>& SingleSimpleValueTypes = CompatibleTypes.FindChecked(ETypeCategory_SingleSimpleValue);
+		TArray<FRigVMTemplateArgumentType>& ArraySimpleValueTypes = CompatibleTypes.FindChecked(ETypeCategory_ArraySimpleValue);
+		TArray<FRigVMTemplateArgumentType>& ArrayArraySimpleValueTypes = CompatibleTypes.FindChecked(ETypeCategory_ArrayArraySimpleValue);
+		TArray<FRigVMTemplateArgumentType>& SingleMathStructValueTypes = CompatibleTypes.FindChecked(ETypeCategory_SingleMathStructValue);
+		TArray<FRigVMTemplateArgumentType>& ArrayMathStructValueTypes = CompatibleTypes.FindChecked(ETypeCategory_ArrayMathStructValue);
+		TArray<FRigVMTemplateArgumentType>& ArrayArrayMathStructValueTypes = CompatibleTypes.FindChecked(ETypeCategory_ArrayArrayMathStructValue);
+		TArray<FRigVMTemplateArgumentType>& SingleScriptStructValueTypes = CompatibleTypes.FindChecked(ETypeCategory_SingleScriptStructValue);
+		TArray<FRigVMTemplateArgumentType>& ArrayScriptStructValueTypes = CompatibleTypes.FindChecked(ETypeCategory_ArrayScriptStructValue);
+		TArray<FRigVMTemplateArgumentType>& ArrayArrayScriptStructValueTypes = CompatibleTypes.FindChecked(ETypeCategory_ArrayArrayScriptStructValue);
+		TArray<FRigVMTemplateArgumentType>& SingleEnumValueTypes = CompatibleTypes.FindChecked(ETypeCategory_SingleEnumValue);
+		TArray<FRigVMTemplateArgumentType>& ArrayEnumValueTypes = CompatibleTypes.FindChecked(ETypeCategory_ArrayEnumValue);
+		TArray<FRigVMTemplateArgumentType>& ArrayArrayEnumValueTypes = CompatibleTypes.FindChecked(ETypeCategory_ArrayArrayEnumValue);
+		TArray<FRigVMTemplateArgumentType>& SingleObjectValueTypes = CompatibleTypes.FindChecked(ETypeCategory_SingleObjectValue);
+		TArray<FRigVMTemplateArgumentType>& ArrayObjectValueTypes = CompatibleTypes.FindChecked(ETypeCategory_ArrayObjectValue);
+		TArray<FRigVMTemplateArgumentType>& ArrayArrayObjectValueTypes = CompatibleTypes.FindChecked(ETypeCategory_ArrayArrayObjectValue);
 		
-		SingleSimpleValueTypes.Add(FType(RigVMTypeUtils::BoolType));
-		SingleSimpleValueTypes.Add(FType(RigVMTypeUtils::Int32Type));
-		SingleSimpleValueTypes.Add(FType(RigVMTypeUtils::UInt8Type));
-		SingleSimpleValueTypes.Add(FType(RigVMTypeUtils::FloatType));
-		SingleSimpleValueTypes.Add(FType(RigVMTypeUtils::DoubleType));
-		SingleSimpleValueTypes.Add(FType(RigVMTypeUtils::FNameType));
-		SingleSimpleValueTypes.Add(FType(RigVMTypeUtils::FStringType));
-
-		for(const FType& Type : SingleSimpleValueTypes)
+		SingleSimpleValueTypes.Add(FRigVMTemplateArgumentType(RigVMTypeUtils::BoolType));
+		SingleSimpleValueTypes.Add(FRigVMTemplateArgumentType(RigVMTypeUtils::Int32Type));
+		SingleSimpleValueTypes.Add(FRigVMTemplateArgumentType(RigVMTypeUtils::UInt8Type));
+		SingleSimpleValueTypes.Add(FRigVMTemplateArgumentType(RigVMTypeUtils::FloatType));
+		SingleSimpleValueTypes.Add(FRigVMTemplateArgumentType(RigVMTypeUtils::DoubleType));
+		SingleSimpleValueTypes.Add(FRigVMTemplateArgumentType(RigVMTypeUtils::FNameType));
+		SingleSimpleValueTypes.Add(FRigVMTemplateArgumentType(RigVMTypeUtils::FStringType));
+
+		for(const FRigVMTemplateArgumentType& Type : SingleSimpleValueTypes)
 		{
-			ArraySimpleValueTypes.Add(FType(RigVMTypeUtils::ArrayTypeFromBaseType(Type.CPPType)));
-			ArrayArraySimpleValueTypes.Add(FType(RigVMTypeUtils::ArrayTypeFromBaseType(ArraySimpleValueTypes.Last().CPPType)));
+			ArraySimpleValueTypes.Add(FRigVMTemplateArgumentType(RigVMTypeUtils::ArrayTypeFromBaseType(Type.CPPType)));
+			ArrayArraySimpleValueTypes.Add(FRigVMTemplateArgumentType(RigVMTypeUtils::ArrayTypeFromBaseType(ArraySimpleValueTypes.Last().CPPType)));
 		}
 
 		SingleAnyValueTypes = SingleSimpleValueTypes;
@@ -249,9 +236,9 @@ const TArray<FRigVMTemplateArgument::FType> FRigVMTemplateArgument::GetCompatibl
 
 		for(UScriptStruct* MathType : MathTypes)
 		{
-			SingleMathStructValueTypes.Add(FType(MathType->GetStructCPPName(), MathType));
-			ArrayMathStructValueTypes.Add(FType(RigVMTypeUtils::ArrayTypeFromBaseType(MathType->GetStructCPPName()), MathType));
-			ArrayArrayMathStructValueTypes.Add(FType(RigVMTypeUtils::ArrayTypeFromBaseType(ArrayMathStructValueTypes.Last().CPPType), MathType));
+			SingleMathStructValueTypes.Add(FRigVMTemplateArgumentType(MathType->GetStructCPPName(), MathType));
+			ArrayMathStructValueTypes.Add(FRigVMTemplateArgumentType(RigVMTypeUtils::ArrayTypeFromBaseType(MathType->GetStructCPPName()), MathType));
+			ArrayArrayMathStructValueTypes.Add(FRigVMTemplateArgumentType(RigVMTypeUtils::ArrayTypeFromBaseType(ArrayMathStructValueTypes.Last().CPPType), MathType));
 		}
 
 		struct FTypeTraverser
@@ -376,12 +363,12 @@ const TArray<FRigVMTemplateArgument::FType> FRigVMTemplateArgument::GetCompatibl
 
 			if(!ScriptStruct->IsChildOf(FRigVMExecuteContext::StaticStruct()))
 			{
-				SingleAnyValueTypes.Add(FType(CPPType, ScriptStruct));
-				ArrayAnyValueTypes.Add(FType(RigVMTypeUtils::ArrayTypeFromBaseType(CPPType), ScriptStruct));
-				ArrayArrayAnyValueTypes.Add(FType(RigVMTypeUtils::ArrayTypeFromBaseType(ArrayAnyValueTypes.Last().CPPType), ScriptStruct));
-				SingleScriptStructValueTypes.Add(FType(CPPType, ScriptStruct));
-				ArrayScriptStructValueTypes.Add(FType(RigVMTypeUtils::ArrayTypeFromBaseType(CPPType), ScriptStruct));
-				ArrayArrayScriptStructValueTypes.Add(FType(RigVMTypeUtils::ArrayTypeFromBaseType(ArrayScriptStructValueTypes.Last().CPPType), ScriptStruct));
+				SingleAnyValueTypes.Add(FRigVMTemplateArgumentType(CPPType, ScriptStruct));
+				ArrayAnyValueTypes.Add(FRigVMTemplateArgumentType(RigVMTypeUtils::ArrayTypeFromBaseType(CPPType), ScriptStruct));
+				ArrayArrayAnyValueTypes.Add(FRigVMTemplateArgumentType(RigVMTypeUtils::ArrayTypeFromBaseType(ArrayAnyValueTypes.Last().CPPType), ScriptStruct));
+				SingleScriptStructValueTypes.Add(FRigVMTemplateArgumentType(CPPType, ScriptStruct));
+				ArrayScriptStructValueTypes.Add(FRigVMTemplateArgumentType(RigVMTypeUtils::ArrayTypeFromBaseType(CPPType), ScriptStruct));
+				ArrayArrayScriptStructValueTypes.Add(FRigVMTemplateArgumentType(RigVMTypeUtils::ArrayTypeFromBaseType(ArrayScriptStructValueTypes.Last().CPPType), ScriptStruct));
 			}
 		}
 
@@ -394,12 +381,12 @@ const TArray<FRigVMTemplateArgument::FType> FRigVMTemplateArgument::GetCompatibl
 				continue;
 			}
 			const FString CPPType = Enum->CppType.IsEmpty() ? Enum->GetName() : Enum->CppType;
-			SingleAnyValueTypes.Add(FType(CPPType, Enum));
-			ArrayAnyValueTypes.Add(FType(RigVMTypeUtils::ArrayTypeFromBaseType(CPPType), Enum));
-			ArrayArrayAnyValueTypes.Add(FType(RigVMTypeUtils::ArrayTypeFromBaseType(ArrayAnyValueTypes.Last().CPPType), Enum));
-			SingleEnumValueTypes.Add(FType(CPPType, Enum));
-			ArrayEnumValueTypes.Add(FType(RigVMTypeUtils::ArrayTypeFromBaseType(CPPType), Enum));
-			ArrayArrayEnumValueTypes.Add(FType(RigVMTypeUtils::ArrayTypeFromBaseType(ArrayEnumValueTypes.Last().CPPType), Enum));
+			SingleAnyValueTypes.Add(FRigVMTemplateArgumentType(CPPType, Enum));
+			ArrayAnyValueTypes.Add(FRigVMTemplateArgumentType(RigVMTypeUtils::ArrayTypeFromBaseType(CPPType), Enum));
+			ArrayArrayAnyValueTypes.Add(FRigVMTemplateArgumentType(RigVMTypeUtils::ArrayTypeFromBaseType(ArrayAnyValueTypes.Last().CPPType), Enum));
+			SingleEnumValueTypes.Add(FRigVMTemplateArgumentType(CPPType, Enum));
+			ArrayEnumValueTypes.Add(FRigVMTemplateArgumentType(RigVMTypeUtils::ArrayTypeFromBaseType(CPPType), Enum));
+			ArrayArrayEnumValueTypes.Add(FRigVMTemplateArgumentType(RigVMTypeUtils::ArrayTypeFromBaseType(ArrayEnumValueTypes.Last().CPPType), Enum));
 		}
 
 		// add all classes
@@ -414,12 +401,12 @@ const TArray<FRigVMTemplateArgument::FType> FRigVMTemplateArgument::GetCompatibl
 			}
 
 			const FString CPPType = Class->GetPrefixCPP() + Class->GetName();
-			SingleAnyValueTypes.Add(FType(CPPType, Class));
-			ArrayAnyValueTypes.Add(FType(RigVMTypeUtils::ArrayTypeFromBaseType(CPPType), Class));
-			ArrayArrayAnyValueTypes.Add(FType(RigVMTypeUtils::ArrayTypeFromBaseType(ArrayAnyValueTypes.Last().CPPType), Class));
-			SingleObjectValueTypes.Add(FType(CPPType, Class));
-			ArrayObjectValueTypes.Add(FType(RigVMTypeUtils::ArrayTypeFromBaseType(CPPType), Class));
-			ArrayArrayObjectValueTypes.Add(FType(RigVMTypeUtils::ArrayTypeFromBaseType(ArrayObjectValueTypes.Last().CPPType), Class));
+			SingleAnyValueTypes.Add(FRigVMTemplateArgumentType(CPPType, Class));
+			ArrayAnyValueTypes.Add(FRigVMTemplateArgumentType(RigVMTypeUtils::ArrayTypeFromBaseType(CPPType), Class));
+			ArrayArrayAnyValueTypes.Add(FRigVMTemplateArgumentType(RigVMTypeUtils::ArrayTypeFromBaseType(ArrayAnyValueTypes.Last().CPPType), Class));
+			SingleObjectValueTypes.Add(FRigVMTemplateArgumentType(CPPType, Class));
+			ArrayObjectValueTypes.Add(FRigVMTemplateArgumentType(RigVMTypeUtils::ArrayTypeFromBaseType(CPPType), Class));
+			ArrayArrayObjectValueTypes.Add(FRigVMTemplateArgumentType(RigVMTypeUtils::ArrayTypeFromBaseType(ArrayObjectValueTypes.Last().CPPType), Class));
 		}
 
 		// check that matching categories have the same size
@@ -440,17 +427,17 @@ const TArray<FRigVMTemplateArgument::FType> FRigVMTemplateArgument::GetCompatibl
 	return CompatibleTypes.FindChecked(InCategory);
 }
 
-const TArray<FRigVMTemplateArgument::FType>& FRigVMTemplateArgument::GetTypes() const
+const TArray<FRigVMTemplateArgumentType>& FRigVMTemplateArgument::GetTypes() const
 {
 	return Types;
 }
 
-TArray<FRigVMTemplateArgument::FType> FRigVMTemplateArgument::GetSupportedTypes(const TArray<int32>& InPermutationIndices) const
+TArray<FRigVMTemplateArgumentType> FRigVMTemplateArgument::GetSupportedTypes(const TArray<int32>& InPermutationIndices) const
 {
-	TArray<FType> SupportedTypes;
+	TArray<FRigVMTemplateArgumentType> SupportedTypes;
 	if(InPermutationIndices.IsEmpty())
 	{
-		for (const FType& Type : Types)
+		for (const FRigVMTemplateArgumentType& Type : Types)
 		{
 			SupportedTypes.AddUnique(Type);
 		}
@@ -470,7 +457,7 @@ TArray<FString> FRigVMTemplateArgument::GetSupportedTypeStrings(const TArray<int
 	TArray<FString> SupportedTypes;
 	if(InPermutationIndices.IsEmpty())
 	{
-		for (const FType& Type : Types)
+		for (const FRigVMTemplateArgumentType& Type : Types)
 		{
 			SupportedTypes.AddUnique(Type.CPPType);
 		}
@@ -718,6 +705,11 @@ FText FRigVMTemplate::GetTooltipText(const TArray<int32>& InPermutationIndices)
 
 	auto VisitPermutation = [&ResolvedTooltipText, this](int32 InPermutationIndex) -> bool
 	{
+		if (InPermutationIndex >= NumPermutations())
+		{
+			return false;
+		}
+		
 		const FRigVMFunction* ResolvedFunction = GetPermutation(InPermutationIndex);
 		if(ResolvedFunction == nullptr)
 		{
@@ -896,11 +888,6 @@ bool FRigVMTemplate::Merge(const FRigVMTemplate& InOther)
 			}
 		}
 		NewArgs[ArgumentIndex].Types.Add(InOther.Arguments[ArgumentIndex].Types[0]);
-
-		if (!Arguments[ArgumentIndex].Types.Contains(InOther.Arguments[ArgumentIndex].Types[0]))
-		{
-			NewArgs[ArgumentIndex].bSingleton = false;
-		}
 	}
 
 	Arguments = NewArgs;
@@ -917,7 +904,7 @@ const FRigVMTemplateArgument* FRigVMTemplate::FindArgument(const FName& InArgume
 	});
 }
 
-bool FRigVMTemplate::ArgumentSupportsType(const FName& InArgumentName, const FString& InCPPType, FRigVMTemplateArgument::FType* OutType) const
+bool FRigVMTemplate::ArgumentSupportsType(const FName& InArgumentName, const FString& InCPPType, FRigVMTemplateArgumentType* OutType) const
 {
 	if (const FRigVMTemplateArgument* Argument = FindArgument(InArgumentName))
 	{
@@ -977,14 +964,14 @@ bool FRigVMTemplate::Resolve(FTypeMap& InOutTypes, TArray<int32>& OutPermutation
 	
 	for (const FRigVMTemplateArgument& Argument : Arguments)
 	{
-		if (Argument.bSingleton)
+		if (Argument.IsSingleton())
 		{
 			InOutTypes.Add(Argument.Name, Argument.Types[0]);
 			continue;
 		}
-		else if (const FRigVMTemplateArgument::FType* InputType = InputTypes.Find(Argument.Name))
+		else if (const FRigVMTemplateArgumentType* InputType = InputTypes.Find(Argument.Name))
 		{
-			FRigVMTemplateArgument::FType MatchedType = *InputType;
+			FRigVMTemplateArgumentType MatchedType = *InputType;
 			bool bFoundMatch = false;
 			bool bFoundPerfectMatch = false;
 			for (int32 PermutationIndex = 0; PermutationIndex < Argument.Types.Num(); PermutationIndex++)
@@ -1039,23 +1026,23 @@ bool FRigVMTemplate::Resolve(FTypeMap& InOutTypes, TArray<int32>& OutPermutation
 		const FRigVMTemplateArgument::EArrayType ArrayType = Argument.GetArrayType();
 		if(ArrayType == FRigVMTemplateArgument::EArrayType_Mixed)
 		{
-			InOutTypes.Add(Argument.Name, FRigVMTemplateArgument::FType());
+			InOutTypes.Add(Argument.Name, FRigVMTemplateArgumentType());
 
-			if(const FRigVMTemplateArgument::FType* InputType = InputTypes.Find(Argument.Name))
+			if(const FRigVMTemplateArgumentType* InputType = InputTypes.Find(Argument.Name))
 			{
 				if(InputType->IsArray())
 				{
-					InOutTypes.FindChecked(Argument.Name) = FRigVMTemplateArgument::FType::Array();
+					InOutTypes.FindChecked(Argument.Name) = FRigVMTemplateArgumentType::Array();
 				}
 			}
 		}
 		else if(ArrayType == FRigVMTemplateArgument::EArrayType_ArrayValue)
 		{
-			InOutTypes.Add(Argument.Name, FRigVMTemplateArgument::FType::Array());
+			InOutTypes.Add(Argument.Name, FRigVMTemplateArgumentType::Array());
 		}
 		else
 		{
-			InOutTypes.Add(Argument.Name, FRigVMTemplateArgument::FType());
+			InOutTypes.Add(Argument.Name, FRigVMTemplateArgumentType());
 		}
 	}
 
@@ -1081,7 +1068,7 @@ bool FRigVMTemplate::Resolve(FTypeMap& InOutTypes, TArray<int32>& OutPermutation
 	return !OutPermutationIndices.IsEmpty();
 }
 
-bool FRigVMTemplate::ResolveArgument(const FName& InArgumentName, const FRigVMTemplateArgument::FType& InType,
+bool FRigVMTemplate::ResolveArgument(const FName& InArgumentName, const FRigVMTemplateArgumentType& InType,
 	FTypeMap& InOutTypes) const
 {
 	auto RemoveWildCardTypes = [](const FTypeMap& InTypes)
@@ -1109,7 +1096,7 @@ bool FRigVMTemplate::ResolveArgument(const FName& InArgumentName, const FRigVMTe
 	if(Resolve(ResolvedTypes, PermutationIndices, true))
 	{
 		// let's see if the input argument resolved into the expected type
-		const FRigVMTemplateArgument::FType ResolvedInputType = ResolvedTypes.FindChecked(InArgumentName);
+		const FRigVMTemplateArgumentType ResolvedInputType = ResolvedTypes.FindChecked(InArgumentName);
 		if(!ResolvedInputType.Matches(InType.CPPType, true))
 		{
 			return false;
diff --git a/Engine/Source/Runtime/RigVM/Public/RigVMCore/RigVMTemplate.h b/Engine/Source/Runtime/RigVM/Public/RigVMCore/RigVMTemplate.h
--- a/Engine/Source/Runtime/RigVM/Public/RigVMCore/RigVMTemplate.h
+++ b/Engine/Source/Runtime/RigVM/Public/RigVMCore/RigVMTemplate.h
@@ -6,142 +6,152 @@
 #include "RigVMFunction.h"
 #include "UObject/Object.h"
 #include "RigVMTypeUtils.h"
+#include "RigVMTemplate.generated.h"
 
-/**
- * The template argument represents a single parameter
- * in a function call and all of its possible types
- */
-struct RIGVM_API FRigVMTemplateArgument
+USTRUCT()
+struct RIGVM_API FRigVMTemplateArgumentType
 {
-	struct FType
+	GENERATED_BODY()
+
+	UPROPERTY()
+	FString CPPType;
+	
+	UPROPERTY()
+	TObjectPtr<UObject> CPPTypeObject; 
+
+	FRigVMTemplateArgumentType()
+		: CPPType()
+		, CPPTypeObject(nullptr)
 	{
-		FString CPPType;
-		UObject* CPPTypeObject;
+		CPPType = RigVMTypeUtils::GetWildCardCPPType();
+		CPPTypeObject = RigVMTypeUtils::GetWildCardCPPTypeObject();
+	}
 
-		FType()
-			: CPPType()
-			, CPPTypeObject(nullptr)
-		{
-			CPPType = RigVMTypeUtils::GetWildCardCPPType();
-			CPPTypeObject = RigVMTypeUtils::GetWildCardCPPTypeObject();
-		}
+	FRigVMTemplateArgumentType(const FString& InCPPType, UObject* InCPPTypeObject = nullptr)
+		: CPPType(InCPPType)
+		, CPPTypeObject(InCPPTypeObject)
+	{
+		check(!CPPType.IsEmpty());
+	}
 
-		FType(const FString& InCPPType, UObject* InCPPTypeObject = nullptr)
-			: CPPType(InCPPType)
-			, CPPTypeObject(InCPPTypeObject)
-		{
-			check(!CPPType.IsEmpty());
-		}
+	static FRigVMTemplateArgumentType Array()
+	{
+		return FRigVMTemplateArgumentType(RigVMTypeUtils::GetWildCardArrayCPPType(), RigVMTypeUtils::GetWildCardCPPTypeObject());
+	}
 
-		static FType Array()
-		{
-			return FType(RigVMTypeUtils::GetWildCardArrayCPPType(), RigVMTypeUtils::GetWildCardCPPTypeObject());
-		}
+	bool operator == (const FRigVMTemplateArgumentType& InOther) const
+	{
+		return CPPType == InOther.CPPType;
+	}
 
-		bool operator == (const FType& InOther) const
-		{
-			return CPPType == InOther.CPPType;
-		}
+	bool operator != (const FRigVMTemplateArgumentType& InOther) const
+	{
+		return CPPType != InOther.CPPType;
+	}
 
-		bool operator != (const FType& InOther) const
-		{
-			return CPPType != InOther.CPPType;
-		}
+	friend FORCEINLINE uint32 GetTypeHash(const FRigVMTemplateArgumentType& InType)
+	{
+		return GetTypeHash(InType.CPPType);
+	}
 
-		friend FORCEINLINE uint32 GetTypeHash(const FType& InType)
+	FORCEINLINE bool Matches(const FString& InCPPType, bool bAllowFloatingPointCasts = true) const
+	{
+		if(CPPType == InCPPType)
 		{
-			return GetTypeHash(InType.CPPType);
+			return true;
 		}
-
-		FORCEINLINE bool Matches(const FString& InCPPType, bool bAllowFloatingPointCasts = true) const
+		if(bAllowFloatingPointCasts)
 		{
-			if(CPPType == InCPPType)
+			if(InCPPType == RigVMTypeUtils::FloatType && CPPType == RigVMTypeUtils::DoubleType)
 			{
 				return true;
 			}
-			if(bAllowFloatingPointCasts)
-			{
-				if(InCPPType == RigVMTypeUtils::FloatType && CPPType == RigVMTypeUtils::DoubleType)
-				{
-					return true;
-				}
-				if(InCPPType == RigVMTypeUtils::DoubleType && CPPType == RigVMTypeUtils::FloatType)
-				{
-					return true;
-				}
-				if(InCPPType == RigVMTypeUtils::FloatArrayType && CPPType == RigVMTypeUtils::DoubleArrayType)
-				{
-					return true;
-				}
-				if(InCPPType == RigVMTypeUtils::DoubleArrayType && CPPType == RigVMTypeUtils::FloatArrayType)
-				{
-					return true;
-				}
-			}
-			return false;
-		}
-
-		static TArray<FString> GetCompatibleTypes(const FString& InCPPType)
-		{
-			if(InCPPType == RigVMTypeUtils::FloatType)
-			{
-				return {RigVMTypeUtils::DoubleType};
-			}
-			if(InCPPType == RigVMTypeUtils::DoubleType)
+			if(InCPPType == RigVMTypeUtils::DoubleType && CPPType == RigVMTypeUtils::FloatType)
 			{
-				return {RigVMTypeUtils::FloatType};
+				return true;
 			}
-			if(InCPPType == RigVMTypeUtils::FloatArrayType)
+			if(InCPPType == RigVMTypeUtils::FloatArrayType && CPPType == RigVMTypeUtils::DoubleArrayType)
 			{
-				return {RigVMTypeUtils::DoubleArrayType};
+				return true;
 			}
-			if(InCPPType == RigVMTypeUtils::DoubleArrayType)
+			if(InCPPType == RigVMTypeUtils::DoubleArrayType && CPPType == RigVMTypeUtils::FloatArrayType)
 			{
-				return {RigVMTypeUtils::FloatArrayType};
+				return true;
 			}
-			return {};
 		}
+		return false;
+	}
 
-		FName GetCPPTypeObjectPath() const
+	static TArray<FString> GetCompatibleTypes(const FString& InCPPType)
+	{
+		if(InCPPType == RigVMTypeUtils::FloatType)
 		{
-			if(CPPTypeObject)
-			{
-				return *CPPTypeObject->GetPathName();
-			}
-			return NAME_None;
+			return {RigVMTypeUtils::DoubleType};
 		}
-
-		bool IsWildCard() const
+		if(InCPPType == RigVMTypeUtils::DoubleType)
 		{
-			return CPPTypeObject == RigVMTypeUtils::GetWildCardCPPTypeObject() ||
-				CPPType == RigVMTypeUtils::GetWildCardCPPType() ||
-				CPPType == RigVMTypeUtils::GetWildCardArrayCPPType();
+			return {RigVMTypeUtils::FloatType};
 		}
-
-		bool IsArray() const
+		if(InCPPType == RigVMTypeUtils::FloatArrayType)
 		{
-			return RigVMTypeUtils::IsArrayType(CPPType);
+			return {RigVMTypeUtils::DoubleArrayType};
 		}
-
-		FString GetBaseCPPType() const
+		if(InCPPType == RigVMTypeUtils::DoubleArrayType)
 		{
-			if(IsArray())
-			{
-				return RigVMTypeUtils::BaseTypeFromArrayType(CPPType);
-			}
-			return CPPType;
+			return {RigVMTypeUtils::FloatArrayType};
 		}
+		return {};
+	}
 
-		void ConvertToArray() 
+	FName GetCPPTypeObjectPath() const
+	{
+		if(CPPTypeObject)
 		{
-			CPPType = RigVMTypeUtils::ArrayTypeFromBaseType(CPPType);
+			return *CPPTypeObject->GetPathName();
 		}
+		return NAME_None;
+	}
+
+	bool IsWildCard() const
+	{
+		return CPPTypeObject == RigVMTypeUtils::GetWildCardCPPTypeObject() ||
+			CPPType == RigVMTypeUtils::GetWildCardCPPType() ||
+			CPPType == RigVMTypeUtils::GetWildCardArrayCPPType();
+	}
+
+	bool IsArray() const
+	{
+		return RigVMTypeUtils::IsArrayType(CPPType);
+	}
 
-		void ConvertToBaseElement() 
+	FString GetBaseCPPType() const
+	{
+		if(IsArray())
 		{
-			CPPType = RigVMTypeUtils::BaseTypeFromArrayType(CPPType);
+			return RigVMTypeUtils::BaseTypeFromArrayType(CPPType);
 		}
-	};
+		return CPPType;
+	}
+
+	void ConvertToArray() 
+	{
+		CPPType = RigVMTypeUtils::ArrayTypeFromBaseType(CPPType);
+	}
+
+	void ConvertToBaseElement() 
+	{
+		CPPType = RigVMTypeUtils::BaseTypeFromArrayType(CPPType);
+	}
+};
+
+/**
+ * The template argument represents a single parameter
+ * in a function call and all of its possible types
+ */
+USTRUCT()
+struct RIGVM_API FRigVMTemplateArgument
+{
+	GENERATED_BODY()
 	
 	enum EArrayType
 	{
@@ -175,8 +185,11 @@ struct RIGVM_API FRigVMTemplateArgument
 		ETypeCategory_Invalid
 	};
 
-	FRigVMTemplateArgument(const FName& InName, ERigVMPinDirection InDirection, const FType& InType);
-	FRigVMTemplateArgument(const FName& InName, ERigVMPinDirection InDirection, const TArray<FType>& InTypes);
+	// default constructor
+	FRigVMTemplateArgument();
+
+	FRigVMTemplateArgument(const FName& InName, ERigVMPinDirection InDirection, const FRigVMTemplateArgumentType& InType);
+	FRigVMTemplateArgument(const FName& InName, ERigVMPinDirection InDirection, const TArray<FRigVMTemplateArgumentType>& InTypes);
 
 	// returns the name of the argument
 	const FName& GetName() const { return Name; }
@@ -185,13 +198,13 @@ struct RIGVM_API FRigVMTemplateArgument
 	ERigVMPinDirection GetDirection() const { return Direction; }
 
 	// returns true if this argument supports a given type across a set of permutations
-	bool SupportsType(const FString& InCPPType, FType* OutType = nullptr) const;
+	bool SupportsType(const FString& InCPPType, FRigVMTemplateArgumentType* OutType = nullptr) const;
 
 	// returns the flat list of types (including duplicates) of this argument
-	const TArray<FType>& GetTypes() const;
+	const TArray<FRigVMTemplateArgumentType>& GetTypes() const;
 
 	// returns an array of all of the supported types
-	TArray<FType> GetSupportedTypes(const TArray<int32>& InPermutationIndices = TArray<int32>()) const;
+	TArray<FRigVMTemplateArgumentType> GetSupportedTypes(const TArray<int32>& InPermutationIndices = TArray<int32>()) const;
 
 	// returns an array of all supported types as strings
 	TArray<FString> GetSupportedTypeStrings(const TArray<int32>& InPermutationIndices = TArray<int32>()) const;
@@ -203,19 +216,26 @@ struct RIGVM_API FRigVMTemplateArgument
 	EArrayType GetArrayType() const;
 	
 	// Returns all compatible types given a category
-	static const TArray<FType> GetCompatibleTypes(ETypeCategory InCategory);
+	static const TArray<FRigVMTemplateArgumentType> GetCompatibleTypes(ETypeCategory InCategory);
 
 protected:
 
+	UPROPERTY()
 	int32 Index;
+
+	UPROPERTY()
 	FName Name;
+
+	UPROPERTY()
 	ERigVMPinDirection Direction;
-	bool bSingleton;
-	TArray<FType> Types;
-	TMap<FString, TArray<int32>> TypeToPermutations;
 
-	// default constructor
-	FRigVMTemplateArgument();
+	UPROPERTY()
+	bool bSingleton_DEPRECATED;
+
+	UPROPERTY()
+	TArray<FRigVMTemplateArgumentType> Types;
+
+	TMap<FString, TArray<int32>> TypeToPermutations;
 
 	// constructor from a property
 	FRigVMTemplateArgument(FProperty* InProperty);
@@ -231,12 +251,17 @@ protected:
  * to build polymorphic nodes (RigVMTemplateNode) that can
  * take on any of the permutations supported by the template.
  */
+USTRUCT()
 struct RIGVM_API FRigVMTemplate
 {
+	GENERATED_BODY()
 public:
 
-	typedef TMap<FName, FRigVMTemplateArgument::FType> FTypeMap;
-	typedef TPair<FName, FRigVMTemplateArgument::FType> FTypePair;
+	typedef TMap<FName, FRigVMTemplateArgumentType> FTypeMap;
+	typedef TPair<FName, FRigVMTemplateArgumentType> FTypePair;
+
+	// Default constructor
+	FRigVMTemplate();
 
 	// returns true if this is a valid template
 	bool IsValid() const;
@@ -263,7 +288,7 @@ public:
 	const FRigVMTemplateArgument* FindArgument(const FName& InArgumentName) const;
 
 	// returns true if a given arg supports a type
-	bool ArgumentSupportsType(const FName& InArgumentName, const FString& InCPPType, FRigVMTemplateArgument::FType* OutType = nullptr) const;
+	bool ArgumentSupportsType(const FName& InArgumentName, const FString& InCPPType, FRigVMTemplateArgumentType* OutType = nullptr) const;
 
 	// returns the number of permutations supported by this template
 	int32 NumPermutations() const { return Permutations.Num(); }
@@ -284,7 +309,7 @@ public:
 	bool Resolve(FTypeMap& InOutTypes, TArray<int32> & OutPermutationIndices, bool bAllowFloatingPointCasts) const;
 
 	// returns true if the template can resolve an argument to a new type
-	bool ResolveArgument(const FName& InArgumentName, const FRigVMTemplateArgument::FType& InType, FTypeMap& InOutTypes) const;
+	bool ResolveArgument(const FName& InArgumentName, const FRigVMTemplateArgumentType& InType, FTypeMap& InOutTypes) const;
 
 	// returns true if a given argument is valid for a template
 	static bool IsValidArgumentForTemplate(const FRigVMTemplateArgument& InArgument);
@@ -313,20 +338,26 @@ public:
 
 private:
 
-	// Default constructor
-	FRigVMTemplate();
-
 	// Constructor from a struct, a template name and a function index
 	FRigVMTemplate(UScriptStruct* InStruct, const FString& InTemplateName, int32 InFunctionIndex);
 
 	// Constructor from a template name, arguments and a function index
 	FRigVMTemplate(const FName& InTemplateName, const TArray<FRigVMTemplateArgument>& InArguments, int32 InFunctionIndex);
 
+	UPROPERTY()
 	int32 Index;
+
+	UPROPERTY()
 	FName Notation;
+
+	UPROPERTY()
 	TArray<FRigVMTemplateArgument> Arguments;
+
+	UPROPERTY()
 	TArray<int32> Permutations;
 
 	friend struct FRigVMRegistry;
+	friend class URigVMController;
+	friend class URigVMLibraryNode;
 };
 
