Remove unused and uninitialized bool.

diff --git a/Engine/Source/Runtime/RigVM/Public/RigVMCore/RigVMTemplate.h b/Engine/Source/Runtime/RigVM/Public/RigVMCore/RigVMTemplate.h
--- a/Engine/Source/Runtime/RigVM/Public/RigVMCore/RigVMTemplate.h
+++ b/Engine/Source/Runtime/RigVM/Public/RigVMCore/RigVMTemplate.h
@@ -229,9 +229,6 @@ protected:
 	UPROPERTY()
 	ERigVMPinDirection Direction;
 
-	UPROPERTY()
-	bool bSingleton_DEPRECATED;
-
 	UPROPERTY()
 	TArray<FRigVMTemplateArgumentType> Types;
 
