RigVM: Allow access to custom widget also on template pins + use quotes correctly for template default arguments

diff --git a/Engine/Source/Developer/RigVMDeveloper/Private/RigVMModel/Nodes/RigVMTemplateNode.cpp b/Engine/Source/Developer/RigVMDeveloper/Private/RigVMModel/Nodes/RigVMTemplateNode.cpp
--- a/Engine/Source/Developer/RigVMDeveloper/Private/RigVMModel/Nodes/RigVMTemplateNode.cpp
+++ b/Engine/Source/Developer/RigVMDeveloper/Private/RigVMModel/Nodes/RigVMTemplateNode.cpp
@@ -352,21 +352,22 @@ FString URigVMTemplateNode::GetInitialDefaultValueForPin(const FName& InRootPinN
 	for(const int32 PermutationIndex : PermutationIndices)
 	{
 		FString NewDefaultValue;
-		
-		if(const FRigVMFunction* Permutation = GetTemplate()->GetPermutation(PermutationIndex))
+
+		const FRigVMTemplate* Template = GetTemplate();
+		const FRigVMTemplateArgument* Argument = Template->FindArgument(InRootPinName);
+		const FRigVMTemplateArgumentType Type = Argument->GetTypes()[PermutationIndex];
+
+		if(const FRigVMFunction* Permutation = Template->GetPermutation(PermutationIndex))
 		{
 			const TSharedPtr<FStructOnScope> StructOnScope = MakeShareable(new FStructOnScope(Permutation->Struct));
 			const FRigVMStruct* DefaultStruct = (const FRigVMStruct*)StructOnScope->GetStructMemory();
 
+			const bool bUseQuotes = Type.CPPType != RigVMTypeUtils::FStringType && Type.CPPType != RigVMTypeUtils::FNameType;
 			NewDefaultValue = DefaultStruct->ExportToFullyQualifiedText(
-				Cast<UScriptStruct>(StructOnScope->GetStruct()), InRootPinName);
+				Cast<UScriptStruct>(StructOnScope->GetStruct()), InRootPinName, nullptr, bUseQuotes);
 		}
 		else
 		{
-			const FRigVMTemplate* Template = GetTemplate();
-			const FRigVMTemplateArgument* Argument = Template->FindArgument(InRootPinName);
-			const FRigVMTemplateArgumentType Type = Argument->GetTypes()[PermutationIndex];
-
 			if (Type.IsArray())
 			{
 				NewDefaultValue = TEXT("()");
diff --git a/Engine/Source/Developer/RigVMDeveloper/Private/RigVMModel/RigVMPin.cpp b/Engine/Source/Developer/RigVMDeveloper/Private/RigVMModel/RigVMPin.cpp
--- a/Engine/Source/Developer/RigVMDeveloper/Private/RigVMModel/RigVMPin.cpp
+++ b/Engine/Source/Developer/RigVMDeveloper/Private/RigVMModel/RigVMPin.cpp
@@ -985,6 +985,35 @@ FName URigVMPin::GetCustomWidgetName() const
 	{
 		return GetParentPin()->GetCustomWidgetName();
 	}
+
+#if WITH_EDITOR
+	if(CustomWidgetName.IsNone())
+	{
+		if(const URigVMUnitNode* UnitNode = Cast<URigVMUnitNode>(GetNode()))
+		{
+			if(const UScriptStruct* Struct = UnitNode->GetScriptStruct())
+			{
+				if(const FProperty* Property = Struct->FindPropertyByName(GetFName()))
+				{
+					const FString MetaData = Property->GetMetaData(FRigVMStruct::CustomWidgetMetaName);
+					if(!MetaData.IsEmpty())
+					{
+						return *MetaData;
+					}
+				}
+			}
+			else if(const FRigVMTemplate* Template = UnitNode->GetTemplate())
+			{
+				const FString MetaData = Template->GetArgumentMetaData(GetFName(), FRigVMStruct::CustomWidgetMetaName);
+				if(!MetaData.IsEmpty())
+				{
+					return *MetaData;
+				}
+			}
+		}
+	}
+#endif
+	
 	return CustomWidgetName;
 }
 
diff --git a/Engine/Source/Runtime/RigVM/Private/RigVMCore/RigVMStruct.cpp b/Engine/Source/Runtime/RigVM/Private/RigVMCore/RigVMStruct.cpp
--- a/Engine/Source/Runtime/RigVM/Private/RigVMCore/RigVMStruct.cpp
+++ b/Engine/Source/Runtime/RigVM/Private/RigVMCore/RigVMStruct.cpp
@@ -581,7 +581,7 @@ FString FRigVMStruct::ExportToFullyQualifiedText(const FProperty* InMemberProper
 	return DefaultValue;
 }
 
-FString FRigVMStruct::ExportToFullyQualifiedText(const UScriptStruct* InStruct, const uint8* InStructMemoryPtr)
+FString FRigVMStruct::ExportToFullyQualifiedText(const UScriptStruct* InStruct, const uint8* InStructMemoryPtr, bool bUseQuotes)
 {
 	check(InStruct);
 	check(InStructMemoryPtr);
@@ -596,7 +596,7 @@ FString FRigVMStruct::ExportToFullyQualifiedText(const UScriptStruct* InStruct,
 		
 		FString PropertyName = It->GetName();
 		const uint8* StructMemberMemoryPtr = It->ContainerPtrToValuePtr<uint8>(InStructMemoryPtr);
-		FString DefaultValue = ExportToFullyQualifiedText(*It, StructMemberMemoryPtr);
+		FString DefaultValue = ExportToFullyQualifiedText(*It, StructMemberMemoryPtr, bUseQuotes);
 		FieldValues.Add(FString::Printf(TEXT("%s=%s"), *PropertyName, *DefaultValue));
 	}
 
@@ -608,7 +608,7 @@ FString FRigVMStruct::ExportToFullyQualifiedText(const UScriptStruct* InStruct,
 	return FString::Printf(TEXT("(%s)"), *FString::Join(FieldValues, TEXT(",")));
 }
 
-FString FRigVMStruct::ExportToFullyQualifiedText(const UScriptStruct* InScriptStruct, const FName& InPropertyName, const uint8* InStructMemoryPointer) const
+FString FRigVMStruct::ExportToFullyQualifiedText(const UScriptStruct* InScriptStruct, const FName& InPropertyName, const uint8* InStructMemoryPointer, bool bUseQuotes) const
 {
 	check(InScriptStruct);
 	if(InStructMemoryPointer == nullptr)
@@ -623,7 +623,7 @@ FString FRigVMStruct::ExportToFullyQualifiedText(const UScriptStruct* InScriptSt
 	}
 
 	const uint8* StructMemberMemoryPtr = Property->ContainerPtrToValuePtr<uint8>(InStructMemoryPointer);
-	return ExportToFullyQualifiedText(Property, StructMemberMemoryPtr);
+	return ExportToFullyQualifiedText(Property, StructMemberMemoryPtr, bUseQuotes);
 }
 
 FName FRigVMStruct::GetNextAggregateName(const FName& InLastAggregatePinName) const
diff --git a/Engine/Source/Runtime/RigVM/Private/RigVMCore/RigVMTemplate.cpp b/Engine/Source/Runtime/RigVM/Private/RigVMCore/RigVMTemplate.cpp
--- a/Engine/Source/Runtime/RigVM/Private/RigVMCore/RigVMTemplate.cpp
+++ b/Engine/Source/Runtime/RigVM/Private/RigVMCore/RigVMTemplate.cpp
@@ -649,6 +649,7 @@ FLinearColor FRigVMTemplate::GetColor(const TArray<int32>& InPermutationIndices)
 		static const FName NodeColorName = TEXT("NodeColor");
 		FString NodeColorMetadata;
 
+		// if we can't find one permutation we are not going to find any, so it's ok to return false here
 		const FRigVMFunction* ResolvedFunction = GetPermutation(InPermutationIndex);
 		if(ResolvedFunction == nullptr)
 		{
@@ -710,6 +711,7 @@ FText FRigVMTemplate::GetTooltipText(const TArray<int32>& InPermutationIndices)
 			return false;
 		}
 		
+		// if we can't find one permutation we are not going to find any, so it's ok to return false here
 		const FRigVMFunction* ResolvedFunction = GetPermutation(InPermutationIndex);
 		if(ResolvedFunction == nullptr)
 		{
@@ -765,6 +767,7 @@ FText FRigVMTemplate::GetDisplayNameForArgument(const FName& InArgumentName, con
 
 		auto VisitPermutation = [InArgumentName, &ResolvedDisplayName, this](int32 InPermutationIndex) -> bool
 		{
+			// if we can't find one permutation we are not going to find any, so it's ok to return false here
 			const FRigVMFunction* ResolvedFunction = GetPermutation(InPermutationIndex);
 			if(ResolvedFunction == nullptr)
 			{
@@ -816,6 +819,66 @@ FText FRigVMTemplate::GetDisplayNameForArgument(const FName& InArgumentName, con
 	return FText();
 }
 
+FString FRigVMTemplate::GetArgumentMetaData(const FName& InArgumentName, const FName& InMetaDataKey, const TArray<int32>& InPermutationIndices) const
+{
+	if(const FRigVMTemplateArgument* Argument = FindArgument(InArgumentName))
+	{
+		FString ResolvedMetaData;
+
+		auto VisitPermutation = [InArgumentName, &ResolvedMetaData, InMetaDataKey, this](int32 InPermutationIndex) -> bool
+		{
+			// if we can't find one permutation we are not going to find any, so it's ok to return false here
+			const FRigVMFunction* ResolvedFunction = GetPermutation(InPermutationIndex);
+			if(ResolvedFunction == nullptr)
+			{
+				return false;
+			}
+			
+			if(const FProperty* Property = ResolvedFunction->Struct->FindPropertyByName(InArgumentName))
+			{
+				const FString MetaData = Property->GetMetaData(InMetaDataKey);
+				if (!ResolvedMetaData.IsEmpty())
+				{
+					if(!ResolvedMetaData.Equals(MetaData))
+					{
+						ResolvedMetaData = FString();
+						return false;
+					}
+				}
+				else
+				{
+					ResolvedMetaData = MetaData;
+				}
+			}
+			return true;
+		};
+
+		if(InPermutationIndices.IsEmpty())
+		{
+			for(int32 PermutationIndex = 0; PermutationIndex < Permutations.Num(); PermutationIndex++)
+			{
+				if(!VisitPermutation(PermutationIndex))
+				{
+					break;
+				}
+			}
+		}
+		else
+		{
+			for(const int32 PermutationIndex : InPermutationIndices)
+			{
+				if(!VisitPermutation(PermutationIndex))
+				{
+					break;
+				}
+			}
+		}
+
+		return ResolvedMetaData;
+	}
+	return FString();
+}
+
 #endif
 
 bool FRigVMTemplate::IsCompatible(const FRigVMTemplate& InOther) const
diff --git a/Engine/Source/Runtime/RigVM/Public/RigVMCore/RigVMStruct.h b/Engine/Source/Runtime/RigVM/Public/RigVMCore/RigVMStruct.h
--- a/Engine/Source/Runtime/RigVM/Public/RigVMCore/RigVMStruct.h
+++ b/Engine/Source/Runtime/RigVM/Public/RigVMCore/RigVMStruct.h
@@ -245,7 +245,7 @@ public:
 	static bool CheckFunctionExists(UScriptStruct* InStruct, const FName& FunctionName, FString* OutErrorMessage = nullptr);
 #endif
 	static FString ExportToFullyQualifiedText(const FProperty* InMemberProperty, const uint8* InMemberMemoryPtr, bool bUseQuotes = true);
-	static FString ExportToFullyQualifiedText(const UScriptStruct* InStruct, const uint8* InStructMemoryPtr);
+	static FString ExportToFullyQualifiedText(const UScriptStruct* InStruct, const uint8* InStructMemoryPtr, bool bUseQuotes = true);
 
 	template <
 		typename T,
@@ -265,7 +265,7 @@ public:
 		return ExportToFullyQualifiedText(T::StaticStruct(), (const uint8*)&InStructValue);
 	}
 
-	FString ExportToFullyQualifiedText(const UScriptStruct* InScriptStruct, const FName& InPropertyName, const uint8* InStructMemoryPointer = nullptr) const;
+	FString ExportToFullyQualifiedText(const UScriptStruct* InScriptStruct, const FName& InPropertyName, const uint8* InStructMemoryPointer = nullptr, bool bUseQuotes = true) const;
 	
 	virtual FName GetNextAggregateName(const FName& InLastAggregatePinName) const;
 	virtual FRigVMStructUpgradeInfo GetUpgradeInfo() const { return FRigVMStructUpgradeInfo(); }
diff --git a/Engine/Source/Runtime/RigVM/Public/RigVMCore/RigVMTemplate.h b/Engine/Source/Runtime/RigVM/Public/RigVMCore/RigVMTemplate.h
--- a/Engine/Source/Runtime/RigVM/Public/RigVMCore/RigVMTemplate.h
+++ b/Engine/Source/Runtime/RigVM/Public/RigVMCore/RigVMTemplate.h
@@ -328,6 +328,9 @@ public:
 	// Returns the display name text for an argument 
 	FText GetDisplayNameForArgument(const FName& InArgumentName, const TArray<int32>& InPermutationIndices = TArray<int32>()) const;
 
+	// Returns meta data on the property of the permutations 
+	FString GetArgumentMetaData(const FName& InArgumentName, const FName& InMetaDataKey, const TArray<int32>& InPermutationIndices = TArray<int32>()) const;
+
 	FString GetCategory() const;
 	FString GetKeywords() const;
 
