RigVM: Performance improvement: Convert debug breakpoints from TSharedPtr to guid based struct

diff --git a/Engine/Plugins/Animation/ControlRig/Source/ControlRig/Private/ControlRig.cpp b/Engine/Plugins/Animation/ControlRig/Source/ControlRig/Private/ControlRig.cpp
--- a/Engine/Plugins/Animation/ControlRig/Source/ControlRig/Private/ControlRig.cpp
+++ b/Engine/Plugins/Animation/ControlRig/Source/ControlRig/Private/ControlRig.cpp
@@ -1371,7 +1371,7 @@ void UControlRig::ExecuteUnits(FRigUnitContext& InOutContext, const FName& InEve
 			{
 				if(URigVM* SnapShotVM = GetSnapshotVM(false)) // don't create it for normal runs
 				{
-					if (VM->GetHaltedAtBreakpoint() != nullptr)
+					if (VM->GetHaltedAtBreakpoint().IsValid())
 					{
 						VM->CopyFrom(SnapShotVM, false, false, false, true, true);	
 					}
@@ -2885,7 +2885,7 @@ void UControlRig::AddBreakpoint(int32 InstructionIndex, URigVMNode* InNode, uint
 
 bool UControlRig::ExecuteBreakpointAction(const ERigVMBreakpointAction BreakpointAction)
 {
-	if (VM->GetHaltedAtBreakpoint() != nullptr)
+	if (VM->GetHaltedAtBreakpoint().IsValid())
 	{
 		VM->SetBreakpointAction(BreakpointAction);
 		return true;
diff --git a/Engine/Plugins/Animation/ControlRig/Source/ControlRigDeveloper/Private/ControlRigBlueprint.cpp b/Engine/Plugins/Animation/ControlRig/Source/ControlRigDeveloper/Private/ControlRigBlueprint.cpp
--- a/Engine/Plugins/Animation/ControlRig/Source/ControlRigDeveloper/Private/ControlRigBlueprint.cpp
+++ b/Engine/Plugins/Animation/ControlRig/Source/ControlRigDeveloper/Private/ControlRigBlueprint.cpp
@@ -1700,7 +1700,7 @@ void UControlRigBlueprint::RefreshControlRigBreakpoints()
 {
 	UControlRigBlueprintGeneratedClass* RigClass = GetControlRigBlueprintGeneratedClass();
 	UControlRig* CDO = Cast<UControlRig>(RigClass->GetDefaultObject(false));
-	CDO->GetDebugInfo().Clear();
+	CDO->GetDebugInfo().Reset();
 	for (URigVMNode* Node : RigVMBreakpointNodes)
 	{
 		AddBreakpointToControlRig(Node);
diff --git a/Engine/Source/Runtime/RigVM/Private/RigVMCore/RigVM.cpp b/Engine/Source/Runtime/RigVM/Private/RigVMCore/RigVM.cpp
--- a/Engine/Source/Runtime/RigVM/Private/RigVMCore/RigVM.cpp
+++ b/Engine/Source/Runtime/RigVM/Private/RigVMCore/RigVM.cpp
@@ -74,7 +74,6 @@ URigVM::URigVM()
 	, NumExecutions(0)
 #if WITH_EDITOR
 	, DebugInfo(nullptr)
-	, HaltedAtBreakpoint(nullptr)
 	, HaltedAtBreakpointHit(INDEX_NONE)
 #endif
     , FunctionNamesPtr(&FunctionNamesStorage)
@@ -894,14 +893,14 @@ bool URigVM::CanExecuteEntry(const FName& InEntryName) const
 
 bool URigVM::ResumeExecution()
 {
-	HaltedAtBreakpoint = nullptr;
+	HaltedAtBreakpoint.Reset();
 	HaltedAtBreakpointHit = INDEX_NONE;
 	if (DebugInfo)
 	{
-		if (const TSharedPtr<FRigVMBreakpoint> CurrentBreakpoint = DebugInfo->GetCurrentActiveBreakpoint())
+		if (const FRigVMBreakpoint& CurrentBreakpoint = DebugInfo->GetCurrentActiveBreakpoint())
 		{
 			DebugInfo->IncrementBreakpointActivationOnHit(CurrentBreakpoint);
-			DebugInfo->SetCurrentActiveBreakpoint(nullptr);
+			DebugInfo->SetCurrentActiveBreakpoint(FRigVMBreakpoint());
 			return true;
 		}
 	}
@@ -1258,10 +1257,20 @@ void URigVM::RebuildByteCodeOnLoad()
 #if WITH_EDITOR
 bool URigVM::ShouldHaltAtInstruction(const FName& InEventName, const uint16 InstructionIndex)
 {
+	if(DebugInfo == nullptr)
+	{
+		return false;
+	}
+
+	if(DebugInfo->IsEmpty())
+	{
+		return false;
+	}
+	
 	FRigVMByteCode& ByteCode = GetByteCode();
 
-	TArray<TSharedPtr<FRigVMBreakpoint>> BreakpointsAtInstruction = DebugInfo->FindBreakpointsAtInstruction(InstructionIndex);
-	for (TSharedPtr<FRigVMBreakpoint> Breakpoint : BreakpointsAtInstruction)
+	TArray<FRigVMBreakpoint> BreakpointsAtInstruction = DebugInfo->FindBreakpointsAtInstruction(InstructionIndex);
+	for (FRigVMBreakpoint Breakpoint : BreakpointsAtInstruction)
 	{
 		if (DebugInfo->IsActive(Breakpoint))
 		{
@@ -1281,9 +1290,9 @@ bool URigVM::ShouldHaltAtInstruction(const FName& InEventName, const uint16 Inst
 						const TArray<UObject*>* FullCallstack = ByteCode.GetCallstackForInstruction(Context.PublicData.InstructionIndex);
 						if (FullCallstack)
 						{
-							DebugInfo->SetCurrentActiveBreakpointCallstack(TArray<UObject*>(FullCallstack->GetData(), FullCallstack->Find((UObject*)Breakpoint->Subject)+1));
+							DebugInfo->SetCurrentActiveBreakpointCallstack(TArray<UObject*>(FullCallstack->GetData(), FullCallstack->Find((UObject*)Breakpoint.Subject)+1));
 						}
-						ExecutionHalted().Broadcast(Context.PublicData.InstructionIndex, Breakpoint->Subject, InEventName);
+						ExecutionHalted().Broadcast(Context.PublicData.InstructionIndex, Breakpoint.Subject, InEventName);
 					}
 					return true;
 				}
@@ -1315,7 +1324,7 @@ bool URigVM::ShouldHaltAtInstruction(const FName& InEventName, const uint16 Inst
 						// We want to keep the callstack up to the node that produced the halt
 						if (FullCallstack)
 						{
-							DebugInfo->SetCurrentActiveBreakpointCallstack(TArray<UObject*>(FullCallstack->GetData(), FullCallstack->Find((UObject*)DebugInfo->GetCurrentActiveBreakpoint()->Subject)+1));
+							DebugInfo->SetCurrentActiveBreakpointCallstack(TArray<UObject*>(FullCallstack->GetData(), FullCallstack->Find((UObject*)DebugInfo->GetCurrentActiveBreakpoint().Subject)+1));
 						}
 					}							
 					
@@ -1402,7 +1411,7 @@ bool URigVM::ShouldHaltAtInstruction(const FName& InEventName, const uint16 Inst
 				}
 
 				// Create new temporary breakpoint
-				TSharedPtr<FRigVMBreakpoint> NewBreakpoint = DebugInfo->AddBreakpoint(Context.PublicData.InstructionIndex, NewBreakpointNode, 0, true);
+				const FRigVMBreakpoint& NewBreakpoint = DebugInfo->AddBreakpoint(Context.PublicData.InstructionIndex, NewBreakpointNode, 0, true);
 				DebugInfo->SetBreakpointHits(NewBreakpoint, GetInstructionVisitedCount(Context.PublicData.InstructionIndex));
 				DebugInfo->SetBreakpointActivationOnHit(NewBreakpoint, GetInstructionVisitedCount(Context.PublicData.InstructionIndex));
 				CurrentBreakpointAction = ERigVMBreakpointAction::None;					
@@ -1757,7 +1766,7 @@ bool URigVM::Execute(TArrayView<URigVMMemoryStorage*> Memory, TArrayView<void*>
 	while (Instructions.IsValidIndex(Context.PublicData.InstructionIndex))
 	{
 #if WITH_EDITOR
-		if (DebugInfo && ShouldHaltAtInstruction(InEntryName, Context.PublicData.InstructionIndex))
+		if (ShouldHaltAtInstruction(InEntryName, Context.PublicData.InstructionIndex))
 		{
 			// we'll recursively exit all invoked
 			// entries here.
@@ -2020,10 +2029,10 @@ bool URigVM::Execute(TArrayView<URigVMMemoryStorage*> Memory, TArrayView<void*>
 #endif
 					ExecutionReachedExit().Broadcast(InEntryName);
 #if WITH_EDITOR					
-					if (HaltedAtBreakpoint != nullptr)
+					if (HaltedAtBreakpoint.IsValid())
 					{
-						HaltedAtBreakpoint = nullptr;
-						DebugInfo->SetCurrentActiveBreakpoint(nullptr);
+						HaltedAtBreakpoint.Reset();
+						DebugInfo->SetCurrentActiveBreakpoint(FRigVMBreakpoint());
 						ExecutionHalted().Broadcast(INDEX_NONE, nullptr, InEntryName);
 					}
 #endif
@@ -2602,7 +2611,7 @@ bool URigVM::Execute(TArrayView<URigVMMemoryStorage*> Memory, TArrayView<void*>
 
 #if WITH_EDITOR
 					// if we are halted at a break point we need to exit here
-					if (DebugInfo && ShouldHaltAtInstruction(InEntryName, Context.PublicData.InstructionIndex))
+					if (ShouldHaltAtInstruction(InEntryName, Context.PublicData.InstructionIndex))
 					{
 						return true;
 					}
@@ -2642,10 +2651,10 @@ bool URigVM::Execute(TArrayView<URigVMMemoryStorage*> Memory, TArrayView<void*>
 #if WITH_EDITOR
 	if(bIsRootEntry)
 	{
-		if (HaltedAtBreakpoint != nullptr)
+		if (HaltedAtBreakpoint.IsValid())
 		{
-			DebugInfo->SetCurrentActiveBreakpoint(nullptr);
-			HaltedAtBreakpoint = nullptr;
+			DebugInfo->SetCurrentActiveBreakpoint(FRigVMBreakpoint());
+			HaltedAtBreakpoint.Reset();
 			ExecutionHalted().Broadcast(INDEX_NONE, nullptr, InEntryName);
 		}
 	}
diff --git a/Engine/Source/Runtime/RigVM/Private/RigVMCore/RigVMDebugInfo.cpp b/Engine/Source/Runtime/RigVM/Private/RigVMCore/RigVMDebugInfo.cpp
--- a/Engine/Source/Runtime/RigVM/Private/RigVMCore/RigVMDebugInfo.cpp
+++ b/Engine/Source/Runtime/RigVM/Private/RigVMCore/RigVMDebugInfo.cpp
@@ -6,137 +6,162 @@ void FRigVMDebugInfo::ResetState()
 {
 	BreakpointHits.Empty();
 	BreakpointActivationOnHit.Empty();
-	TemporaryBreakpoint = nullptr;
-	CurrentActiveBreakpoint = nullptr;
+	TemporaryBreakpoint.Reset();
+	CurrentActiveBreakpoint.Invalidate();
 	CurrentActiveBreakpointCallstack.Empty();
 }
 
 void FRigVMDebugInfo::StartExecution()
 {
 	BreakpointHits.Empty();
-	CurrentActiveBreakpoint = nullptr;
+	CurrentActiveBreakpoint.Invalidate();
 	CurrentActiveBreakpointCallstack.Empty();
 }
 
-TSharedPtr<FRigVMBreakpoint> FRigVMDebugInfo::FindBreakpoint(const uint16 InInstructionIndex, const UObject* InSubject)
+const FRigVMBreakpoint& FRigVMDebugInfo::FindBreakpoint(const uint16 InInstructionIndex, const UObject* InSubject) const
 {
-	TArray<TSharedPtr<FRigVMBreakpoint>> BreakpointAtInstruction = FindBreakpointsAtInstruction(InInstructionIndex);
-	for (TSharedPtr<FRigVMBreakpoint> BP : BreakpointAtInstruction)
+	TArray<FRigVMBreakpoint> BreakpointAtInstruction = FindBreakpointsAtInstruction(InInstructionIndex);
+	for (const FRigVMBreakpoint& BP : BreakpointAtInstruction)
 	{
-		if (BP->Subject == InSubject)
+		if (BP.Subject == InSubject)
 		{
 			return BP;
 		}
 	}
-	return nullptr;
+	
+	static const FRigVMBreakpoint EmptyBreakpoint;
+	return EmptyBreakpoint;
 }
 
-TArray<TSharedPtr<FRigVMBreakpoint>> FRigVMDebugInfo::FindBreakpointsAtInstruction(const uint16 InInstructionIndex) 
+TArray<FRigVMBreakpoint> FRigVMDebugInfo::FindBreakpointsAtInstruction(const uint16 InInstructionIndex) const
 {
-	TArray<TSharedPtr<FRigVMBreakpoint>> Result;
-	for (TSharedPtr<FRigVMBreakpoint> BP : Breakpoints)
+	TArray<FRigVMBreakpoint> Result;
+	for (const FRigVMBreakpoint& BP : Breakpoints)
 	{
-		if (BP->InstructionIndex == InInstructionIndex)
+		if (BP.InstructionIndex == InInstructionIndex)
 		{
 			Result.Add(BP);
 		}
 	}
 
-	if (TemporaryBreakpoint.IsValid() && TemporaryBreakpoint->InstructionIndex == InInstructionIndex)
+	if (TemporaryBreakpoint.IsValid() && TemporaryBreakpoint.InstructionIndex == InInstructionIndex)
 	{
 		Result.Add(TemporaryBreakpoint);
 	}
+	
 	return Result;
 }
 
-TSharedPtr<FRigVMBreakpoint> FRigVMDebugInfo::AddBreakpoint(const uint16 InstructionIndex, UObject* InNode, const uint16 InDepth,
-	const bool bIsTemporary)
+const FRigVMBreakpoint& FRigVMDebugInfo::FindBreakpoint(const FGuid& InGuid) const
 {
-	for (const TSharedPtr<FRigVMBreakpoint>& BP : FindBreakpointsAtInstruction(InstructionIndex))
+	if(InGuid.IsValid())
 	{
-		if (BP->Subject == InNode)
+		if(TemporaryBreakpoint.IsValid())
 		{
-			return nullptr;
+			if(TemporaryBreakpoint.Guid == InGuid)
+			{
+				return TemporaryBreakpoint;
+			}
+		}
+
+		for(const FRigVMBreakpoint& Breakpoint : Breakpoints)
+		{
+			if(Breakpoint.Guid == InGuid)
+			{
+				return Breakpoint;
+			}
 		}
 	}
 	
+	static const FRigVMBreakpoint EmptyBreakpoint;
+	return EmptyBreakpoint;
+}
+
+const FRigVMBreakpoint& FRigVMDebugInfo::AddBreakpoint(const uint16 InstructionIndex, UObject* InNode, const uint16 InDepth,
+                                                            const bool bIsTemporary)
+{
+	static const FRigVMBreakpoint EmptyBreakpoint;
+	
+	for (const FRigVMBreakpoint& BP : FindBreakpointsAtInstruction(InstructionIndex))
+	{
+		if (BP.Subject == InNode)
+		{
+			return EmptyBreakpoint;
+		}
+	}
 	
 	if (bIsTemporary)
 	{
-		TemporaryBreakpoint = MakeShared<FRigVMBreakpoint>(InstructionIndex, InNode, 0);
+		TemporaryBreakpoint = FRigVMBreakpoint(InstructionIndex, InNode, 0);
 			
 		// Do not override the state if it already exists
-		BreakpointActivationOnHit.FindOrAdd(TemporaryBreakpoint, 0);
-		BreakpointHits.FindOrAdd(TemporaryBreakpoint, 0);
+		BreakpointActivationOnHit.FindOrAdd(TemporaryBreakpoint.Guid, 0);
+		BreakpointHits.FindOrAdd(TemporaryBreakpoint.Guid, 0);
 		return TemporaryBreakpoint;
 	}
-	else
+
+	// Breakpoints are sorted by instruction index and callstack depth
+	int32 Index = 0;
+	for(; Index<Breakpoints.Num(); ++Index)
 	{
-		// Breakpoints are sorted by instruction index and callstack depth
-		int32 Index = 0;
-		for(; Index<Breakpoints.Num(); ++Index)
+		if (InstructionIndex < Breakpoints[Index].InstructionIndex ||
+			(InstructionIndex == Breakpoints[Index].InstructionIndex && InDepth < Breakpoints[Index].Depth))
 		{
-			if (InstructionIndex < Breakpoints[Index]->InstructionIndex ||
-				(InstructionIndex == Breakpoints[Index]->InstructionIndex && InDepth < Breakpoints[Index]->Depth))
-			{
-				break;
-			}
+			break;
 		}
-		TSharedRef<FRigVMBreakpoint> NewBP = MakeShared<FRigVMBreakpoint>(InstructionIndex, InNode, InDepth);
-		Breakpoints.Insert(NewBP, Index);
-		
-		// Do not override the state if it already exists
-		BreakpointActivationOnHit.FindOrAdd(NewBP, 0);
-		BreakpointHits.FindOrAdd(NewBP, 0);
-		return NewBP;
 	}
+	const int32 InsertedIndex = Breakpoints.Insert(FRigVMBreakpoint(InstructionIndex, InNode, InDepth), Index);
+	const FRigVMBreakpoint& NewBP = Breakpoints[InsertedIndex];
 	
-	return nullptr;
+	// Do not override the state if it already exists
+	BreakpointActivationOnHit.FindOrAdd(NewBP.Guid, 0);
+	BreakpointHits.FindOrAdd(NewBP.Guid, 0);
+	return NewBP;
 }
 
-bool FRigVMDebugInfo::RemoveBreakpoint(const TSharedPtr<FRigVMBreakpoint> InBreakpoint)
+bool FRigVMDebugInfo::RemoveBreakpoint(const FRigVMBreakpoint& InBreakpoint)
 {
-	TSharedRef<FRigVMBreakpoint> BreakpointRef = InBreakpoint.ToSharedRef();
-
 	bool Found = false;
-	int32 Index = Breakpoints.Find(BreakpointRef);
+	int32 Index = Breakpoints.Find(InBreakpoint);
 	if (Index != INDEX_NONE)
 	{
-		BreakpointHits.Remove(BreakpointRef);
-		BreakpointActivationOnHit.Remove(BreakpointRef);
+		BreakpointHits.Remove(InBreakpoint.Guid);
+		BreakpointActivationOnHit.Remove(InBreakpoint.Guid);
 		Breakpoints.RemoveAt(Index);
 		Found = true;
 	}
 	if (TemporaryBreakpoint.IsValid() && InBreakpoint == TemporaryBreakpoint)
 	{
-		BreakpointHits.Remove(TemporaryBreakpoint);
-		BreakpointActivationOnHit.Remove(TemporaryBreakpoint);
-		TemporaryBreakpoint = nullptr;
+		BreakpointHits.Remove(TemporaryBreakpoint.Guid);
+		BreakpointActivationOnHit.Remove(TemporaryBreakpoint.Guid);
+		TemporaryBreakpoint.Reset();
 		Found = true;
 	}
 	return Found;
 }
 
-bool FRigVMDebugInfo::IsActive(const TSharedPtr<FRigVMBreakpoint> InBreakpoint) const
+bool FRigVMDebugInfo::IsActive(const FRigVMBreakpoint& InBreakpoint) const
 {
-	TSharedPtr<FRigVMBreakpoint> Breakpoint = InBreakpoint;
-	if (!Breakpoint && TemporaryBreakpoint.IsValid() && TemporaryBreakpoint == InBreakpoint)
+	if(!InBreakpoint.IsValid())
 	{
-		Breakpoint = TemporaryBreakpoint;
+		if(TemporaryBreakpoint.IsValid())
+		{
+			return IsActive(TemporaryBreakpoint);
+		}
 	}
-	if (Breakpoint)
+	if (InBreakpoint.IsValid())
 	{
-		if (Breakpoint->bIsActive)
+		if (InBreakpoint.bIsActive)
 		{
 			uint16 Hits = 0;
 			uint16 OnHit = 0;
-			if (BreakpointHits.Contains(Breakpoint))
+			if (BreakpointHits.Contains(InBreakpoint.Guid))
 			{
-				Hits = BreakpointHits.FindChecked(Breakpoint);
+				Hits = BreakpointHits.FindChecked(InBreakpoint.Guid);
 			}
-			if (BreakpointActivationOnHit.Contains(Breakpoint))
+			if (BreakpointActivationOnHit.Contains(InBreakpoint.Guid))
 			{
-				OnHit = BreakpointActivationOnHit.FindChecked(Breakpoint);
+				OnHit = BreakpointActivationOnHit.FindChecked(InBreakpoint.Guid);
 			}
 			return Hits == OnHit;
 		}				
@@ -144,59 +169,59 @@ bool FRigVMDebugInfo::IsActive(const TSharedPtr<FRigVMBreakpoint> InBreakpoint)
 	return false;		
 }
 
-void FRigVMDebugInfo::SetBreakpointHits(const TSharedPtr<FRigVMBreakpoint> InBreakpoint, const uint16 InBreakpointHits)
+void FRigVMDebugInfo::SetBreakpointHits(const FRigVMBreakpoint& InBreakpoint, const uint16 InBreakpointHits)
 {
-	if (BreakpointHits.Contains(InBreakpoint))
+	if (BreakpointHits.Contains(InBreakpoint.Guid))
 	{
-		BreakpointHits[InBreakpoint] = InBreakpointHits;
+		BreakpointHits[InBreakpoint.Guid] = InBreakpointHits;
 	}
 	else
 	{
-		BreakpointHits.Add(InBreakpoint, InBreakpointHits);
+		BreakpointHits.Add(InBreakpoint.Guid, InBreakpointHits);
 	}
 }
 
-void FRigVMDebugInfo::HitBreakpoint(const TSharedPtr<FRigVMBreakpoint> InBreakpoint)
+void FRigVMDebugInfo::HitBreakpoint(const FRigVMBreakpoint& InBreakpoint)
 {
-	if (BreakpointHits.Contains(InBreakpoint))
+	if (BreakpointHits.Contains(InBreakpoint.Guid))
 	{
-		BreakpointHits[InBreakpoint]++;
+		BreakpointHits[InBreakpoint.Guid]++;
 	}
 	else
 	{
-		BreakpointHits.Add(InBreakpoint, 1);
+		BreakpointHits.Add(InBreakpoint.Guid, 1);
 	}
 }
 
-void FRigVMDebugInfo::SetBreakpointActivationOnHit(const TSharedPtr<FRigVMBreakpoint> InBreakpoint, const uint16 InActivationOnHit)
+void FRigVMDebugInfo::SetBreakpointActivationOnHit(const FRigVMBreakpoint& InBreakpoint, const uint16 InActivationOnHit)
 {
-	if (BreakpointActivationOnHit.Contains(InBreakpoint))
+	if (BreakpointActivationOnHit.Contains(InBreakpoint.Guid))
 	{
-		BreakpointActivationOnHit[InBreakpoint] = InActivationOnHit;
+		BreakpointActivationOnHit[InBreakpoint.Guid] = InActivationOnHit;
 	}
 	else
 	{
-		BreakpointActivationOnHit.Add(InBreakpoint, InActivationOnHit);
+		BreakpointActivationOnHit.Add(InBreakpoint.Guid, InActivationOnHit);
 	}
 }
 
-void FRigVMDebugInfo::IncrementBreakpointActivationOnHit(const TSharedPtr<FRigVMBreakpoint> InBreakpoint)
+void FRigVMDebugInfo::IncrementBreakpointActivationOnHit(const FRigVMBreakpoint& InBreakpoint)
 {
-	if (BreakpointActivationOnHit.Contains(InBreakpoint))
+	if (BreakpointActivationOnHit.Contains(InBreakpoint.Guid))
 	{
-		BreakpointActivationOnHit[InBreakpoint]++;
+		BreakpointActivationOnHit[InBreakpoint.Guid]++;
 	}
 	else
 	{
-		BreakpointActivationOnHit.Add(InBreakpoint, 1);
+		BreakpointActivationOnHit.Add(InBreakpoint.Guid, 1);
 	}
 }
 
-uint16 FRigVMDebugInfo::GetBreakpointHits(const TSharedPtr<FRigVMBreakpoint> InBreakpoint) const
+uint16 FRigVMDebugInfo::GetBreakpointHits(const FRigVMBreakpoint& InBreakpoint) const
 {
-	if (BreakpointActivationOnHit.Contains(InBreakpoint))
+	if (BreakpointActivationOnHit.Contains(InBreakpoint.Guid))
 	{
-		return BreakpointActivationOnHit[InBreakpoint];
+		return BreakpointActivationOnHit[InBreakpoint.Guid];
 	}
 	return 0;
 }
diff --git a/Engine/Source/Runtime/RigVM/Public/RigVMCore/RigVM.h b/Engine/Source/Runtime/RigVM/Public/RigVMCore/RigVM.h
--- a/Engine/Source/Runtime/RigVM/Public/RigVMCore/RigVM.h
+++ b/Engine/Source/Runtime/RigVM/Public/RigVMCore/RigVM.h
@@ -627,7 +627,7 @@ public:
 
 	void SetDebugInfo(FRigVMDebugInfo* InDebugInfo) { DebugInfo = InDebugInfo; }
 
-	TSharedPtr<FRigVMBreakpoint> GetHaltedAtBreakpoint() const { return HaltedAtBreakpoint; }
+	const FRigVMBreakpoint& GetHaltedAtBreakpoint() const { return HaltedAtBreakpoint; }
 
 	void SetBreakpointAction(const ERigVMBreakpointAction& Action) { CurrentBreakpointAction = Action; }
 #endif
@@ -661,7 +661,7 @@ private:
 
 #if WITH_EDITOR
 	FRigVMDebugInfo* DebugInfo;
-	TSharedPtr<FRigVMBreakpoint> HaltedAtBreakpoint;
+	FRigVMBreakpoint HaltedAtBreakpoint;
 	int32 HaltedAtBreakpointHit;
 	ERigVMBreakpointAction CurrentBreakpointAction;
 
diff --git a/Engine/Source/Runtime/RigVM/Public/RigVMCore/RigVMDebugInfo.h b/Engine/Source/Runtime/RigVM/Public/RigVMCore/RigVMDebugInfo.h
--- a/Engine/Source/Runtime/RigVM/Public/RigVMCore/RigVMDebugInfo.h
+++ b/Engine/Source/Runtime/RigVM/Public/RigVMCore/RigVMDebugInfo.h
@@ -18,11 +18,12 @@ struct RIGVM_API FRigVMBreakpoint
 	, Subject(nullptr)
 	, Depth(0)
 	{
-
+		Guid.Invalidate();
 	}
 	
 	FRigVMBreakpoint(const uint16 InInstructionIndex, UObject* InNode, const uint16 InDepth)
 	: bIsActive(true)
+	, Guid(FGuid::NewGuid())
 	, InstructionIndex(InInstructionIndex)
 	, Subject(InNode)
 	, Depth(InDepth)
@@ -30,17 +31,31 @@ struct RIGVM_API FRigVMBreakpoint
 
 	}
 
-	FORCEINLINE void Clear()
+	FORCEINLINE bool IsValid() const
+	{
+		return InstructionIndex != INDEX_NONE && Guid.IsValid();
+	}
+
+	FORCEINLINE void Reset()
 	{
 		bIsActive = true;
 		InstructionIndex = INDEX_NONE;
+		Guid.Invalidate();
 		Subject = nullptr;
 		Depth = 0;
 	}
 
+	FORCEINLINE operator bool() const
+	{
+		return IsValid();
+	}
+
 	// Whether or not the breakpoint is active
 	bool bIsActive;
 
+	// guid used to identify the breakpoint
+	FGuid Guid;
+	
 	// Instruction where this breakpoint is set
 	uint16 InstructionIndex;
 
@@ -52,7 +67,12 @@ struct RIGVM_API FRigVMBreakpoint
 
 	bool operator==(const FRigVMBreakpoint& Other) const
 	{
-		return InstructionIndex == Other.InstructionIndex && Subject == Other.Subject;
+		return Guid == Other.Guid;
+	}
+
+	bool operator!=(const FRigVMBreakpoint& Other) const
+	{
+		return Guid != Other.Guid;
 	}
 };
 
@@ -62,8 +82,6 @@ struct RIGVM_API FRigVMDebugInfo
 	GENERATED_BODY()
 
 	FORCEINLINE FRigVMDebugInfo()
-		: TemporaryBreakpoint(nullptr)
-	, CurrentActiveBreakpoint(nullptr)		
 	{
 	}
 
@@ -71,33 +89,39 @@ struct RIGVM_API FRigVMDebugInfo
 
 	void StartExecution();
 
-	FORCEINLINE void Clear()
+	FORCEINLINE void Reset()
 	{
-		Breakpoints.Empty();
+		Breakpoints.Reset();
 		// Do not remove state
 	}
 
-	TSharedPtr<FRigVMBreakpoint> FindBreakpoint(const uint16 InInstructionIndex, const UObject* InSubject);
-	TArray<TSharedPtr<FRigVMBreakpoint>> FindBreakpointsAtInstruction(const uint16 InInstructionIndex);
+	FORCEINLINE bool IsEmpty() const
+	{
+		return GetBreakpoints().IsEmpty() && !TemporaryBreakpoint.IsValid();
+	}
+
+	const FRigVMBreakpoint& FindBreakpoint(const uint16 InInstructionIndex, const UObject* InSubject) const;
+	TArray<FRigVMBreakpoint> FindBreakpointsAtInstruction(const uint16 InInstructionIndex) const;
+	const FRigVMBreakpoint& FindBreakpoint(const FGuid& InGuid) const;
 	
-	TSharedPtr<FRigVMBreakpoint> AddBreakpoint(const uint16 InstructionIndex, UObject* InNode, const uint16 InDepth, const bool bIsTemporary = false);
+	const FRigVMBreakpoint& AddBreakpoint(const uint16 InstructionIndex, UObject* InNode, const uint16 InDepth, const bool bIsTemporary = false);
 
-	bool RemoveBreakpoint(const TSharedPtr<FRigVMBreakpoint> Breakpoint);
+	bool RemoveBreakpoint(const FRigVMBreakpoint& Breakpoint);
 
-	TArray<TSharedRef<FRigVMBreakpoint>> GetBreakpoints() const { return Breakpoints; }
+	const TArray<FRigVMBreakpoint>& GetBreakpoints() const { return Breakpoints; }
 
-	void SetBreakpoints(const TArray<TSharedRef<FRigVMBreakpoint>>& InBreakpoints)
+	void SetBreakpoints(const TArray<FRigVMBreakpoint>& InBreakpoints)
 	{
 		Breakpoints = InBreakpoints;
 
-		if (CurrentActiveBreakpoint.IsValid() && !Breakpoints.Contains(CurrentActiveBreakpoint))
+		if (CurrentActiveBreakpoint.IsValid() && !FindBreakpoint(CurrentActiveBreakpoint).IsValid())
 		{
-			CurrentActiveBreakpoint = nullptr;
+			CurrentActiveBreakpoint.Invalidate();
 			CurrentActiveBreakpointCallstack.Reset();
 		}
 	}
 
-	bool IsTemporaryBreakpoint(TSharedPtr<FRigVMBreakpoint> Breakpoint) const
+	bool IsTemporaryBreakpoint(FRigVMBreakpoint Breakpoint) const
 	{
 		if (Breakpoint.IsValid())
 		{
@@ -106,32 +130,42 @@ struct RIGVM_API FRigVMDebugInfo
 		return false;
 	}
 
-	bool IsActive(const TSharedPtr<FRigVMBreakpoint> InBreakpoint) const;
+	bool IsActive(const FRigVMBreakpoint& InBreakpoint) const;
 
-	void SetBreakpointHits(const TSharedPtr<FRigVMBreakpoint> InBreakpoint, const uint16 InBreakpointHits);
+	void SetBreakpointHits(const FRigVMBreakpoint& InBreakpoint, const uint16 InBreakpointHits);
 
-	void HitBreakpoint(const TSharedPtr<FRigVMBreakpoint> InBreakpoint);
+	void HitBreakpoint(const FRigVMBreakpoint& InBreakpoint);
 
-	void SetBreakpointActivationOnHit(const TSharedPtr<FRigVMBreakpoint> InBreakpoint, const uint16 InActivationOnHit);
+	void SetBreakpointActivationOnHit(const FRigVMBreakpoint& InBreakpoint, const uint16 InActivationOnHit);
 
-	void IncrementBreakpointActivationOnHit(const TSharedPtr<FRigVMBreakpoint> InBreakpoint);
+	void IncrementBreakpointActivationOnHit(const FRigVMBreakpoint& InBreakpoint);
 
-	uint16 GetBreakpointHits(const TSharedPtr<FRigVMBreakpoint> InBreakpoint) const;
+	uint16 GetBreakpointHits(const FRigVMBreakpoint& InBreakpoint) const;
 
-	TSharedPtr<FRigVMBreakpoint> GetCurrentActiveBreakpoint() const { return CurrentActiveBreakpoint; }
+	const FRigVMBreakpoint& GetCurrentActiveBreakpoint() const { return FindBreakpoint(CurrentActiveBreakpoint); }
 
-	void SetCurrentActiveBreakpoint(TSharedPtr<FRigVMBreakpoint> Breakpoint) { CurrentActiveBreakpoint = Breakpoint; }
+	FORCEINLINE void SetCurrentActiveBreakpoint(const FRigVMBreakpoint& InBreakpoint)
+	{
+		if(InBreakpoint.IsValid())
+		{
+			CurrentActiveBreakpoint = InBreakpoint.Guid;
+		}
+		else
+		{
+			CurrentActiveBreakpoint.Invalidate();
+		}
+	}
 
 	TArray<UObject*>& GetCurrentActiveBreakpointCallstack() { return CurrentActiveBreakpointCallstack; }
 
 	void SetCurrentActiveBreakpointCallstack(TArray<UObject*> Callstack) { CurrentActiveBreakpointCallstack = Callstack; }
 
 private:
-	TArray<TSharedRef<FRigVMBreakpoint>> Breakpoints;
-	TSharedPtr<FRigVMBreakpoint> TemporaryBreakpoint;
-	TMap<TSharedPtr<FRigVMBreakpoint> , uint16> BreakpointActivationOnHit; // After how many instruction executions, this breakpoint becomes active
-	TMap<TSharedPtr<FRigVMBreakpoint> , uint16> BreakpointHits; // How many times this instruction has been executed
+	TArray<FRigVMBreakpoint> Breakpoints;
+	FRigVMBreakpoint TemporaryBreakpoint;
+	TMap<FGuid, uint16> BreakpointActivationOnHit; // After how many instruction executions, this breakpoint becomes active
+	TMap<FGuid, uint16> BreakpointHits; // How many times this instruction has been executed
 
-	TSharedPtr<FRigVMBreakpoint> CurrentActiveBreakpoint;
+	FGuid CurrentActiveBreakpoint;
 	TArray<UObject*> CurrentActiveBreakpointCallstack;
 };
