Control Rig: Don't clear the RF_Public flag from the memory objects in cooked builds. Only apply this fixup in editor builds.

diff --git a/Engine/Source/Runtime/RigVM/Private/RigVMCore/RigVM.cpp b/Engine/Source/Runtime/RigVM/Private/RigVMCore/RigVM.cpp
--- a/Engine/Source/Runtime/RigVM/Private/RigVMCore/RigVM.cpp
+++ b/Engine/Source/Runtime/RigVM/Private/RigVMCore/RigVM.cpp
@@ -789,30 +789,36 @@ URigVMMemoryStorage* URigVM::GetMemoryByType(ERigVMMemoryType InMemoryType, bool
 
 void URigVM::ClearMemory()
 {
-	// At one point our memory objects were saved with RF_Public,
-	// so to truly clear them, we have to also clear the flags
+	// At one point our memory objects were saved with RF_Public, so to truly clear them, we have to also clear the flags
 	// RF_Public will make them stay around as zombie unreferenced objects, and get included in SavePackage and cooking.
 	// Clear their flags so they are not included by editor or cook SavePackage calls.
 
 	// we now make sure that only the literal memory object on the CDO is marked as RF_Public
 	// and work memory objects are no longer marked as RF_Public
-	
-	TArray<UObject*> SubObjects;
-	GetObjectsWithOuter(this, SubObjects);
-	for (UObject* SubObject : SubObjects)
+	// We don't do this for packaged builds, though.
+
+#if WITH_EDITOR
+	// Running with `-game` will set GIsEditor to nullptr.
+	if (GIsEditor)
 	{
-		if (URigVMMemoryStorage* MemoryObject = Cast<URigVMMemoryStorage>(SubObject))
+		TArray<UObject*> SubObjects;
+		GetObjectsWithOuter(this, SubObjects);
+		for (UObject* SubObject : SubObjects)
 		{
-			// we don't care about memory type here because
-			// 
-			// if "this" is not CDO, its subobjects will not include the literal memory and
-			// thus only clears the flag for work mem
-			// 
-			// if "this" is CDO, its subobjects will include the literal memory and this allows
-			// us to actually clear the literal memory
-			MemoryObject->ClearFlags(RF_Public);
+			if (URigVMMemoryStorage* MemoryObject = Cast<URigVMMemoryStorage>(SubObject))
+			{
+				// we don't care about memory type here because
+				// 
+				// if "this" is not CDO, its subobjects will not include the literal memory and
+				// thus only clears the flag for work mem
+				// 
+				// if "this" is CDO, its subobjects will include the literal memory and this allows
+				// us to actually clear the literal memory
+				MemoryObject->ClearFlags(RF_Public);
+			}
 		}
 	}
+#endif
 
 	LiteralMemoryStorageObject = nullptr;
 
