Control Rig: Construction Event

diff --git a/Engine/Source/Runtime/RigVM/Public/RigVMCore/RigVMByteCode.h b/Engine/Source/Runtime/RigVM/Public/RigVMCore/RigVMByteCode.h
--- a/Engine/Source/Runtime/RigVM/Public/RigVMCore/RigVMByteCode.h
+++ b/Engine/Source/Runtime/RigVM/Public/RigVMCore/RigVMByteCode.h
@@ -1041,6 +1041,12 @@ public:
 	// returns all of the operands for a given instruction
 	FRigVMOperandArray GetOperandsForOp(const FRigVMInstruction& InInstruction) const;
 
+	// returns all of the operands for a given instruction
+	TArray<int32> GetInstructionsForOperand(const FRigVMOperand& InOperand) const;
+
+	// returns true if the operator in question is used by multiple instructions
+	bool IsOperandShared(const FRigVMOperand& InOperand) const { return GetInstructionsForOperand(InOperand).Num() > 1; }
+
 	// returns the raw data of the byte code
 	FORCEINLINE const TArrayView<const uint8> GetByteCode() const
 	{
