Control Rig: Use TChunkedArray to guarantee stability of the location of the array elements in memory.

#rnx

diff --git a/Engine/Plugins/Animation/ControlRig/Source/ControlRigDeveloper/Private/ControlRigBlueprint.cpp b/Engine/Plugins/Animation/ControlRig/Source/ControlRigDeveloper/Private/ControlRigBlueprint.cpp
--- a/Engine/Plugins/Animation/ControlRig/Source/ControlRigDeveloper/Private/ControlRigBlueprint.cpp
+++ b/Engine/Plugins/Animation/ControlRig/Source/ControlRigDeveloper/Private/ControlRigBlueprint.cpp
@@ -2512,12 +2512,10 @@ UControlRig* UControlRigBlueprint::CreateControlRig()
 
 TArray<UStruct*> UControlRigBlueprint::GetAvailableRigUnits()
 {
-	const TArray<FRigVMFunction>& Functions = FRigVMRegistry::Get().GetFunctions();
-
 	TArray<UStruct*> Structs;
 	UStruct* BaseStruct = FRigUnit::StaticStruct();
 
-	for (const FRigVMFunction& Function : Functions)
+	for (const FRigVMFunction& Function : FRigVMRegistry::Get().GetFunctions())
 	{
 		if (Function.Struct)
 		{
diff --git a/Engine/Source/Runtime/RigVM/Private/RigVMCore/RigVMRegistry.cpp b/Engine/Source/Runtime/RigVM/Private/RigVMCore/RigVMRegistry.cpp
--- a/Engine/Source/Runtime/RigVM/Private/RigVMCore/RigVMRegistry.cpp
+++ b/Engine/Source/Runtime/RigVM/Private/RigVMCore/RigVMRegistry.cpp
@@ -25,7 +25,7 @@ void FRigVMRegistry::Register(const TCHAR* InName, FRigVMFunctionPtr InFunctionP
 	}
 
 	const FRigVMFunction Function(InName, InFunctionPtr, InStruct, Functions.Num(), InArguments);
-	Functions.Add(Function);
+	Functions.AddElement(Function);
 	FunctionNameToIndex.Add(InName, Function.Index);
 
 #if WITH_EDITOR
@@ -62,7 +62,7 @@ void FRigVMRegistry::Register(const TCHAR* InName, FRigVMFunctionPtr InFunctionP
 				{
 					Template.Index = Templates.Num();
 					Functions[Function.Index].TemplateIndex = Template.Index;
-					Templates.Add(Template);
+					Templates.AddElement(Template);
 
 					if(ExistingTemplateIndexPtr == nullptr)
 					{
@@ -94,7 +94,7 @@ const FRigVMFunction* FRigVMRegistry::FindFunction(UScriptStruct* InStruct, cons
 	return FindFunction(*FunctionName);
 }
 
-const TArray<FRigVMFunction>& FRigVMRegistry::GetFunctions() const
+const TChunkedArray<FRigVMFunction>& FRigVMRegistry::GetFunctions() const
 {
 	return Functions;
 }
@@ -114,7 +114,7 @@ const FRigVMTemplate* FRigVMRegistry::FindTemplate(const FName& InNotation) cons
 	return nullptr;
 }
 
-const TArray<FRigVMTemplate>& FRigVMRegistry::GetTemplates() const
+const TChunkedArray<FRigVMTemplate>& FRigVMRegistry::GetTemplates() const
 {
 	return Templates;
 }
@@ -192,7 +192,7 @@ const FRigVMTemplate* FRigVMRegistry::GetOrAddTemplateFromArguments(const FName&
 		Template.Permutations[Index] = INDEX_NONE;
 	}
 
-	const int32 Index = Templates.Add(Template);
+	const int32 Index = Templates.AddElement(Template);
 	Templates[Index].Index = Index;
 	TemplateNotationToIndex.Add(Template.GetNotation(), Index);
 	return &Templates[Index];
diff --git a/Engine/Source/Runtime/RigVM/Public/RigVMCore/RigVMRegistry.h b/Engine/Source/Runtime/RigVM/Public/RigVMCore/RigVMRegistry.h
--- a/Engine/Source/Runtime/RigVM/Public/RigVMCore/RigVMRegistry.h
+++ b/Engine/Source/Runtime/RigVM/Public/RigVMCore/RigVMRegistry.h
@@ -36,13 +36,13 @@ public:
 	const FRigVMFunction* FindFunction(UScriptStruct* InStruct, const TCHAR* InName) const;
 
 	// Returns all current rigvm functions
-	const TArray<FRigVMFunction>& GetFunctions() const;
+	const TChunkedArray<FRigVMFunction>& GetFunctions() const;
 
 	// Returns a template pointer given its notation (or nullptr)
 	const FRigVMTemplate* FindTemplate(const FName& InNotation) const;
 
 	// Returns all current rigvm functions
-	const TArray<FRigVMTemplate>& GetTemplates() const;
+	const TChunkedArray<FRigVMTemplate>& GetTemplates() const;
 
 	// Defines and retrieves a template given its arguments
 	const FRigVMTemplate* GetOrAddTemplateFromArguments(const FName& InName, const TArray<FRigVMTemplateArgument>& InArguments);
@@ -59,10 +59,11 @@ private:
 	FRigVMRegistry& operator= (const FRigVMRegistry &InOther) = delete;
 
 	// memory for all functions
-	TArray<FRigVMFunction> Functions;
+	// We use TChunkedArray because we need the memory locations to be stable, since we only ever add and never remove.
+	TChunkedArray<FRigVMFunction> Functions;
 
 	// memory for all templates
-	TArray<FRigVMTemplate> Templates;
+	TChunkedArray<FRigVMTemplate> Templates;
 
 	// name lookup for functions
 	TMap<FName, int32> FunctionNameToIndex;
