Control Rig: Move static initializer referencing UStructs to a static function.

#rnx

diff --git a/Engine/Source/Runtime/RigVM/Private/RigVMCore/RigVMRegistry.cpp b/Engine/Source/Runtime/RigVM/Private/RigVMCore/RigVMRegistry.cpp
--- a/Engine/Source/Runtime/RigVM/Private/RigVMCore/RigVMRegistry.cpp
+++ b/Engine/Source/Runtime/RigVM/Private/RigVMCore/RigVMRegistry.cpp
@@ -15,6 +15,26 @@ FRigVMRegistry& FRigVMRegistry::Get()
 	return s_RigVMRegistry;
 }
 
+const TArray<UScriptStruct*>& FRigVMRegistry::GetMathTypes()
+{
+	// The list of base math types to automatically register 
+	static const TArray<UScriptStruct*> MathTypes = { 
+		TBaseStructure<FRotator>::Get(),
+		TBaseStructure<FQuat>::Get(),
+		TBaseStructure<FTransform>::Get(),
+		TBaseStructure<FLinearColor>::Get(),
+		TBaseStructure<FColor>::Get(),
+		TBaseStructure<FPlane>::Get(),
+		TBaseStructure<FVector>::Get(),
+		TBaseStructure<FVector2D>::Get(),
+		TBaseStructure<FVector4>::Get(),
+		TBaseStructure<FBox2D>::Get()
+	};
+
+	return MathTypes;
+}
+
+
 void FRigVMRegistry::InitializeIfNeeded()
 {
 	if(!Types.IsEmpty())
@@ -32,9 +52,9 @@ void FRigVMRegistry::InitializeIfNeeded()
 	TypesPerCategory.Add(ETypeCategory_SingleSimpleValue, TArray<int32>()).Reserve(8);
 	TypesPerCategory.Add(ETypeCategory_ArraySimpleValue, TArray<int32>()).Reserve(8);
 	TypesPerCategory.Add(ETypeCategory_ArrayArraySimpleValue, TArray<int32>()).Reserve(8);
-	TypesPerCategory.Add(ETypeCategory_SingleMathStructValue, TArray<int32>()).Reserve(MathTypes.Num());
-	TypesPerCategory.Add(ETypeCategory_ArrayMathStructValue, TArray<int32>()).Reserve(MathTypes.Num());
-	TypesPerCategory.Add(ETypeCategory_ArrayArrayMathStructValue, TArray<int32>()).Reserve(MathTypes.Num());
+	TypesPerCategory.Add(ETypeCategory_SingleMathStructValue, TArray<int32>()).Reserve(GetMathTypes().Num());
+	TypesPerCategory.Add(ETypeCategory_ArrayMathStructValue, TArray<int32>()).Reserve(GetMathTypes().Num());
+	TypesPerCategory.Add(ETypeCategory_ArrayArrayMathStructValue, TArray<int32>()).Reserve(GetMathTypes().Num());
 	TypesPerCategory.Add(ETypeCategory_SingleScriptStructValue, TArray<int32>()).Reserve(128);
 	TypesPerCategory.Add(ETypeCategory_ArrayScriptStructValue, TArray<int32>()).Reserve(128);
 	TypesPerCategory.Add(ETypeCategory_ArrayArrayScriptStructValue, TArray<int32>()).Reserve(128);
@@ -63,7 +83,7 @@ void FRigVMRegistry::InitializeIfNeeded()
 	RigVMTypeUtils::TypeIndex::WildCardArray = FindOrAddType(FRigVMTemplateArgumentType(RigVMTypeUtils::GetWildCardArrayCPPTypeName(), RigVMTypeUtils::GetWildCardCPPTypeObject()));
 
 	// register the default math types
-	for(UScriptStruct* MathType : MathTypes)
+	for(UScriptStruct* MathType : GetMathTypes())
 	{
 		FindOrAddType(FRigVMTemplateArgumentType(*MathType->GetStructCPPName(), MathType));
 	}
@@ -217,7 +237,7 @@ int32 FRigVMRegistry::FindOrAddType(const FRigVMTemplateArgumentType& InType)
 		{
 			if(IsAllowedType(Struct))
 			{
-				if(MathTypes.Contains(CPPTypeObject))
+				if(GetMathTypes().Contains(CPPTypeObject))
 				{
 					switch(ArrayDimension)
 					{
diff --git a/Engine/Source/Runtime/RigVM/Public/RigVMCore/RigVMRegistry.h b/Engine/Source/Runtime/RigVM/Public/RigVMCore/RigVMRegistry.h
--- a/Engine/Source/Runtime/RigVM/Public/RigVMCore/RigVMRegistry.h
+++ b/Engine/Source/Runtime/RigVM/Public/RigVMCore/RigVMRegistry.h
@@ -122,19 +122,7 @@ public:
 	// Defines and retrieves a template given its arguments
 	const FRigVMTemplate* GetOrAddTemplateFromArguments(const FName& InName, const TArray<FRigVMTemplateArgument>& InArguments);
 
-	// The list of base math types to automatically register 
-	inline static const TArray<UScriptStruct*> MathTypes = { 
-		TBaseStructure<FRotator>::Get(),
-		TBaseStructure<FQuat>::Get(),
-		TBaseStructure<FTransform>::Get(),
-		TBaseStructure<FLinearColor>::Get(),
-		TBaseStructure<FColor>::Get(),
-		TBaseStructure<FPlane>::Get(),
-		TBaseStructure<FVector>::Get(),
-		TBaseStructure<FVector2D>::Get(),
-		TBaseStructure<FVector4>::Get(),
-		TBaseStructure<FBox2D>::Get()
-	};
+	static const TArray<UScriptStruct*>& GetMathTypes();
 
 private:
 
