RigVM: Register all types per function

diff --git a/Engine/Source/Runtime/RigVM/Private/RigVMCore/RigVMRegistry.cpp b/Engine/Source/Runtime/RigVM/Private/RigVMCore/RigVMRegistry.cpp
--- a/Engine/Source/Runtime/RigVM/Private/RigVMCore/RigVMRegistry.cpp
+++ b/Engine/Source/Runtime/RigVM/Private/RigVMCore/RigVMRegistry.cpp
@@ -608,6 +608,13 @@ void FRigVMRegistry::Register(const TCHAR* InName, FRigVMFunctionPtr InFunctionP
 	Functions.AddElement(Function);
 	FunctionNameToIndex.Add(InName, Function.Index);
 
+	// register all of the types used by the function
+	for (TFieldIterator<FProperty> It(InStruct); It; ++It)
+	{
+		// creating the argument causes the registration
+		FRigVMTemplateArgument Argument(*It);
+	}
+
 #if WITH_EDITOR
 	
 	FString TemplateMetadata;
