RigVM: Avoid resolval of None to Object

diff --git a/Engine/Source/Runtime/RigVM/Private/RigVMCore/RigVMMemoryDeprecated.cpp b/Engine/Source/Runtime/RigVM/Private/RigVMCore/RigVMMemoryDeprecated.cpp
--- a/Engine/Source/Runtime/RigVM/Private/RigVMCore/RigVMMemoryDeprecated.cpp
+++ b/Engine/Source/Runtime/RigVM/Private/RigVMCore/RigVMMemoryDeprecated.cpp
@@ -123,7 +123,14 @@ void FRigVMRegisterOffset::Load(FArchive& Ar)
 		FName ScriptStructPath;
 		Ar << ScriptStructPath;
 
-		ScriptStruct = UClass::TryFindTypeSlow<UScriptStruct>(ScriptStructPath.ToString());
+		if(ScriptStructPath.IsNone())
+		{
+			ScriptStruct = nullptr;
+		}
+		else
+		{
+			ScriptStruct = UClass::TryFindTypeSlow<UScriptStruct>(ScriptStructPath.ToString());
+		}
 	}
 	else
 	{
