RigVMTemplate to perform type loading lazily

diff --git a/Engine/Plugins/Animation/ControlRig/Source/ControlRigEditor/Private/Graph/SControlRigGraphChangePinType.cpp b/Engine/Plugins/Animation/ControlRig/Source/ControlRigEditor/Private/Graph/SControlRigGraphChangePinType.cpp
--- a/Engine/Plugins/Animation/ControlRig/Source/ControlRigEditor/Private/Graph/SControlRigGraphChangePinType.cpp
+++ b/Engine/Plugins/Animation/ControlRig/Source/ControlRigEditor/Private/Graph/SControlRigGraphChangePinType.cpp
@@ -304,7 +304,7 @@ void SControlRigChangePinType::FillPinTypeMenu(FMenuBuilder& MenuBuilder)
 
 		const bool bHasAllTypes =
 			SortedTypes.Num() >=
-				FRigVMRegistry::Get().GetTypesForCategory(FRigVMRegistry::ETypeCategory_SingleAnyValue).Num();
+				FRigVMRegistry::Get().GetTypesForCategory(FRigVMTemplateArgument::ETypeCategory_SingleAnyValue).Num();
 		
 		for(int32 SortedIndex=0; SortedIndex < SortedTypes.Num(); SortedIndex++)
 		{
@@ -330,7 +330,7 @@ void SControlRigChangePinType::FillPinTypeMenu(FMenuBuilder& MenuBuilder)
 					const FRigVMTemplateArgumentType MathType(*ScriptStruct->GetStructCPPName(), ScriptStruct);
 					const int32 MathTypeIndex = FRigVMRegistry::Get().GetTypeIndex(MathType);
 					
-					if(!FRigVMRegistry::Get().GetTypesForCategory(FRigVMRegistry::ETypeCategory_SingleMathStructValue).Contains(MathTypeIndex))
+					if(!FRigVMRegistry::Get().GetTypesForCategory(FRigVMTemplateArgument::ETypeCategory_SingleMathStructValue).Contains(MathTypeIndex))
 					{
 						continue;
 					}
diff --git a/Engine/Source/Developer/RigVMDeveloper/Private/RigVMModel/Nodes/RigVMArrayNode.cpp b/Engine/Source/Developer/RigVMDeveloper/Private/RigVMModel/Nodes/RigVMArrayNode.cpp
--- a/Engine/Source/Developer/RigVMDeveloper/Private/RigVMModel/Nodes/RigVMArrayNode.cpp
+++ b/Engine/Source/Developer/RigVMDeveloper/Private/RigVMModel/Nodes/RigVMArrayNode.cpp
@@ -268,6 +268,20 @@ const FRigVMTemplate* URigVMArrayNode::GetTemplate() const
 			CachedTemplate = TemplatesPerOpCode[OpCodeIndex]; 
 			return TemplatesPerOpCode[OpCodeIndex];
 		}
+
+		static const FName ExecuteFName = FRigVMStruct::ExecuteContextName;
+		static const FName ArrayFName = *ArrayName;
+		static const FName NumFName = *NumName;
+		static const FName IndexFName = *IndexName;
+		static const FName ElementFName = *ElementName;
+		static const FName SuccessFName = *SuccessName;
+		static const FName OtherFName = *OtherName;
+		static const FName CloneFName = *CloneName;
+		static const FName CountFName = *CountName;
+		static const FName RatioFName = *RatioName;
+		static const FName ResultFName = *ResultName;
+		static const FName ContinueFName = *ContinueName;
+		static const FName CompletedFName = *CompletedName;
 		
 		TArray<FRigVMTemplateArgument>& Arguments = ArgumentsPerOpCode[OpCodeIndex];
 		
@@ -275,113 +289,106 @@ const FRigVMTemplate* URigVMArrayNode::GetTemplate() const
 		{
 			Arguments.Reserve(8);
 			
-			static const TArray<int32>& SingleTypes = FRigVMRegistry::Get().GetTypesForCategory(FRigVMRegistry::ETypeCategory_SingleAnyValue);
-			static const TArray<int32>& ArrayTypes = FRigVMRegistry::Get().GetTypesForCategory(FRigVMRegistry::ETypeCategory_ArrayAnyValue);
-
-			static TArray<int32> ArrayValueTypes, ElementTypes;
-			if(ArrayValueTypes.IsEmpty())
-			{
-				ArrayValueTypes.Append(ArrayTypes);
-				ElementTypes.Append(SingleTypes);
-			}
+			static const TArray<FRigVMTemplateArgument::ETypeCategory> SingleCategory = {FRigVMTemplateArgument::ETypeCategory_SingleAnyValue};
+			static const TArray<FRigVMTemplateArgument::ETypeCategory> ArrayCategory = {FRigVMTemplateArgument::ETypeCategory_ArrayAnyValue};
 
 			UScriptStruct* ExecuteStruct = GetGraph()->GetExecuteContextStruct();
 			const FRigVMTemplateArgumentType ExecuteType(*ExecuteStruct->GetStructCPPName(), ExecuteStruct);
 			const int32 ExecuteTypeIndex = FRigVMRegistry::Get().FindOrAddType(ExecuteType);
-		
+			
 			switch(OpCode)
 			{
 				case ERigVMOpCode::ArrayReset:
 				case ERigVMOpCode::ArrayReverse:
 				{
-					Arguments.Emplace(*FRigVMStruct::ExecuteContextName.ToString(), ERigVMPinDirection::IO, ExecuteTypeIndex);
-					Arguments.Emplace(*ArrayName, ERigVMPinDirection::IO, ArrayValueTypes);
+					Arguments.Emplace(ExecuteFName, ERigVMPinDirection::IO, ExecuteTypeIndex);
+					Arguments.Emplace(ArrayFName, ERigVMPinDirection::IO, ArrayCategory);
 					break;
 				}
 				case ERigVMOpCode::ArrayGetNum:
 				{
-					Arguments.Emplace(*ArrayName, ERigVMPinDirection::Input, ArrayValueTypes);
-					Arguments.Emplace(*NumName, ERigVMPinDirection::Output, RigVMTypeUtils::TypeIndex::Int32);
+					Arguments.Emplace(ArrayFName, ERigVMPinDirection::Input, ArrayCategory);
+					Arguments.Emplace(NumFName, ERigVMPinDirection::Output, RigVMTypeUtils::TypeIndex::Int32);
 					break;
 				} 
 				case ERigVMOpCode::ArraySetNum:
 				{
-					Arguments.Emplace(*FRigVMStruct::ExecuteContextName.ToString(), ERigVMPinDirection::IO, ExecuteTypeIndex);
-					Arguments.Emplace(*ArrayName, ERigVMPinDirection::IO, ArrayValueTypes);
-					Arguments.Emplace(*NumName, ERigVMPinDirection::Input, RigVMTypeUtils::TypeIndex::Int32);
+					Arguments.Emplace(ExecuteFName, ERigVMPinDirection::IO, ExecuteTypeIndex);
+					Arguments.Emplace(ArrayFName, ERigVMPinDirection::IO, ArrayCategory);
+					Arguments.Emplace(NumFName, ERigVMPinDirection::Input, RigVMTypeUtils::TypeIndex::Int32);
 					break;
 				}
 				case ERigVMOpCode::ArrayGetAtIndex:
 				{
-					Arguments.Emplace(*ArrayName, ERigVMPinDirection::Input, ArrayValueTypes);
-					Arguments.Emplace(*IndexName, ERigVMPinDirection::Input, RigVMTypeUtils::TypeIndex::Int32);
-					Arguments.Emplace(*ElementName, ERigVMPinDirection::Output, ElementTypes);
+					Arguments.Emplace(ArrayFName, ERigVMPinDirection::Input, ArrayCategory);
+					Arguments.Emplace(IndexFName, ERigVMPinDirection::Input, RigVMTypeUtils::TypeIndex::Int32);
+					Arguments.Emplace(ElementFName, ERigVMPinDirection::Output, SingleCategory);
 					break;
 				}  
 				case ERigVMOpCode::ArraySetAtIndex:
 				case ERigVMOpCode::ArrayInsert:
 				{
-					Arguments.Emplace(*FRigVMStruct::ExecuteContextName.ToString(), ERigVMPinDirection::IO, ExecuteTypeIndex);
-					Arguments.Emplace(*ArrayName, ERigVMPinDirection::IO, ArrayValueTypes);
-					Arguments.Emplace(*IndexName, ERigVMPinDirection::Input, RigVMTypeUtils::TypeIndex::Int32);
-					Arguments.Emplace(*ElementName, ERigVMPinDirection::Input, ElementTypes);
+					Arguments.Emplace(ExecuteFName, ERigVMPinDirection::IO, ExecuteTypeIndex);
+					Arguments.Emplace(ArrayFName, ERigVMPinDirection::IO, ArrayCategory);
+					Arguments.Emplace(IndexFName, ERigVMPinDirection::Input, RigVMTypeUtils::TypeIndex::Int32);
+					Arguments.Emplace(ElementFName, ERigVMPinDirection::Input, SingleCategory);
 					break;
 				}
 				case ERigVMOpCode::ArrayAdd:
 				{
-					Arguments.Emplace(*FRigVMStruct::ExecuteContextName.ToString(), ERigVMPinDirection::IO, ExecuteTypeIndex);
-					Arguments.Emplace(*ArrayName, ERigVMPinDirection::IO, ArrayValueTypes);
-					Arguments.Emplace(*ElementName, ERigVMPinDirection::Input, ElementTypes);
-					Arguments.Emplace(*IndexName, ERigVMPinDirection::Output, RigVMTypeUtils::TypeIndex::Int32);
+					Arguments.Emplace(ExecuteFName, ERigVMPinDirection::IO, ExecuteTypeIndex);
+					Arguments.Emplace(ArrayFName, ERigVMPinDirection::IO, ArrayCategory);
+					Arguments.Emplace(ElementFName, ERigVMPinDirection::Input, SingleCategory);
+					Arguments.Emplace(IndexFName, ERigVMPinDirection::Output, RigVMTypeUtils::TypeIndex::Int32);
 					break;
 				}
 				case ERigVMOpCode::ArrayRemove:
 				{
-					Arguments.Emplace(*FRigVMStruct::ExecuteContextName.ToString(), ERigVMPinDirection::IO, ExecuteTypeIndex);
-					Arguments.Emplace(*ArrayName, ERigVMPinDirection::IO, ArrayValueTypes);
-					Arguments.Emplace(*IndexName, ERigVMPinDirection::Input, RigVMTypeUtils::TypeIndex::Int32);
+					Arguments.Emplace(ExecuteFName, ERigVMPinDirection::IO, ExecuteTypeIndex);
+					Arguments.Emplace(ArrayFName, ERigVMPinDirection::IO, ArrayCategory);
+					Arguments.Emplace(IndexFName, ERigVMPinDirection::Input, RigVMTypeUtils::TypeIndex::Int32);
 					break;
 				}
 				case ERigVMOpCode::ArrayFind:
 				{
-					Arguments.Emplace(*ArrayName, ERigVMPinDirection::Input, ArrayValueTypes);
-					Arguments.Emplace(*ElementName, ERigVMPinDirection::Input, ElementTypes);
-					Arguments.Emplace(*IndexName, ERigVMPinDirection::Output, RigVMTypeUtils::TypeIndex::Int32);
-					Arguments.Emplace(*SuccessName, ERigVMPinDirection::Output, RigVMTypeUtils::TypeIndex::Bool);
+					Arguments.Emplace(ArrayFName, ERigVMPinDirection::Input, ArrayCategory);
+					Arguments.Emplace(ElementFName, ERigVMPinDirection::Input, SingleCategory);
+					Arguments.Emplace(IndexFName, ERigVMPinDirection::Output, RigVMTypeUtils::TypeIndex::Int32);
+					Arguments.Emplace(SuccessFName, ERigVMPinDirection::Output, RigVMTypeUtils::TypeIndex::Bool);
 					break;
 				}
 				case ERigVMOpCode::ArrayAppend:
 				case ERigVMOpCode::ArrayUnion:
 				{
-					Arguments.Emplace(*FRigVMStruct::ExecuteContextName.ToString(), ERigVMPinDirection::IO, ExecuteTypeIndex);
-					Arguments.Emplace(*ArrayName, ERigVMPinDirection::IO, ArrayValueTypes);
-					Arguments.Emplace(*OtherName, ERigVMPinDirection::Input, ArrayValueTypes);
+					Arguments.Emplace(ExecuteFName, ERigVMPinDirection::IO, ExecuteTypeIndex);
+					Arguments.Emplace(ArrayFName, ERigVMPinDirection::IO, ArrayCategory);
+					Arguments.Emplace(OtherFName, ERigVMPinDirection::Input, ArrayCategory);
 					break;
 				}
 				case ERigVMOpCode::ArrayClone:
 				{
-					Arguments.Emplace(*ArrayName, ERigVMPinDirection::Input, ArrayValueTypes);
-					Arguments.Emplace(*CloneName, ERigVMPinDirection::Output, ArrayValueTypes);
+					Arguments.Emplace(ArrayFName, ERigVMPinDirection::Input, ArrayCategory);
+					Arguments.Emplace(CloneFName, ERigVMPinDirection::Output, ArrayCategory);
 					break;    				
 				}
 				case ERigVMOpCode::ArrayIterator:
 				{
-					Arguments.Emplace(*FRigVMStruct::ExecuteContextName.ToString(), ERigVMPinDirection::IO, ExecuteTypeIndex);
-					Arguments.Emplace(*ArrayName, ERigVMPinDirection::Input, ArrayValueTypes);
-					Arguments.Emplace(*ElementName, ERigVMPinDirection::Output, ElementTypes);
-					Arguments.Emplace(*IndexName, ERigVMPinDirection::Output, RigVMTypeUtils::TypeIndex::Int32);
-					Arguments.Emplace(*CountName, ERigVMPinDirection::Output, RigVMTypeUtils::TypeIndex::Int32);
-					Arguments.Emplace(*RatioName, ERigVMPinDirection::Output, RigVMTypeUtils::TypeIndex::Float);
-					Arguments.Emplace(*ContinueName, ERigVMPinDirection::Hidden, ExecuteTypeIndex);
-					Arguments.Emplace(*CompletedName, ERigVMPinDirection::Output, ExecuteTypeIndex);
+					Arguments.Emplace(ExecuteFName, ERigVMPinDirection::IO, ExecuteTypeIndex);
+					Arguments.Emplace(ArrayFName, ERigVMPinDirection::Input, ArrayCategory);
+					Arguments.Emplace(ElementFName, ERigVMPinDirection::Output, SingleCategory);
+					Arguments.Emplace(IndexFName, ERigVMPinDirection::Output, RigVMTypeUtils::TypeIndex::Int32);
+					Arguments.Emplace(CountFName, ERigVMPinDirection::Output, RigVMTypeUtils::TypeIndex::Int32);
+					Arguments.Emplace(RatioFName, ERigVMPinDirection::Output, RigVMTypeUtils::TypeIndex::Float);
+					Arguments.Emplace(ContinueFName, ERigVMPinDirection::Hidden, ExecuteTypeIndex);
+					Arguments.Emplace(CompletedFName, ERigVMPinDirection::Output, ExecuteTypeIndex);
 					break;
 				}
 				case ERigVMOpCode::ArrayDifference:
 				case ERigVMOpCode::ArrayIntersection:
 				{
-					Arguments.Emplace(*ArrayName, ERigVMPinDirection::Input, ArrayValueTypes);
-					Arguments.Emplace(*OtherName, ERigVMPinDirection::Input, ArrayValueTypes);
-					Arguments.Emplace(*ResultName, ERigVMPinDirection::Output, ArrayValueTypes);
+					Arguments.Emplace(ArrayFName, ERigVMPinDirection::Input, ArrayCategory);
+					Arguments.Emplace(OtherFName, ERigVMPinDirection::Input, ArrayCategory);
+					Arguments.Emplace(ResultFName, ERigVMPinDirection::Output, ArrayCategory);
 					break;
 				}
 				default:
@@ -391,9 +398,172 @@ const FRigVMTemplate* URigVMArrayNode::GetTemplate() const
 				}
 			}
 		}
-		
+
+		const FRigVMTemplateArgumentType ExecuteType(GetGraph()->GetExecuteContextStruct());
+		const int32 ExecuteTypeIndex = FRigVMRegistry::Get().FindOrAddType(ExecuteType);
 		const FString TemplateName = StaticEnum<ERigVMOpCode>()->GetDisplayNameTextByValue((int64)OpCode).ToString();
-		TemplatesPerOpCode[OpCodeIndex] = CachedTemplate = FRigVMRegistry::Get().GetOrAddTemplateFromArguments(*TemplateName, Arguments);
+
+		FRigVMTemplateDelegates Delegates;
+		Delegates.NewArgumentTypeDelegate = 
+			FRigVMTemplate_NewArgumentTypeDelegate::CreateLambda([OpCodeIndex, ExecuteTypeIndex](const FRigVMTemplate*, const FName& InArgumentName, int32 InTypeIndex)
+			{
+				FRigVMTemplateTypeMap Types;
+				
+				const ERigVMOpCode CurrentOpCode = (ERigVMOpCode)(OpCodeIndex + int32(ERigVMOpCode::FirstArrayOpCode));
+				const FRigVMRegistry& Registry = FRigVMRegistry::Get();
+
+				switch(CurrentOpCode)
+				{
+					case ERigVMOpCode::ArrayReset:
+					case ERigVMOpCode::ArrayReverse:
+					{
+						if(InArgumentName == ArrayFName)
+						{
+							Types.Add(ExecuteFName, ExecuteTypeIndex);
+							Types.Add(ArrayFName, InTypeIndex);
+						}
+						break;
+					}
+					case ERigVMOpCode::ArrayGetNum:
+					{
+						if(InArgumentName == ArrayFName)
+						{
+							Types.Add(ArrayFName, InTypeIndex);
+							Types.Add(NumFName, RigVMTypeUtils::TypeIndex::Int32);
+						}
+						break;
+					} 
+					case ERigVMOpCode::ArraySetNum:
+					{
+						if(InArgumentName == ArrayFName)
+						{
+							Types.Add(ExecuteFName, ExecuteTypeIndex);
+							Types.Add(ArrayFName, InTypeIndex);
+							Types.Add(NumFName, RigVMTypeUtils::TypeIndex::Int32);
+						}
+						break;
+					}
+					case ERigVMOpCode::ArrayGetAtIndex:
+					{
+						Types.Add(IndexFName, RigVMTypeUtils::TypeIndex::Int32);
+						if(InArgumentName == ArrayFName)
+						{
+							Types.Add(ArrayFName, InTypeIndex);
+							Types.Add(ElementFName, Registry.GetBaseTypeFromArrayTypeIndex(InTypeIndex));
+						}
+						else if(InArgumentName == ElementFName)
+						{
+							Types.Add(ArrayFName, Registry.GetArrayTypeFromBaseTypeIndex(InTypeIndex));
+							Types.Add(ElementFName, InTypeIndex);
+						}
+						break;
+					}  
+					case ERigVMOpCode::ArraySetAtIndex:
+					case ERigVMOpCode::ArrayInsert:
+					case ERigVMOpCode::ArrayAdd:
+					{
+						Types.Add(ExecuteFName, ExecuteTypeIndex);
+						Types.Add(IndexFName, RigVMTypeUtils::TypeIndex::Int32);
+						if(InArgumentName == ArrayFName)
+						{
+							Types.Add(ArrayFName, InTypeIndex);
+							Types.Add(ElementFName, Registry.GetBaseTypeFromArrayTypeIndex(InTypeIndex));
+						}
+						else if(InArgumentName == ElementFName)
+						{
+							Types.Add(ArrayFName, Registry.GetArrayTypeFromBaseTypeIndex(InTypeIndex));
+							Types.Add(ElementFName, InTypeIndex);
+						}
+						break;
+					}
+					case ERigVMOpCode::ArrayRemove:
+					{
+						if(InArgumentName == ArrayFName)
+						{
+							Types.Add(ExecuteFName, ExecuteTypeIndex);
+							Types.Add(ArrayFName, InTypeIndex);
+							Types.Add(IndexFName, RigVMTypeUtils::TypeIndex::Int32);
+						}
+						break;
+					}
+					case ERigVMOpCode::ArrayFind:
+					{
+						Types.Add(IndexFName, RigVMTypeUtils::TypeIndex::Int32);
+						Types.Add(SuccessFName, RigVMTypeUtils::TypeIndex::Bool);
+						if(InArgumentName == ArrayFName)
+						{
+							Types.Add(ArrayFName, InTypeIndex);
+							Types.Add(ElementFName, Registry.GetBaseTypeFromArrayTypeIndex(InTypeIndex));
+						}
+						else if(InArgumentName == ElementFName)
+						{
+							Types.Add(ArrayFName, Registry.GetArrayTypeFromBaseTypeIndex(InTypeIndex));
+							Types.Add(ElementFName, InTypeIndex);
+						}
+						break;
+					}
+					case ERigVMOpCode::ArrayAppend:
+					case ERigVMOpCode::ArrayUnion:
+					{
+						Types.Add(ExecuteFName, ExecuteTypeIndex);
+						if(InArgumentName == ArrayFName || InArgumentName == OtherFName)
+						{
+							Types.Add(ArrayFName, InTypeIndex);
+							Types.Add(OtherFName, InTypeIndex);
+						}
+						break;
+					}
+					case ERigVMOpCode::ArrayClone:
+					{
+						if(InArgumentName == ArrayFName || InArgumentName == CloneFName)
+						{
+							Types.Add(ArrayFName, InTypeIndex);
+							Types.Add(CloneFName, InTypeIndex);
+						}
+						break;    				
+					}
+					case ERigVMOpCode::ArrayIterator:
+					{
+						Types.Add(ExecuteFName, ExecuteTypeIndex);
+						Types.Add(ContinueFName, ExecuteTypeIndex);
+						Types.Add(CompletedFName, ExecuteTypeIndex);
+						Types.Add(IndexFName, RigVMTypeUtils::TypeIndex::Int32);
+						Types.Add(CountFName, RigVMTypeUtils::TypeIndex::Int32);
+						Types.Add(RatioFName, RigVMTypeUtils::TypeIndex::Float);
+						if(InArgumentName == ArrayFName)
+						{
+							Types.Add(ArrayFName, InTypeIndex);
+							Types.Add(ElementFName, Registry.GetBaseTypeFromArrayTypeIndex(InTypeIndex));
+						}
+						else if(InArgumentName == ElementFName)
+						{
+							Types.Add(ArrayFName, Registry.GetArrayTypeFromBaseTypeIndex(InTypeIndex));
+							Types.Add(ElementFName, InTypeIndex);
+						}
+						break;
+					}
+					case ERigVMOpCode::ArrayDifference:
+					case ERigVMOpCode::ArrayIntersection:
+					{
+						if(InArgumentName == ArrayFName || InArgumentName == OtherFName || InArgumentName == ResultFName)
+						{
+							Types.Add(ArrayFName, InTypeIndex);
+							Types.Add(OtherFName, InTypeIndex);
+							Types.Add(ResultFName, InTypeIndex);
+						}
+						break;
+					}
+					default:
+					{
+						ensure(false);
+						break;
+					}
+				}
+
+				return Types;
+			});
+
+		TemplatesPerOpCode[OpCodeIndex] = CachedTemplate = FRigVMRegistry::Get().GetOrAddTemplateFromArguments(*TemplateName, Arguments, Delegates);
 	}
 	return CachedTemplate;
 }
diff --git a/Engine/Source/Developer/RigVMDeveloper/Private/RigVMModel/Nodes/RigVMIfNode.cpp b/Engine/Source/Developer/RigVMDeveloper/Private/RigVMModel/Nodes/RigVMIfNode.cpp
--- a/Engine/Source/Developer/RigVMDeveloper/Private/RigVMModel/Nodes/RigVMIfNode.cpp
+++ b/Engine/Source/Developer/RigVMDeveloper/Private/RigVMModel/Nodes/RigVMIfNode.cpp
@@ -30,25 +30,44 @@ const FRigVMTemplate* URigVMIfNode::GetTemplate() const
 		{
 			return IfNodeTemplate;
 		}
-		
+
+		static const FName ConditionFName = *ConditionName;
+		static const FName TrueFName = *TrueName;
+		static const FName FalseFName = *FalseName;
+		static const FName ResultFName = *ResultName;
+
 		static TArray<FRigVMTemplateArgument> Arguments;
 		if(Arguments.IsEmpty())
 		{
-			static TArray<int32> Types;
-			if(Types.IsEmpty())
-			{
-				Types.Append(FRigVMRegistry::Get().GetTypesForCategory(FRigVMRegistry::ETypeCategory_SingleAnyValue));
-				Types.Append(FRigVMRegistry::Get().GetTypesForCategory(FRigVMRegistry::ETypeCategory_ArrayAnyValue));
-			}
-
+			static const TArray<FRigVMTemplateArgument::ETypeCategory> Categories = {
+				FRigVMTemplateArgument::ETypeCategory_SingleAnyValue,
+				FRigVMTemplateArgument::ETypeCategory_ArrayAnyValue
+			};
 			Arguments.Reserve(4);
-			Arguments.Emplace(*ConditionName, ERigVMPinDirection::Input, RigVMTypeUtils::TypeIndex::Bool);
-			Arguments.Emplace(*TrueName, ERigVMPinDirection::Input, Types);
-			Arguments.Emplace(*FalseName, ERigVMPinDirection::Input, Types);
-			Arguments.Emplace(*ResultName, ERigVMPinDirection::Output, Types);
+			Arguments.Emplace(ConditionFName, ERigVMPinDirection::Input, RigVMTypeUtils::TypeIndex::Bool);
+			Arguments.Emplace(TrueFName, ERigVMPinDirection::Input, Categories);
+			Arguments.Emplace(FalseFName, ERigVMPinDirection::Input, Categories);
+			Arguments.Emplace(ResultFName, ERigVMPinDirection::Output, Categories);
 		}
-		
-		IfNodeTemplate = CachedTemplate = FRigVMRegistry::Get().GetOrAddTemplateFromArguments(*IfName, Arguments);
+
+		FRigVMTemplateDelegates Delegates;
+		Delegates.NewArgumentTypeDelegate = 
+			FRigVMTemplate_NewArgumentTypeDelegate::CreateLambda([](const FRigVMTemplate*, const FName& InArgumentName, int32 InTypeIndex)
+			{
+				FRigVMTemplateTypeMap Types;
+
+				if(InArgumentName == TrueFName || InArgumentName == FalseFName || InArgumentName == ResultFName)
+				{
+					Types.Add(ConditionFName, RigVMTypeUtils::TypeIndex::Bool);
+					Types.Add(TrueFName, InTypeIndex);
+					Types.Add(FalseFName, InTypeIndex);
+					Types.Add(ResultFName, InTypeIndex);
+				}
+
+				return Types;
+			});
+
+		IfNodeTemplate = CachedTemplate = FRigVMRegistry::Get().GetOrAddTemplateFromArguments(*IfName, Arguments, Delegates);
 	}
 	return CachedTemplate;
 }
diff --git a/Engine/Source/Developer/RigVMDeveloper/Private/RigVMModel/Nodes/RigVMRerouteNode.cpp b/Engine/Source/Developer/RigVMDeveloper/Private/RigVMModel/Nodes/RigVMRerouteNode.cpp
--- a/Engine/Source/Developer/RigVMDeveloper/Private/RigVMModel/Nodes/RigVMRerouteNode.cpp
+++ b/Engine/Source/Developer/RigVMDeveloper/Private/RigVMModel/Nodes/RigVMRerouteNode.cpp
@@ -89,21 +89,25 @@ const FRigVMTemplate* URigVMRerouteNode::GetTemplate() const
 		static TArray<FRigVMTemplateArgument> Arguments;
 		if(Arguments.IsEmpty())
 		{
-			static TArray<int32> Types;
-			if(Types.IsEmpty())
+			static const TArray<FRigVMTemplateArgument::ETypeCategory> Categories = {
+				FRigVMTemplateArgument::ETypeCategory_Execute,
+				FRigVMTemplateArgument::ETypeCategory_SingleAnyValue,
+				FRigVMTemplateArgument::ETypeCategory_ArrayAnyValue
+			};
+			Arguments.Emplace(TEXT("Value"), ERigVMPinDirection::IO, Categories);
+		}
+
+		FRigVMTemplateDelegates Delegates;
+		Delegates.NewArgumentTypeDelegate = 
+			FRigVMTemplate_NewArgumentTypeDelegate::CreateLambda([](const FRigVMTemplate*, const FName& InArgumentName, int32 InTypeIndex)
 			{
-				Types.Append(FRigVMRegistry::Get().GetTypesForCategory(FRigVMRegistry::ETypeCategory_SingleAnyValue));
-				Types.Append(FRigVMRegistry::Get().GetTypesForCategory(FRigVMRegistry::ETypeCategory_ArrayAnyValue));
+				// since a reroute has only one argument this is simple
+				FRigVMTemplateTypeMap Types;
+				Types.Add(InArgumentName, InTypeIndex);
+				return Types;
+			});
 
-				UScriptStruct* ExecuteStruct = GetGraph()->GetExecuteContextStruct();
-				const FRigVMTemplateArgumentType ExecuteType(*ExecuteStruct->GetStructCPPName(), ExecuteStruct);
-				const int32 ExecuteTypeIndex = FRigVMRegistry::Get().FindOrAddType(ExecuteType);
-				Types.Add(ExecuteTypeIndex);
-			}
-		
-			Arguments.Emplace(TEXT("Value"), ERigVMPinDirection::IO, Types);
-		}
-		RerouteNodeTemplate = CachedTemplate = FRigVMRegistry::Get().GetOrAddTemplateFromArguments(*RerouteName, Arguments);
+		RerouteNodeTemplate = CachedTemplate = FRigVMRegistry::Get().GetOrAddTemplateFromArguments(*RerouteName, Arguments, Delegates);
 	}
 	return CachedTemplate;
 }
diff --git a/Engine/Source/Developer/RigVMDeveloper/Private/RigVMModel/Nodes/RigVMSelectNode.cpp b/Engine/Source/Developer/RigVMDeveloper/Private/RigVMModel/Nodes/RigVMSelectNode.cpp
--- a/Engine/Source/Developer/RigVMDeveloper/Private/RigVMModel/Nodes/RigVMSelectNode.cpp
+++ b/Engine/Source/Developer/RigVMDeveloper/Private/RigVMModel/Nodes/RigVMSelectNode.cpp
@@ -42,28 +42,58 @@ const FRigVMTemplate* URigVMSelectNode::GetTemplate() const
 			return SelectNodeTemplate;
 		}
 
+		static const FName IndexFName = *IndexName;
+		static const FName ValueFName = *ValueName;
+		static const FName ResultFName = *ResultName;
+
 		static TArray<FRigVMTemplateArgument> Arguments;
 		if(Arguments.IsEmpty())
 		{
-			static const TArray<int32>& SingleTypes = FRigVMRegistry::Get().GetTypesForCategory(FRigVMRegistry::ETypeCategory_SingleAnyValue);
-			static const TArray<int32>& ArrayTypes = FRigVMRegistry::Get().GetTypesForCategory(FRigVMRegistry::ETypeCategory_ArrayAnyValue);
-			static const TArray<int32>& ArrayArrayTypes = FRigVMRegistry::Get().GetTypesForCategory(FRigVMRegistry::ETypeCategory_ArrayArrayAnyValue);
-
-			static TArray<int32> ResultTypes, ValueTypes;
-			if(ResultTypes.IsEmpty())
-			{
-				ResultTypes.Append(SingleTypes);
-				ResultTypes.Append(ArrayTypes);
-				ValueTypes.Append(ArrayTypes);
-				ValueTypes.Append(ArrayArrayTypes);
-			}
-
+			static const TArray<FRigVMTemplateArgument::ETypeCategory> ValueTypeCategories = {
+				FRigVMTemplateArgument::ETypeCategory_ArrayAnyValue,
+				FRigVMTemplateArgument::ETypeCategory_ArrayArrayAnyValue
+			};
+			static const TArray<FRigVMTemplateArgument::ETypeCategory> ResultTypeCategories = {
+				FRigVMTemplateArgument::ETypeCategory_SingleAnyValue,
+				FRigVMTemplateArgument::ETypeCategory_ArrayAnyValue
+			};
 			Arguments.Reserve(3);
-			Arguments.Emplace(*IndexName, ERigVMPinDirection::Input, RigVMTypeUtils::TypeIndex::Int32);
-			Arguments.Emplace(*ValueName, ERigVMPinDirection::Input, ValueTypes);
-			Arguments.Emplace(*ResultName, ERigVMPinDirection::Output, ResultTypes);
+			Arguments.Emplace(IndexFName, ERigVMPinDirection::Input, RigVMTypeUtils::TypeIndex::Int32);
+			Arguments.Emplace(ValueFName, ERigVMPinDirection::Input, ValueTypeCategories);
+			Arguments.Emplace(ResultFName, ERigVMPinDirection::Output, ResultTypeCategories);
 		}
-		SelectNodeTemplate = CachedTemplate = FRigVMRegistry::Get().GetOrAddTemplateFromArguments(*SelectName, Arguments);
+		
+		FRigVMTemplateDelegates Delegates;
+		Delegates.NewArgumentTypeDelegate = 
+			FRigVMTemplate_NewArgumentTypeDelegate::CreateLambda([](const FRigVMTemplate*, const FName& InArgumentName, int32 InTypeIndex)
+			{
+				FRigVMTemplateTypeMap Types;
+
+				int32 ValueTypeIndex = INDEX_NONE;
+				int32 ResultTypeIndex = INDEX_NONE;
+
+				if(InArgumentName == ValueFName)
+				{
+					ValueTypeIndex = InTypeIndex;
+					ResultTypeIndex = FRigVMRegistry::Get().GetBaseTypeFromArrayTypeIndex(InTypeIndex);
+				}
+				else if(InArgumentName == ValueFName)
+				{
+					ValueTypeIndex = FRigVMRegistry::Get().GetArrayTypeFromBaseTypeIndex(InTypeIndex);;
+					ResultTypeIndex = InTypeIndex;
+				}
+				
+				if(ValueTypeIndex != INDEX_NONE && ResultTypeIndex != INDEX_NONE)
+				{
+					Types.Add(IndexFName, RigVMTypeUtils::TypeIndex::Int32);
+					Types.Add(ValueFName, ValueTypeIndex);
+					Types.Add(ResultFName, ResultTypeIndex);
+				}
+
+				return Types;
+			});
+
+		SelectNodeTemplate = CachedTemplate = FRigVMRegistry::Get().GetOrAddTemplateFromArguments(*SelectName, Arguments, Delegates);
 	}
 	return CachedTemplate;
 }
\ No newline at end of file
diff --git a/Engine/Source/Runtime/RigVM/Private/RigVMCore/RigVMRegistry.cpp b/Engine/Source/Runtime/RigVM/Private/RigVMCore/RigVMRegistry.cpp
--- a/Engine/Source/Runtime/RigVM/Private/RigVMCore/RigVMRegistry.cpp
+++ b/Engine/Source/Runtime/RigVM/Private/RigVMCore/RigVMRegistry.cpp
@@ -44,27 +44,50 @@ void FRigVMRegistry::InitializeIfNeeded()
 
 	Types.Reserve(512);
 	TypeToIndex.Reserve(512);
-	TypesPerCategory.Reserve(18);
+	TypesPerCategory.Reserve(19);
+	ArgumentsPerCategory.Reserve(19);
 	
-	TypesPerCategory.Add(ETypeCategory_SingleAnyValue, TArray<int32>()).Reserve(256);
-	TypesPerCategory.Add(ETypeCategory_ArrayAnyValue, TArray<int32>()).Reserve(256);
-	TypesPerCategory.Add(ETypeCategory_ArrayArrayAnyValue, TArray<int32>()).Reserve(256);
-	TypesPerCategory.Add(ETypeCategory_SingleSimpleValue, TArray<int32>()).Reserve(8);
-	TypesPerCategory.Add(ETypeCategory_ArraySimpleValue, TArray<int32>()).Reserve(8);
-	TypesPerCategory.Add(ETypeCategory_ArrayArraySimpleValue, TArray<int32>()).Reserve(8);
-	TypesPerCategory.Add(ETypeCategory_SingleMathStructValue, TArray<int32>()).Reserve(GetMathTypes().Num());
-	TypesPerCategory.Add(ETypeCategory_ArrayMathStructValue, TArray<int32>()).Reserve(GetMathTypes().Num());
-	TypesPerCategory.Add(ETypeCategory_ArrayArrayMathStructValue, TArray<int32>()).Reserve(GetMathTypes().Num());
-	TypesPerCategory.Add(ETypeCategory_SingleScriptStructValue, TArray<int32>()).Reserve(128);
-	TypesPerCategory.Add(ETypeCategory_ArrayScriptStructValue, TArray<int32>()).Reserve(128);
-	TypesPerCategory.Add(ETypeCategory_ArrayArrayScriptStructValue, TArray<int32>()).Reserve(128);
-	TypesPerCategory.Add(ETypeCategory_SingleEnumValue, TArray<int32>()).Reserve(128);
-	TypesPerCategory.Add(ETypeCategory_ArrayEnumValue, TArray<int32>()).Reserve(128);
-	TypesPerCategory.Add(ETypeCategory_ArrayArrayEnumValue, TArray<int32>()).Reserve(128);
-	TypesPerCategory.Add(ETypeCategory_SingleObjectValue, TArray<int32>()).Reserve(128);
-	TypesPerCategory.Add(ETypeCategory_ArrayObjectValue, TArray<int32>()).Reserve(128);
-	TypesPerCategory.Add(ETypeCategory_ArrayArrayObjectValue, TArray<int32>()).Reserve(128);
-
+	TypesPerCategory.Add(FRigVMTemplateArgument::ETypeCategory_Execute, TArray<int32>()).Reserve(8);
+	TypesPerCategory.Add(FRigVMTemplateArgument::ETypeCategory_SingleAnyValue, TArray<int32>()).Reserve(256);
+	TypesPerCategory.Add(FRigVMTemplateArgument::ETypeCategory_ArrayAnyValue, TArray<int32>()).Reserve(256);
+	TypesPerCategory.Add(FRigVMTemplateArgument::ETypeCategory_ArrayArrayAnyValue, TArray<int32>()).Reserve(256);
+	TypesPerCategory.Add(FRigVMTemplateArgument::ETypeCategory_SingleSimpleValue, TArray<int32>()).Reserve(8);
+	TypesPerCategory.Add(FRigVMTemplateArgument::ETypeCategory_ArraySimpleValue, TArray<int32>()).Reserve(8);
+	TypesPerCategory.Add(FRigVMTemplateArgument::ETypeCategory_ArrayArraySimpleValue, TArray<int32>()).Reserve(8);
+	TypesPerCategory.Add(FRigVMTemplateArgument::ETypeCategory_SingleMathStructValue, TArray<int32>()).Reserve(GetMathTypes().Num());
+	TypesPerCategory.Add(FRigVMTemplateArgument::ETypeCategory_ArrayMathStructValue, TArray<int32>()).Reserve(GetMathTypes().Num());
+	TypesPerCategory.Add(FRigVMTemplateArgument::ETypeCategory_ArrayArrayMathStructValue, TArray<int32>()).Reserve(GetMathTypes().Num());
+	TypesPerCategory.Add(FRigVMTemplateArgument::ETypeCategory_SingleScriptStructValue, TArray<int32>()).Reserve(128);
+	TypesPerCategory.Add(FRigVMTemplateArgument::ETypeCategory_ArrayScriptStructValue, TArray<int32>()).Reserve(128);
+	TypesPerCategory.Add(FRigVMTemplateArgument::ETypeCategory_ArrayArrayScriptStructValue, TArray<int32>()).Reserve(128);
+	TypesPerCategory.Add(FRigVMTemplateArgument::ETypeCategory_SingleEnumValue, TArray<int32>()).Reserve(128);
+	TypesPerCategory.Add(FRigVMTemplateArgument::ETypeCategory_ArrayEnumValue, TArray<int32>()).Reserve(128);
+	TypesPerCategory.Add(FRigVMTemplateArgument::ETypeCategory_ArrayArrayEnumValue, TArray<int32>()).Reserve(128);
+	TypesPerCategory.Add(FRigVMTemplateArgument::ETypeCategory_SingleObjectValue, TArray<int32>()).Reserve(128);
+	TypesPerCategory.Add(FRigVMTemplateArgument::ETypeCategory_ArrayObjectValue, TArray<int32>()).Reserve(128);
+	TypesPerCategory.Add(FRigVMTemplateArgument::ETypeCategory_ArrayArrayObjectValue, TArray<int32>()).Reserve(128);
+
+	ArgumentsPerCategory.Add(FRigVMTemplateArgument::ETypeCategory_Execute, TArray<TPair<int32,int32>>()).Reserve(8);
+	ArgumentsPerCategory.Add(FRigVMTemplateArgument::ETypeCategory_SingleAnyValue, TArray<TPair<int32,int32>>()).Reserve(64);
+	ArgumentsPerCategory.Add(FRigVMTemplateArgument::ETypeCategory_ArrayAnyValue, TArray<TPair<int32,int32>>()).Reserve(64);
+	ArgumentsPerCategory.Add(FRigVMTemplateArgument::ETypeCategory_ArrayArrayAnyValue, TArray<TPair<int32,int32>>()).Reserve(64);
+	ArgumentsPerCategory.Add(FRigVMTemplateArgument::ETypeCategory_SingleSimpleValue, TArray<TPair<int32,int32>>()).Reserve(64);
+	ArgumentsPerCategory.Add(FRigVMTemplateArgument::ETypeCategory_ArraySimpleValue, TArray<TPair<int32,int32>>()).Reserve(64);
+	ArgumentsPerCategory.Add(FRigVMTemplateArgument::ETypeCategory_ArrayArraySimpleValue, TArray<TPair<int32,int32>>()).Reserve(64);
+	ArgumentsPerCategory.Add(FRigVMTemplateArgument::ETypeCategory_SingleMathStructValue, TArray<TPair<int32,int32>>()).Reserve(64);
+	ArgumentsPerCategory.Add(FRigVMTemplateArgument::ETypeCategory_ArrayMathStructValue, TArray<TPair<int32,int32>>()).Reserve(64);
+	ArgumentsPerCategory.Add(FRigVMTemplateArgument::ETypeCategory_ArrayArrayMathStructValue, TArray<TPair<int32,int32>>()).Reserve(64);
+	ArgumentsPerCategory.Add(FRigVMTemplateArgument::ETypeCategory_SingleScriptStructValue, TArray<TPair<int32,int32>>()).Reserve(64);
+	ArgumentsPerCategory.Add(FRigVMTemplateArgument::ETypeCategory_ArrayScriptStructValue, TArray<TPair<int32,int32>>()).Reserve(64);
+	ArgumentsPerCategory.Add(FRigVMTemplateArgument::ETypeCategory_ArrayArrayScriptStructValue, TArray<TPair<int32,int32>>()).Reserve(64);
+	ArgumentsPerCategory.Add(FRigVMTemplateArgument::ETypeCategory_SingleEnumValue, TArray<TPair<int32,int32>>()).Reserve(64);
+	ArgumentsPerCategory.Add(FRigVMTemplateArgument::ETypeCategory_ArrayEnumValue, TArray<TPair<int32,int32>>()).Reserve(64);
+	ArgumentsPerCategory.Add(FRigVMTemplateArgument::ETypeCategory_ArrayArrayEnumValue, TArray<TPair<int32,int32>>()).Reserve(64);
+	ArgumentsPerCategory.Add(FRigVMTemplateArgument::ETypeCategory_SingleObjectValue, TArray<TPair<int32,int32>>()).Reserve(64);
+	ArgumentsPerCategory.Add(FRigVMTemplateArgument::ETypeCategory_ArrayObjectValue, TArray<TPair<int32,int32>>()).Reserve(64);
+	ArgumentsPerCategory.Add(FRigVMTemplateArgument::ETypeCategory_ArrayArrayObjectValue, TArray<TPair<int32,int32>>()).Reserve(64);
+
+	RigVMTypeUtils::TypeIndex::Execute = FindOrAddType(FRigVMTemplateArgumentType(FRigVMExecuteContext::StaticStruct()));
 	RigVMTypeUtils::TypeIndex::Bool = FindOrAddType(FRigVMTemplateArgumentType(RigVMTypeUtils::BoolTypeName, nullptr));
 	RigVMTypeUtils::TypeIndex::Float = FindOrAddType(FRigVMTemplateArgumentType(RigVMTypeUtils::FloatTypeName, nullptr));
 	RigVMTypeUtils::TypeIndex::Double = FindOrAddType(FRigVMTemplateArgumentType(RigVMTypeUtils::DoubleTypeName, nullptr));
@@ -85,7 +108,7 @@ void FRigVMRegistry::InitializeIfNeeded()
 	// register the default math types
 	for(UScriptStruct* MathType : GetMathTypes())
 	{
-		FindOrAddType(FRigVMTemplateArgumentType(*MathType->GetStructCPPName(), MathType));
+		FindOrAddType(FRigVMTemplateArgumentType(MathType));
 	}
 
 	// add all user defined structs
@@ -98,13 +121,13 @@ void FRigVMRegistry::InitializeIfNeeded()
 		}
 
 		// if this is a C++ type - skip it
-		if(ScriptStruct->IsNative())
+		if(ScriptStruct->IsNative() && !ScriptStruct->IsChildOf(FRigVMExecuteContext::StaticStruct()))
 		{
 			continue;
 		}
 
 		const FString CPPType = ScriptStruct->GetStructCPPName();
-		FindOrAddType(FRigVMTemplateArgumentType(*CPPType, ScriptStruct));
+		FindOrAddType(FRigVMTemplateArgumentType(ScriptStruct));
 	}
 
 	// add all user defined enums
@@ -136,10 +159,31 @@ int32 FRigVMRegistry::FindOrAddType(const FRigVMTemplateArgumentType& InType)
 	int32 Index = GetTypeIndex(InType);
 	if(Index == INDEX_NONE)
 	{
+		FRigVMTemplateArgumentType ElementType;
+		ElementType.CPPType = NAME_None;
+		if(InType.IsArray())
+		{
+			ElementType = InType;
+			ElementType.ConvertToBaseElement();
+		}
+		FRigVMTemplateArgumentType ArrayType = InType;
+		ArrayType.ConvertToArray();
+
 		{
 			FTypeInfo Info;
 			Info.Type = InType;
 			Info.bIsArray = InType.IsArray();
+			if(UScriptStruct* ScriptStruct = Cast<UScriptStruct>(InType.CPPTypeObject))
+			{
+				Info.bIsExecute = ScriptStruct->IsChildOf(FRigVMExecuteContext::StaticStruct());
+			}
+
+			if(!ElementType.CPPType.IsNone())
+			{
+				Info.BaseTypeIndex = GetTypeIndex(ElementType);
+			}
+			Info.ArrayTypeIndex = GetTypeIndex(ArrayType);
+			
 			Index = Types.Add(Info);
 		}
 
@@ -159,20 +203,20 @@ int32 FRigVMRegistry::FindOrAddType(const FRigVMTemplateArgumentType& InType)
 				default:
 				case 0:
 				{
-					TypesPerCategory.FindChecked(ETypeCategory_SingleSimpleValue).Add(Index);
-					TypesPerCategory.FindChecked(ETypeCategory_SingleAnyValue).Add(Index);
+					RegisterTypeInCategory(FRigVMTemplateArgument::ETypeCategory_SingleSimpleValue, Index);
+					RegisterTypeInCategory(FRigVMTemplateArgument::ETypeCategory_SingleAnyValue, Index);
 					break;
 				}
 				case 1:
 				{
-					TypesPerCategory.FindChecked(ETypeCategory_ArraySimpleValue).Add(Index);
-					TypesPerCategory.FindChecked(ETypeCategory_ArrayAnyValue).Add(Index);
+					RegisterTypeInCategory(FRigVMTemplateArgument::ETypeCategory_ArraySimpleValue, Index);
+					RegisterTypeInCategory(FRigVMTemplateArgument::ETypeCategory_ArrayAnyValue, Index);
 					break;
 				}
 				case 2:
 				{
-					TypesPerCategory.FindChecked(ETypeCategory_ArrayArraySimpleValue).Add(Index);
-					TypesPerCategory.FindChecked(ETypeCategory_ArrayArrayAnyValue).Add(Index);
+					RegisterTypeInCategory(FRigVMTemplateArgument::ETypeCategory_ArrayArraySimpleValue, Index);
+					RegisterTypeInCategory(FRigVMTemplateArgument::ETypeCategory_ArrayArrayAnyValue, Index);
 					break;
 				}
 			}
@@ -186,20 +230,20 @@ int32 FRigVMRegistry::FindOrAddType(const FRigVMTemplateArgumentType& InType)
 					default:
 					case 0:
 					{
-						TypesPerCategory.FindChecked(ETypeCategory_SingleObjectValue).Add(Index);
-						TypesPerCategory.FindChecked(ETypeCategory_SingleAnyValue).Add(Index);
+						RegisterTypeInCategory(FRigVMTemplateArgument::ETypeCategory_SingleObjectValue, Index);
+						RegisterTypeInCategory(FRigVMTemplateArgument::ETypeCategory_SingleAnyValue, Index);
 						break;
 					}
 					case 1:
 					{
-						TypesPerCategory.FindChecked(ETypeCategory_ArrayObjectValue).Add(Index);
-						TypesPerCategory.FindChecked(ETypeCategory_ArrayAnyValue).Add(Index);
+						RegisterTypeInCategory(FRigVMTemplateArgument::ETypeCategory_ArrayObjectValue, Index);
+						RegisterTypeInCategory(FRigVMTemplateArgument::ETypeCategory_ArrayAnyValue, Index);
 						break;
 					}
 					case 2:
 					{
-						TypesPerCategory.FindChecked(ETypeCategory_ArrayArrayObjectValue).Add(Index);
-						TypesPerCategory.FindChecked(ETypeCategory_ArrayArrayAnyValue).Add(Index);
+						RegisterTypeInCategory(FRigVMTemplateArgument::ETypeCategory_ArrayArrayObjectValue, Index);
+						RegisterTypeInCategory(FRigVMTemplateArgument::ETypeCategory_ArrayArrayAnyValue, Index);
 						break;
 					}
 				}
@@ -214,20 +258,20 @@ int32 FRigVMRegistry::FindOrAddType(const FRigVMTemplateArgumentType& InType)
 					default:
 					case 0:
 					{
-						TypesPerCategory.FindChecked(ETypeCategory_SingleEnumValue).Add(Index);
-						TypesPerCategory.FindChecked(ETypeCategory_SingleAnyValue).Add(Index);
+						RegisterTypeInCategory(FRigVMTemplateArgument::ETypeCategory_SingleEnumValue, Index);
+						RegisterTypeInCategory(FRigVMTemplateArgument::ETypeCategory_SingleAnyValue, Index);
 						break;
 					}
 					case 1:
 					{
-						TypesPerCategory.FindChecked(ETypeCategory_ArrayEnumValue).Add(Index);
-						TypesPerCategory.FindChecked(ETypeCategory_ArrayAnyValue).Add(Index);
+						RegisterTypeInCategory(FRigVMTemplateArgument::ETypeCategory_ArrayEnumValue, Index);
+						RegisterTypeInCategory(FRigVMTemplateArgument::ETypeCategory_ArrayAnyValue, Index);
 						break;
 					}
 					case 2:
 					{
-						TypesPerCategory.FindChecked(ETypeCategory_ArrayArrayEnumValue).Add(Index);
-						TypesPerCategory.FindChecked(ETypeCategory_ArrayArrayAnyValue).Add(Index);
+						RegisterTypeInCategory(FRigVMTemplateArgument::ETypeCategory_ArrayArrayEnumValue, Index);
+						RegisterTypeInCategory(FRigVMTemplateArgument::ETypeCategory_ArrayArrayAnyValue, Index);
 						break;
 					}
 				}
@@ -244,17 +288,17 @@ int32 FRigVMRegistry::FindOrAddType(const FRigVMTemplateArgumentType& InType)
 						default:
 						case 0:
 						{
-							TypesPerCategory.FindChecked(ETypeCategory_SingleMathStructValue).Add(Index);
+							RegisterTypeInCategory(FRigVMTemplateArgument::ETypeCategory_SingleMathStructValue, Index);
 							break;
 						}
 						case 1:
 						{
-							TypesPerCategory.FindChecked(ETypeCategory_ArrayMathStructValue).Add(Index);
+							RegisterTypeInCategory(FRigVMTemplateArgument::ETypeCategory_ArrayMathStructValue, Index);
 							break;
 						}
 						case 2:
 						{
-							TypesPerCategory.FindChecked(ETypeCategory_ArrayArrayMathStructValue).Add(Index);
+							RegisterTypeInCategory(FRigVMTemplateArgument::ETypeCategory_ArrayArrayMathStructValue, Index);
 							break;
 						}
 					}
@@ -265,31 +309,36 @@ int32 FRigVMRegistry::FindOrAddType(const FRigVMTemplateArgumentType& InType)
 					default:
 					case 0:
 					{
-						TypesPerCategory.FindChecked(ETypeCategory_SingleScriptStructValue).Add(Index);
-						TypesPerCategory.FindChecked(ETypeCategory_SingleAnyValue).Add(Index);
+						RegisterTypeInCategory(FRigVMTemplateArgument::ETypeCategory_SingleScriptStructValue, Index);
+						RegisterTypeInCategory(FRigVMTemplateArgument::ETypeCategory_SingleAnyValue, Index);
 						break;
 					}
 					case 1:
 					{
-						TypesPerCategory.FindChecked(ETypeCategory_ArrayScriptStructValue).Add(Index);
-						TypesPerCategory.FindChecked(ETypeCategory_ArrayAnyValue).Add(Index);
+						RegisterTypeInCategory(FRigVMTemplateArgument::ETypeCategory_ArrayScriptStructValue, Index);
+						RegisterTypeInCategory(FRigVMTemplateArgument::ETypeCategory_ArrayAnyValue, Index);
 						break;
 					}
 					case 2:
 					{
-						TypesPerCategory.FindChecked(ETypeCategory_ArrayArrayScriptStructValue).Add(Index);
-						TypesPerCategory.FindChecked(ETypeCategory_ArrayArrayAnyValue).Add(Index);
+						RegisterTypeInCategory(FRigVMTemplateArgument::ETypeCategory_ArrayArrayScriptStructValue, Index);
+						RegisterTypeInCategory(FRigVMTemplateArgument::ETypeCategory_ArrayArrayAnyValue, Index);
 						break;
 					}
 				}
 			}
+			else if(Struct->IsChildOf(FRigVMExecuteContext::StaticStruct()))
+			{
+				if(!Types[Index].bIsArray)
+				{
+					RegisterTypeInCategory(FRigVMTemplateArgument::ETypeCategory_Execute, Index);
+				}
+			}
 		}
 
 		// register the opposing type
 		if(!Types[Index].bIsArray)
 		{
-			FRigVMTemplateArgumentType ArrayType = InType;
-			ArrayType.ConvertToArray();
 			Types[Index].ArrayTypeIndex = FindOrAddType(ArrayType);
 			if(Types[Index].ArrayTypeIndex != INDEX_NONE)
 			{
@@ -298,8 +347,6 @@ int32 FRigVMRegistry::FindOrAddType(const FRigVMTemplateArgumentType& InType)
 		}
 		else
 		{
-			FRigVMTemplateArgumentType ElementType = InType;
-			ElementType.ConvertToBaseElement();
 			Types[Index].BaseTypeIndex = FindOrAddType(ElementType);
 			if(Types[Index].BaseTypeIndex != INDEX_NONE)
 			{
@@ -308,9 +355,7 @@ int32 FRigVMRegistry::FindOrAddType(const FRigVMTemplateArgumentType& InType)
 				// also automatically do the two dimensional array
 				if(GetArrayDimensionsForType(Index) == 1)
 				{
-					FRigVMTemplateArgumentType ArrayArrayType = InType;
-					ArrayArrayType.ConvertToArray();
-					Types[Index].ArrayTypeIndex = FindOrAddType(ArrayArrayType);
+					Types[Index].ArrayTypeIndex = FindOrAddType(ArrayType);
 					if(Types[Index].ArrayTypeIndex != INDEX_NONE)
 					{
 						Types[Types[Index].ArrayTypeIndex].BaseTypeIndex = Index;
@@ -342,6 +387,22 @@ int32 FRigVMRegistry::FindOrAddType(const FRigVMTemplateArgumentType& InType)
 	return Index;
 }
 
+void FRigVMRegistry::RegisterTypeInCategory(FRigVMTemplateArgument::ETypeCategory InCategory, int32 InTypeIndex)
+{
+	check(InCategory != FRigVMTemplateArgument::ETypeCategory_Invalid);
+
+	TypesPerCategory.FindChecked(InCategory).Add(InTypeIndex);
+
+	// when adding a new type - we need to update template arguments which expect to have access to that type 
+	const TArray<TPair<int32,int32>>& ArgumentsToUseType = ArgumentsPerCategory.FindChecked(InCategory);
+	for(const TPair<int32,int32>& Pair : ArgumentsToUseType)
+	{
+		FRigVMTemplate& Template = Templates[Pair.Key];
+		const FRigVMTemplateArgument* Argument = Template.GetArgument(Pair.Value);
+		Template.AddTypeForArgument(Argument->GetName(), InTypeIndex);
+	}
+}
+
 int32 FRigVMRegistry::GetTypeIndex(const FRigVMTemplateArgumentType& InType) const
 {
 	if(const int32* Index = TypeToIndex.Find(InType))
@@ -395,6 +456,15 @@ bool FRigVMRegistry::IsArrayType(int32 InTypeIndex) const
 	return false;
 }
 
+bool FRigVMRegistry::IsExecuteType(int32 InTypeIndex) const
+{
+	if(ensure(Types.IsValidIndex(InTypeIndex)))
+	{
+		return Types[InTypeIndex].bIsExecute;
+	}
+	return false;
+}
+
 int32 FRigVMRegistry::GetArrayDimensionsForType(int32 InTypeIndex) const
 {
 	if(ensure(Types.IsValidIndex(InTypeIndex)))
@@ -472,9 +542,9 @@ const TArray<int32>& FRigVMRegistry::GetCompatibleTypes(int32 InTypeIndex) const
 	return EmptyTypes;
 }
 
-const TArray<int32>& FRigVMRegistry::GetTypesForCategory(ETypeCategory InCategory)
+const TArray<int32>& FRigVMRegistry::GetTypesForCategory(FRigVMTemplateArgument::ETypeCategory InCategory)
 {
-	check(InCategory != ETypeCategory_Invalid);
+	check(InCategory != FRigVMTemplateArgument::ETypeCategory_Invalid);
 	return TypesPerCategory.FindChecked(InCategory);
 }
 
@@ -650,7 +720,7 @@ void FRigVMRegistry::Register(const TCHAR* InName, FRigVMFunctionPtr InFunctionP
 					Template.Index = Templates.Num();
 					Functions[Function.Index].TemplateIndex = Template.Index;
 					Templates.AddElement(Template);
-
+					
 					if(ExistingTemplateIndexPtr == nullptr)
 					{
 						TemplateNotationToIndex.Add(Template.GetNotation(), Template.Index);
@@ -706,7 +776,7 @@ const TChunkedArray<FRigVMTemplate>& FRigVMRegistry::GetTemplates() const
 	return Templates;
 }
 
-const FRigVMTemplate* FRigVMRegistry::GetOrAddTemplateFromArguments(const FName& InName, const TArray<FRigVMTemplateArgument>& InArguments)
+const FRigVMTemplate* FRigVMRegistry::GetOrAddTemplateFromArguments(const FName& InName, const TArray<FRigVMTemplateArgument>& InArguments, const FRigVMTemplateDelegates& InDelegates)
 {
 	FRigVMTemplate Template(InName, InArguments, INDEX_NONE);
 	if(const FRigVMTemplate* ExistingTemplate = FindTemplate(Template.GetNotation()))
@@ -736,11 +806,13 @@ const FRigVMTemplate* FRigVMRegistry::GetOrAddTemplateFromArguments(const FName&
 		{
 			if(IsArrayType(Argument.TypeIndices[0]))
 			{
-				Argument.TypeIndices = GetTypesForCategory(ETypeCategory_ArrayAnyValue);
+				Argument.TypeIndices = GetTypesForCategory(FRigVMTemplateArgument::ETypeCategory_ArrayAnyValue);
+				Argument.TypeCategories.Add(FRigVMTemplateArgument::ETypeCategory_ArrayAnyValue);
 			}
 			else
 			{
-				Argument.TypeIndices = GetTypesForCategory(ETypeCategory_SingleAnyValue);
+				Argument.TypeIndices = GetTypesForCategory(FRigVMTemplateArgument::ETypeCategory_SingleAnyValue);
+				Argument.TypeCategories.Add(FRigVMTemplateArgument::ETypeCategory_SingleAnyValue);
 			}
 
 			for (int32 i = 0; i < Argument.TypeIndices.Num(); ++i)
@@ -781,6 +853,16 @@ const FRigVMTemplate* FRigVMRegistry::GetOrAddTemplateFromArguments(const FName&
 
 	const int32 Index = Templates.AddElement(Template);
 	Templates[Index].Index = Index;
+	Templates[Index].OnNewArgumentType() = InDelegates.NewArgumentTypeDelegate;
 	TemplateNotationToIndex.Add(Template.GetNotation(), Index);
+
+	for(int32 ArgumentIndex=0; ArgumentIndex < Templates[Index].Arguments.Num(); ArgumentIndex++)
+	{
+		for(const FRigVMTemplateArgument::ETypeCategory& ArgumentTypeCategory : Templates[Index].Arguments[ArgumentIndex].TypeCategories)
+		{
+			ArgumentsPerCategory.FindChecked(ArgumentTypeCategory).AddUnique(TPair<int32, int32>(Index, ArgumentIndex));
+		}
+	}
+	
 	return &Templates[Index];
 }
diff --git a/Engine/Source/Runtime/RigVM/Private/RigVMCore/RigVMTemplate.cpp b/Engine/Source/Runtime/RigVM/Private/RigVMCore/RigVMTemplate.cpp
--- a/Engine/Source/Runtime/RigVM/Private/RigVMCore/RigVMTemplate.cpp
+++ b/Engine/Source/Runtime/RigVM/Private/RigVMCore/RigVMTemplate.cpp
@@ -61,7 +61,7 @@ FRigVMTemplateArgument::FRigVMTemplateArgument(const FName& InName, ERigVMPinDir
 , Direction(InDirection)
 , TypeIndices({InTypeIndex})
 {
-	TypeToPermutations.Add(InTypeIndex, {0});	
+	UpdateTypeToPermutations();
 }
 
 FRigVMTemplateArgument::FRigVMTemplateArgument(const FName& InName, ERigVMPinDirection InDirection, const TArray<int32>& InTypeIndices)
@@ -72,6 +72,25 @@ FRigVMTemplateArgument::FRigVMTemplateArgument(const FName& InName, ERigVMPinDir
 {
 	check(TypeIndices.Num() > 0);
 
+	UpdateTypeToPermutations();
+}
+
+FRigVMTemplateArgument::FRigVMTemplateArgument(const FName& InName, ERigVMPinDirection InDirection, const TArray<ETypeCategory>& InTypeCategories)
+: Index(INDEX_NONE)
+, Name(InName)
+, Direction(InDirection)
+, TypeCategories(InTypeCategories)
+{
+	for(ETypeCategory TypeCategory : TypeCategories)
+	{
+		TypeIndices.Append(FRigVMRegistry::Get().GetTypesForCategory(TypeCategory));
+	}
+	
+	UpdateTypeToPermutations();
+}
+
+void FRigVMTemplateArgument::UpdateTypeToPermutations()
+{
 	for(int32 TypeIndex=0;TypeIndex<TypeIndices.Num();TypeIndex++)
 	{
 		if (TArray<int32>* Permutations = TypeToPermutations.Find(TypeIndices[TypeIndex]))
@@ -135,6 +154,19 @@ bool FRigVMTemplateArgument::IsSingleton(const TArray<int32>& InPermutationIndic
 	return true;
 }
 
+bool FRigVMTemplateArgument::IsExecute() const
+{
+	const FRigVMRegistry& Registry = FRigVMRegistry::Get();
+	for (int32 PermutationIndex = 0; PermutationIndex < TypeIndices.Num(); PermutationIndex++)
+	{
+		if (!Registry.IsExecuteType(TypeIndices[PermutationIndex]))
+		{
+			return false;
+		}
+	}
+	return true;
+}
+
 FRigVMTemplateArgument::EArrayType FRigVMTemplateArgument::GetArrayType() const
 {
 	if(TypeIndices.Num() > 0)
@@ -1061,4 +1093,36 @@ FString FRigVMTemplate::GetKeywords() const
 	return FString::Join(KeywordsMetadata, TEXT(","));
 }
 
-#endif
\ No newline at end of file
+#endif
+
+bool FRigVMTemplate::AddTypeForArgument(const FName& InArgumentName, int32 InTypeIndex)
+{
+	if(OnNewArgumentType().IsBound())
+	{
+		const FRigVMTemplateTypeMap Types = OnNewArgumentType().Execute(this, InArgumentName, InTypeIndex);
+		if(Types.Num() == Arguments.Num())
+		{
+			for(FRigVMTemplateArgument& Argument : Arguments)
+			{
+				const int32* TypeIndex = Types.Find(Argument.Name);
+				if(TypeIndex == nullptr)
+				{
+					return false;
+				}
+				if(*TypeIndex == INDEX_NONE)
+				{
+					return false;
+				}
+			}
+			for(FRigVMTemplateArgument& Argument : Arguments)
+			{
+				const int32 TypeIndex = Types.FindChecked(Argument.Name);
+				Argument.TypeIndices.Add(TypeIndex);
+				Argument.TypeToPermutations.FindOrAdd(TypeIndex).Add(Permutations.Num());
+			}
+
+			Permutations.Add(INDEX_NONE);
+		}
+	}
+	return false;
+}
diff --git a/Engine/Source/Runtime/RigVM/Public/RigVMCore/RigVMRegistry.h b/Engine/Source/Runtime/RigVM/Public/RigVMCore/RigVMRegistry.h
--- a/Engine/Source/Runtime/RigVM/Public/RigVMCore/RigVMRegistry.h
+++ b/Engine/Source/Runtime/RigVM/Public/RigVMCore/RigVMRegistry.h
@@ -7,6 +7,11 @@
 #include "RigVMFunction.h"
 #include "RigVMTemplate.h"
 
+struct RIGVM_API FRigVMTemplateDelegates
+{
+	FRigVMTemplate_NewArgumentTypeDelegate NewArgumentTypeDelegate;
+};
+
 /**
  * The FRigVMRegistry is used to manage all known function pointers
  * for use in the RigVM. The Register method is called automatically
@@ -60,6 +65,9 @@ public:
 	// Returns true if the type is an array
 	bool IsArrayType(int32 InTypeIndex) const;
 
+	// Returns true if the type is an execute type
+	bool IsExecuteType(int32 InTypeIndex) const;
+
 	// Returns the dimensions of the array 
 	int32 GetArrayDimensionsForType(int32 InTypeIndex) const;
 
@@ -72,31 +80,8 @@ public:
 	// Returns the list of compatible types for a given type
 	const TArray<int32>& GetCompatibleTypes(int32 InTypeIndex) const;
 
-	enum ETypeCategory
-	{
-		ETypeCategory_SingleAnyValue,
-		ETypeCategory_ArrayAnyValue,
-		ETypeCategory_ArrayArrayAnyValue,
-		ETypeCategory_SingleSimpleValue,
-		ETypeCategory_ArraySimpleValue,
-		ETypeCategory_ArrayArraySimpleValue,
-		ETypeCategory_SingleMathStructValue,
-		ETypeCategory_ArrayMathStructValue,
-		ETypeCategory_ArrayArrayMathStructValue,
-		ETypeCategory_SingleScriptStructValue,
-		ETypeCategory_ArrayScriptStructValue,
-		ETypeCategory_ArrayArrayScriptStructValue,
-		ETypeCategory_SingleEnumValue,
-		ETypeCategory_ArrayEnumValue,
-		ETypeCategory_ArrayArrayEnumValue,
-		ETypeCategory_SingleObjectValue,
-		ETypeCategory_ArrayObjectValue,
-		ETypeCategory_ArrayArrayObjectValue,
-		ETypeCategory_Invalid
-	};
-
 	// Returns all compatible types given a category
-	const TArray<int32>& GetTypesForCategory(ETypeCategory InCategory);
+	const TArray<int32>& GetTypesForCategory(FRigVMTemplateArgument::ETypeCategory InCategory);
 
 	// Returns the type index of the array matching the given element type index
 	int32 GetArrayTypeFromBaseTypeIndex(int32 InTypeIndex) const;
@@ -120,7 +105,10 @@ public:
 	const TChunkedArray<FRigVMTemplate>& GetTemplates() const;
 
 	// Defines and retrieves a template given its arguments
-	const FRigVMTemplate* GetOrAddTemplateFromArguments(const FName& InName, const TArray<FRigVMTemplateArgument>& InArguments);
+	const FRigVMTemplate* GetOrAddTemplateFromArguments(
+		const FName& InName,
+		const TArray<FRigVMTemplateArgument>& InArguments,
+		const FRigVMTemplateDelegates& InDelegates);
 
 	static const TArray<UScriptStruct*>& GetMathTypes();
 
@@ -142,12 +130,14 @@ private:
 			, BaseTypeIndex(INDEX_NONE)
 			, ArrayTypeIndex(INDEX_NONE)
 			, bIsArray(false)
+			, bIsExecute(false)
 		{}
 		
 		FRigVMTemplateArgumentType Type;
 		int32 BaseTypeIndex;
 		int32 ArrayTypeIndex;
 		bool bIsArray;
+		bool bIsExecute;
 	};
 
 	FORCEINLINE static EObjectFlags DisallowedFlags()
@@ -165,6 +155,8 @@ private:
 	static bool IsAllowedType(const UStruct* InStruct);
 	static bool IsAllowedType(const UClass* InClass);
 
+	void RegisterTypeInCategory(FRigVMTemplateArgument::ETypeCategory InCategory, int32 InTypeIndex);
+
 	// memory for all (known) types
 	// We use TChunkedArray because we need the memory locations to be stable, since we only ever add and never remove.
 	TArray<FTypeInfo> Types;
@@ -184,7 +176,10 @@ private:
 	TMap<FName, int32> TemplateNotationToIndex;
 
 	// Maps storing the default types per type category
-	TMap<ETypeCategory, TArray<int32>> TypesPerCategory;
+	TMap<FRigVMTemplateArgument::ETypeCategory, TArray<int32>> TypesPerCategory;
+
+	// Lookup per type category to know which argument to keep in sync
+	TMap<FRigVMTemplateArgument::ETypeCategory, TArray<TPair<int32,int32>>> ArgumentsPerCategory;
 
 	static FRigVMRegistry s_RigVMRegistry;
 	friend struct FRigVMStruct;
diff --git a/Engine/Source/Runtime/RigVM/Public/RigVMCore/RigVMTemplate.h b/Engine/Source/Runtime/RigVM/Public/RigVMCore/RigVMTemplate.h
--- a/Engine/Source/Runtime/RigVM/Public/RigVMCore/RigVMTemplate.h
+++ b/Engine/Source/Runtime/RigVM/Public/RigVMCore/RigVMTemplate.h
@@ -8,6 +8,11 @@
 #include "RigVMTypeUtils.h"
 #include "RigVMTemplate.generated.h"
 
+struct FRigVMTemplate;
+typedef TMap<FName, int32> FRigVMTemplateTypeMap;
+
+DECLARE_DELEGATE_RetVal_ThreeParams(FRigVMTemplateTypeMap, FRigVMTemplate_NewArgumentTypeDelegate, const FRigVMTemplate* /* InTemplate */, const FName& /* InArgumentName */, int32 /* InTypeIndexToAdd */);
+
 USTRUCT()
 struct RIGVM_API FRigVMTemplateArgumentType
 {
@@ -34,6 +39,12 @@ struct RIGVM_API FRigVMTemplateArgumentType
 		check(!CPPType.IsNone());
 	}
 
+	FRigVMTemplateArgumentType(UScriptStruct* InScriptStruct)
+	: CPPType(*InScriptStruct->GetStructCPPName())
+	, CPPTypeObject(InScriptStruct)
+	{
+	}
+
 	static FRigVMTemplateArgumentType Array()
 	{
 		return FRigVMTemplateArgumentType(RigVMTypeUtils::GetWildCardArrayCPPTypeName(), RigVMTypeUtils::GetWildCardCPPTypeObject());
@@ -113,11 +124,36 @@ struct RIGVM_API FRigVMTemplateArgument
 		EArrayType_Invalid
 	};
 
+	enum ETypeCategory
+	{
+		ETypeCategory_Execute,
+		ETypeCategory_SingleAnyValue,
+		ETypeCategory_ArrayAnyValue,
+		ETypeCategory_ArrayArrayAnyValue,
+		ETypeCategory_SingleSimpleValue,
+		ETypeCategory_ArraySimpleValue,
+		ETypeCategory_ArrayArraySimpleValue,
+		ETypeCategory_SingleMathStructValue,
+		ETypeCategory_ArrayMathStructValue,
+		ETypeCategory_ArrayArrayMathStructValue,
+		ETypeCategory_SingleScriptStructValue,
+		ETypeCategory_ArrayScriptStructValue,
+		ETypeCategory_ArrayArrayScriptStructValue,
+		ETypeCategory_SingleEnumValue,
+		ETypeCategory_ArrayEnumValue,
+		ETypeCategory_ArrayArrayEnumValue,
+		ETypeCategory_SingleObjectValue,
+		ETypeCategory_ArrayObjectValue,
+		ETypeCategory_ArrayArrayObjectValue,
+		ETypeCategory_Invalid
+	};
+
 	// default constructor
 	FRigVMTemplateArgument();
 
 	FRigVMTemplateArgument(const FName& InName, ERigVMPinDirection InDirection, int32 InType);
 	FRigVMTemplateArgument(const FName& InName, ERigVMPinDirection InDirection, const TArray<int32>& InTypeIndices);
+	FRigVMTemplateArgument(const FName& InName, ERigVMPinDirection InDirection, const TArray<ETypeCategory>& InTypeCategories);
 
 	// returns the name of the argument
 	const FName& GetName() const { return Name; }
@@ -142,9 +178,12 @@ struct RIGVM_API FRigVMTemplateArgument
 	// returns true if an argument is singleton (same type for all variants)
 	bool IsSingleton(const TArray<int32>& InPermutationIndices = TArray<int32>()) const;
 
+	// returns true if this argument is an execute
+	bool IsExecute() const;
+
 	// returns true if the argument uses an array container
 	EArrayType GetArrayType() const;
-	
+
 protected:
 
 	UPROPERTY()
@@ -160,10 +199,13 @@ protected:
 	TArray<int32> TypeIndices;
 
 	TMap<int32, TArray<int32>> TypeToPermutations;
+	TArray<ETypeCategory> TypeCategories;
 
 	// constructor from a property
 	FRigVMTemplateArgument(FProperty* InProperty);
 
+	void UpdateTypeToPermutations();
+
 	friend struct FRigVMTemplate;
 	friend class URigVMController;
 	friend struct FRigVMRegistry;
@@ -181,7 +223,7 @@ struct RIGVM_API FRigVMTemplate
 	GENERATED_BODY()
 public:
 
-	typedef TMap<FName, int32> FTypeMap;
+	typedef FRigVMTemplateTypeMap FTypeMap;
 	typedef TPair<FName, int32> FTypePair;
 
 	// Default constructor
@@ -266,6 +308,13 @@ public:
 
 #endif
 
+	// Adds a new argument to the template. The template needs to rely on the delegate to ask the
+	// template factory how to deal with the new type.
+	bool AddTypeForArgument(const FName& InArgumentName, int32 InTypeIndex);
+
+	// Returns the delegate to be able to react to type changes dynamically
+	FRigVMTemplate_NewArgumentTypeDelegate& OnNewArgumentType() { return NewArgumentTypeDelegate; }
+
 private:
 
 	// Constructor from a struct, a template name and a function index
@@ -286,6 +335,8 @@ private:
 	UPROPERTY()
 	TArray<int32> Permutations;
 
+	FRigVMTemplate_NewArgumentTypeDelegate NewArgumentTypeDelegate; 
+
 	friend struct FRigVMRegistry;
 	friend class URigVMController;
 	friend class URigVMLibraryNode;
diff --git a/Engine/Source/Runtime/RigVM/Public/RigVMTypeUtils.h b/Engine/Source/Runtime/RigVM/Public/RigVMTypeUtils.h
--- a/Engine/Source/Runtime/RigVM/Public/RigVMTypeUtils.h
+++ b/Engine/Source/Runtime/RigVM/Public/RigVMTypeUtils.h
@@ -48,6 +48,7 @@ namespace RigVMTypeUtils
 	class RIGVM_API TypeIndex
 	{
 	public:
+		static inline int32 Execute = INDEX_NONE;	
 		static inline int32 Bool = INDEX_NONE;	
 		static inline int32 Float = INDEX_NONE;	
 		static inline int32 Double = INDEX_NONE;	
