RigVM: Disable typeindex debugging by default

diff --git a/Engine/Source/Runtime/RigVM/Public/RigVMDefines.h b/Engine/Source/Runtime/RigVM/Public/RigVMDefines.h
--- a/Engine/Source/Runtime/RigVM/Public/RigVMDefines.h
+++ b/Engine/Source/Runtime/RigVM/Public/RigVMDefines.h
@@ -23,5 +23,5 @@
 #endif
 
 #ifndef UE_RIGVM_DEBUG_TYPEINDEX
-#define UE_RIGVM_DEBUG_TYPEINDEX 1
+#define UE_RIGVM_DEBUG_TYPEINDEX 0
 #endif
