RigVM: Fix crash on applicatoin exit

diff --git a/Engine/Source/Runtime/RigVM/Private/RigVMCore/RigVMRegistry.cpp b/Engine/Source/Runtime/RigVM/Private/RigVMCore/RigVMRegistry.cpp
--- a/Engine/Source/Runtime/RigVM/Private/RigVMCore/RigVMRegistry.cpp
+++ b/Engine/Source/Runtime/RigVM/Private/RigVMCore/RigVMRegistry.cpp
@@ -5,21 +5,14 @@
 #include "RigVMTypeUtils.h"
 #include "RigVMModule.h"
 #include "UObject/UObjectIterator.h"
+#include "Misc/CoreDelegates.h"
 
 FRigVMRegistry FRigVMRegistry::s_RigVMRegistry;
 const FName FRigVMRegistry::TemplateNameMetaName = TEXT("TemplateName");
 
 FRigVMRegistry::~FRigVMRegistry()
 {
-	for(FRigVMDispatchFactory* Factory : Factories)
-	{
-		if(const UScriptStruct* ScriptStruct = Factory->GetScriptStruct())
-		{
-			ScriptStruct->DestroyStruct(Factory, 1);
-		}
-		FMemory::Free(Factory);
-	}
-	Factories.Reset();
+	Reset();
 }
 
 FRigVMRegistry& FRigVMRegistry::Get()
@@ -125,6 +118,12 @@ void FRigVMRegistry::InitializeIfNeeded()
 	}
 
 	Refresh();
+
+	// hook the registry to prepare for engine shutdown
+	FCoreDelegates::OnExit.AddLambda([]()
+	{
+		s_RigVMRegistry.Reset();
+	});
 }
 
 void FRigVMRegistry::Refresh()
@@ -175,6 +174,19 @@ void FRigVMRegistry::Refresh()
 	}
 }
 
+void FRigVMRegistry::Reset()
+{
+	for(FRigVMDispatchFactory* Factory : Factories)
+	{
+		if(const UScriptStruct* ScriptStruct = Factory->GetScriptStruct())
+		{
+			ScriptStruct->DestroyStruct(Factory, 1);
+		}
+		FMemory::Free(Factory);
+	}
+	Factories.Reset();
+}
+
 TRigVMTypeIndex FRigVMRegistry::FindOrAddType(const FRigVMTemplateArgumentType& InType)
 {
 	TRigVMTypeIndex Index = GetTypeIndex(InType);
diff --git a/Engine/Source/Runtime/RigVM/Public/RigVMCore/RigVMRegistry.h b/Engine/Source/Runtime/RigVM/Public/RigVMCore/RigVMRegistry.h
--- a/Engine/Source/Runtime/RigVM/Public/RigVMCore/RigVMRegistry.h
+++ b/Engine/Source/Runtime/RigVM/Public/RigVMCore/RigVMRegistry.h
@@ -39,6 +39,9 @@ public:
 	// based on the names.
 	void Refresh();
 
+	// Clear the registry
+	void Reset();
+
 	// Adds a type if it doesn't exist yet and returns its index.
 	// This function is not thead-safe
 	TRigVMTypeIndex FindOrAddType(const FRigVMTemplateArgumentType& InType);
