RigVM: Remove erroreous ensure

diff --git a/Engine/Source/Runtime/RigVM/Private/RigVMCore/RigVMRegistry.cpp b/Engine/Source/Runtime/RigVM/Private/RigVMCore/RigVMRegistry.cpp
--- a/Engine/Source/Runtime/RigVM/Private/RigVMCore/RigVMRegistry.cpp
+++ b/Engine/Source/Runtime/RigVM/Private/RigVMCore/RigVMRegistry.cpp
@@ -522,7 +522,7 @@ bool FRigVMRegistry::IsWildCardType(TRigVMTypeIndex InTypeIndex) const
 
 bool FRigVMRegistry::CanMatchTypes(TRigVMTypeIndex InTypeIndexA, TRigVMTypeIndex InTypeIndexB, bool bAllowFloatingPointCasts) const
 {
-	if(!ensure(Types.IsValidIndex(InTypeIndexA)) || !ensure(Types.IsValidIndex(InTypeIndexB)))
+	if(!Types.IsValidIndex(InTypeIndexA) || !Types.IsValidIndex(InTypeIndexB))
 	{
 		return false;
 	}
