RigVM: Guard against array elements not existing at runtime

diff --git a/Engine/Source/Runtime/RigVM/Private/RigVMCore/RigVM.cpp b/Engine/Source/Runtime/RigVM/Private/RigVMCore/RigVM.cpp
--- a/Engine/Source/Runtime/RigVM/Private/RigVMCore/RigVM.cpp
+++ b/Engine/Source/Runtime/RigVM/Private/RigVMCore/RigVM.cpp
@@ -3187,6 +3187,10 @@ void URigVM::CacheSingleMemoryHandle(int32 InHandleIndex, const FRigVMOperand& I
 		check(ExternalVariable.IsValid(false));
 
 		CachedMemoryHandles[InHandleIndex] = {ExternalVariable.Memory, ExternalVariable.Property, PropertyPath};
+#if UE_BUILD_DEBUG
+		// make sure the handle points to valid memory
+		check(CachedMemoryHandles[InHandleIndex].GetData(false) != nullptr);
+#endif
 		return;
 	}
 
@@ -3203,6 +3207,19 @@ void URigVM::CacheSingleMemoryHandle(int32 InHandleIndex, const FRigVMOperand& I
 	uint8* Data = Memory->GetData<uint8>(InArg.GetRegisterIndex());
 	const FProperty* Property = Memory->GetProperties()[InArg.GetRegisterIndex()];
 	CachedMemoryHandles[InHandleIndex] = {Data, Property, PropertyPath};
+
+#if UE_BUILD_DEBUG
+	// make sure the handle points to valid memory
+	FRigVMMemoryHandle& Handle = CachedMemoryHandles[InHandleIndex]; 
+	check(Handle.GetData(false) != nullptr);
+	if(PropertyPath)
+	{
+		uint8* MemoryPtr = Handle.GetData(false);
+		// don't check the result - since it may be an array element
+		// that doesn't exist yet. 
+		PropertyPath->GetData<uint8>(MemoryPtr, Property);
+	}
+#endif
 }
 
 void URigVM::CopyOperandForDebuggingImpl(const FRigVMOperand& InArg, const FRigVMMemoryHandle& InHandle, const FRigVMOperand& InDebugOperand)
diff --git a/Engine/Source/Runtime/RigVM/Private/RigVMCore/RigVMMemoryStorage.cpp b/Engine/Source/Runtime/RigVM/Private/RigVMCore/RigVMMemoryStorage.cpp
--- a/Engine/Source/Runtime/RigVM/Private/RigVMCore/RigVMMemoryStorage.cpp
+++ b/Engine/Source/Runtime/RigVM/Private/RigVMCore/RigVMMemoryStorage.cpp
@@ -324,6 +324,34 @@ bool FRigVMPropertyDescription::RequiresCPPTypeObject(const FString& InCPPType)
 
 ////////////////////////////////////////////////////////////////////////////////////////////////////////////////////////
 
+class FRigVMMemoryStorageImportErrorContext : public FOutputDevice
+{
+public:
+
+	bool bLogErrors;
+	int32 NumErrors;
+
+	FRigVMMemoryStorageImportErrorContext(bool InLogErrors = true)
+		: FOutputDevice()
+		, bLogErrors(InLogErrors)
+		, NumErrors(0)
+	{
+	}
+
+	FORCEINLINE_DEBUGGABLE void Serialize(const TCHAR* V, ELogVerbosity::Type Verbosity, const class FName& Category) override
+	{
+		if(bLogErrors)
+		{
+#if WITH_EDITOR
+			UE_LOG(LogRigVM, Display, TEXT("Skipping Importing To MemoryStorage: %s"), V);
+#else
+			UE_LOG(LogRigVM, Error, TEXT("Error Importing To MemoryStorage: %s"), V);
+#endif
+		}
+		NumErrors++;
+	}
+};
+
 void URigVMMemoryStorageGeneratorClass::PurgeClass(bool bRecompilingOnLoad)
 {
 	Super::PurgeClass(bRecompilingOnLoad);
@@ -520,7 +548,19 @@ URigVMMemoryStorageGeneratorClass* URigVMMemoryStorageGeneratorClass::CreateStor
 		const FProperty* Property = LinkedProperties[PropertyIndex];
 		uint8* ValuePtr = Property->ContainerPtrToValuePtr<uint8>(CDO);
 
-		Property->ImportText_Direct(*DefaultValue, ValuePtr, nullptr, EPropertyPortFlags::PPF_None);
+		FRigVMMemoryStorageImportErrorContext ErrorPipe(false);
+		Property->ImportText_Direct(*DefaultValue, ValuePtr, nullptr, EPropertyPortFlags::PPF_None, &ErrorPipe);
+		if(ErrorPipe.NumErrors > 0)
+		{
+			// check if the default value was provided as a single element
+			if(const FArrayProperty* ArrayProperty = CastField<FArrayProperty>(Property))
+			{
+				static constexpr TCHAR BraceFormat[] = TEXT("(%s)");
+				const FString DefaultValueWithBraces = FString::Printf(BraceFormat, *DefaultValue);
+				ErrorPipe = FRigVMMemoryStorageImportErrorContext(false);
+				Property->ImportText_Direct(*DefaultValueWithBraces, ValuePtr, nullptr, EPropertyPortFlags::PPF_None, &ErrorPipe);
+			}
+		}
 	}
 
 	return Class;
@@ -844,29 +884,6 @@ FString URigVMMemoryStorage::GetDataAsStringSafe(const FRigVMOperand& InOperand,
 	return GetDataAsStringSafe(PropertyIndex, PortFlags);
 }
 
-class FRigVMMemoryStorageImportErrorContext : public FOutputDevice
-{
-public:
-
-	int32 NumErrors;
-
-	FRigVMMemoryStorageImportErrorContext()
-		: FOutputDevice()
-		, NumErrors(0)
-	{
-	}
-
-	FORCEINLINE_DEBUGGABLE void Serialize(const TCHAR* V, ELogVerbosity::Type Verbosity, const class FName& Category) override
-	{
-#if WITH_EDITOR
-		UE_LOG(LogRigVM, Display, TEXT("Skipping Importing To MemoryStorage: %s"), V);
-#else
-		UE_LOG(LogRigVM, Error, TEXT("Error Importing To MemoryStorage: %s"), V);
-#endif
-		NumErrors++;
-	}
-};
-
 bool URigVMMemoryStorage::SetDataFromString(int32 InPropertyIndex, const FString& InValue)
 {
 	check(IsValidIndex(InPropertyIndex));
@@ -1072,6 +1089,17 @@ bool URigVMMemoryStorage::CopyProperty(
 	TraversePropertyPath(InTargetProperty, InTargetPtr, InTargetPropertyPath);
 	TraversePropertyPath(InSourceProperty, SourcePtr, InSourcePropertyPath);
 
+	if(InTargetPtr == nullptr)
+	{
+		check(!InTargetPropertyPath.IsEmpty());
+		return false;
+	}
+	if(InSourcePtr == nullptr)
+	{
+		check(!InSourcePropertyPath.IsEmpty());
+		return false;
+	}
+
 	return CopyProperty(InTargetProperty, InTargetPtr, InSourceProperty, SourcePtr);
 }
 
diff --git a/Engine/Source/Runtime/RigVM/Private/RigVMCore/RigVMPropertyPath.cpp b/Engine/Source/Runtime/RigVM/Private/RigVMCore/RigVMPropertyPath.cpp
--- a/Engine/Source/Runtime/RigVM/Private/RigVMCore/RigVMPropertyPath.cpp
+++ b/Engine/Source/Runtime/RigVM/Private/RigVMCore/RigVMPropertyPath.cpp
@@ -195,6 +195,10 @@ uint8* FRigVMPropertyPath::GetData_Internal(uint8* InPtr, const FProperty* InPro
 				InProperty = CastFieldChecked<FArrayProperty>(Segment.Property)->Inner;
 #endif
 				FScriptArrayHelper ArrayHelper(CastFieldChecked<FArrayProperty>(Segment.Property), InPtr);
+				if(!ArrayHelper.IsValidIndex(Segment.Index))
+				{
+					return nullptr;
+				}
 				InPtr = ArrayHelper.GetRawPtr(Segment.Index);
 				break;
 			}
