Control Rig: Fix ensure hit when openning rigs with FRigPose pins

diff --git a/Engine/Source/Developer/RigVMDeveloper/Private/RigVMModel/RigVMController.cpp b/Engine/Source/Developer/RigVMDeveloper/Private/RigVMModel/RigVMController.cpp
--- a/Engine/Source/Developer/RigVMDeveloper/Private/RigVMModel/RigVMController.cpp
+++ b/Engine/Source/Developer/RigVMDeveloper/Private/RigVMModel/RigVMController.cpp
@@ -17112,7 +17112,7 @@ bool URigVMController::RecomputeAllTemplateFilteredPermutations(bool bSetupUndoR
 	check(Graph);
 
 	// Save all template pin types, in case we need them after initializing
-	TMap<URigVMPin*, int32> TypesBeforeRecomputing;
+	TMap<URigVMPin*, TRigVMTypeIndex> TypesBeforeRecomputing;
 	for (URigVMNode* Node : Graph->GetNodes())
 	{
 		if (URigVMTemplateNode* TemplateNode = Cast<URigVMTemplateNode>(Node))
@@ -17230,7 +17230,7 @@ bool URigVMController::RecomputeAllTemplateFilteredPermutations(bool bSetupUndoR
 						if (URigVMTemplateNode* TemplateNode = Cast<URigVMTemplateNode>(EdgePin->GetNode()))
 						{
 							URigVMPin* RootPin = EdgePin->GetRootPin();
-							if (int32* TypeIndex = TypesBeforeRecomputing.Find(RootPin))
+							if (TRigVMTypeIndex* TypeIndex = TypesBeforeRecomputing.Find(RootPin))
 							{
 								if (TemplateNode->FilteredSupportsType(RootPin, *TypeIndex))
 								{
diff --git a/Engine/Source/Runtime/RigVM/Private/RigVMCore/RigVMRegistry.cpp b/Engine/Source/Runtime/RigVM/Private/RigVMCore/RigVMRegistry.cpp
--- a/Engine/Source/Runtime/RigVM/Private/RigVMCore/RigVMRegistry.cpp
+++ b/Engine/Source/Runtime/RigVM/Private/RigVMCore/RigVMRegistry.cpp
@@ -192,74 +192,100 @@ TRigVMTypeIndex FRigVMRegistry::FindOrAddType(const FRigVMTemplateArgumentType&
 	TRigVMTypeIndex Index = GetTypeIndex(InType);
 	if(Index == INDEX_NONE)
 	{
-		FRigVMTemplateArgumentType ElementType;
-		ElementType.CPPType = NAME_None;
-		if(InType.IsArray())
+		FRigVMTemplateArgumentType ElementType = InType;
+		while(ElementType.IsArray())
 		{
-			ElementType = InType;
 			ElementType.ConvertToBaseElement();
 		}
-		FRigVMTemplateArgumentType ArrayType = InType;
-		ArrayType.ConvertToArray();
-
+		
+		const UObject* CPPTypeObject = ElementType.CPPTypeObject;
+		if(CPPTypeObject != nullptr)
 		{
-			FTypeInfo Info;
-			Info.Type = InType;
-			Info.bIsArray = InType.IsArray();
-			if(UScriptStruct* ScriptStruct = Cast<UScriptStruct>(InType.CPPTypeObject))
+			if(const UClass* Class = Cast<UClass>(CPPTypeObject))
+			{
+				if(!IsAllowedType(Class))
+				{
+					return Index;
+				}	
+			}
+			else if(const UEnum* Enum = Cast<UEnum>(CPPTypeObject))
+			{
+				if(!IsAllowedType(Enum))
+				{
+					return Index;
+				}
+			}
+			else if(const UStruct* Struct = Cast<UStruct>(CPPTypeObject))
 			{
-				Info.bIsExecute = ScriptStruct->IsChildOf(FRigVMExecuteContext::StaticStruct());
+				if(!IsAllowedType(Struct))
+				{					
+					return Index;
+				}
 			}
+		}
+
+		bool bIsExecute = false;
+		if(const UScriptStruct* ScriptStruct = Cast<UScriptStruct>(CPPTypeObject))
+		{
+			bIsExecute = ScriptStruct->IsChildOf(FRigVMExecuteContext::StaticStruct());
+		}
 
-			if(!ElementType.CPPType.IsNone())
+		TArray<TRigVMTypeIndex> Indices;
+		Indices.Reserve(3);
+		for (int32 ArrayDimension=0; ArrayDimension<3; ++ArrayDimension)
+		{
+			if (bIsExecute && ArrayDimension>0)
 			{
-				Info.BaseTypeIndex = GetTypeIndex(ElementType);
+				continue;
 			}
-			Info.ArrayTypeIndex = GetTypeIndex(ArrayType);
+			
+			FRigVMTemplateArgumentType CurType = ElementType;
+			for (int32 j=0; j<ArrayDimension; ++j)
+			{
+				CurType.ConvertToArray();
+			}
+
+			FTypeInfo Info;
+			Info.Type = CurType;
+			Info.bIsArray = ArrayDimension > 0;
+			Info.bIsExecute = bIsExecute;
 			
 			Index = Types.Add(Info);
 #if UE_RIGVM_DEBUG_TYPEINDEX
 			Index.Name = Info.Type.CPPType;
 #endif
-		}
-
-		TypeToIndex.Add(InType, Index);
+			TypeToIndex.Add(CurType, Index);
 
-		// add the type to the category map
-		static constexpr TCHAR ArrayArrayPrefix[] = TEXT("TArray<TArray<");
-		const int32 ArrayDimension = Types[Index].bIsArray ?
-			(InType.CPPType.ToString().StartsWith(ArrayArrayPrefix) ? 2 : 1) : 0;
-		const UObject* CPPTypeObject = Types[Index].Type.CPPTypeObject;
+			Indices.Add(Index);
 
-		// simple types
-		if(CPPTypeObject == nullptr)
-		{
-			switch(ArrayDimension)
+			// Add to category
+			// simple types
+			if(CPPTypeObject == nullptr)
 			{
-				default:
-				case 0:
-				{
-					RegisterTypeInCategory(FRigVMTemplateArgument::ETypeCategory_SingleSimpleValue, Index);
-					RegisterTypeInCategory(FRigVMTemplateArgument::ETypeCategory_SingleAnyValue, Index);
-					break;
-				}
-				case 1:
-				{
-					RegisterTypeInCategory(FRigVMTemplateArgument::ETypeCategory_ArraySimpleValue, Index);
-					RegisterTypeInCategory(FRigVMTemplateArgument::ETypeCategory_ArrayAnyValue, Index);
-					break;
-				}
-				case 2:
+				switch(ArrayDimension)
 				{
-					RegisterTypeInCategory(FRigVMTemplateArgument::ETypeCategory_ArrayArraySimpleValue, Index);
-					RegisterTypeInCategory(FRigVMTemplateArgument::ETypeCategory_ArrayArrayAnyValue, Index);
-					break;
+					default:
+					case 0:
+					{
+						RegisterTypeInCategory(FRigVMTemplateArgument::ETypeCategory_SingleSimpleValue, Index);
+						RegisterTypeInCategory(FRigVMTemplateArgument::ETypeCategory_SingleAnyValue, Index);
+						break;
+					}
+					case 1:
+					{
+						RegisterTypeInCategory(FRigVMTemplateArgument::ETypeCategory_ArraySimpleValue, Index);
+						RegisterTypeInCategory(FRigVMTemplateArgument::ETypeCategory_ArrayAnyValue, Index);
+						break;
+					}
+					case 2:
+					{
+						RegisterTypeInCategory(FRigVMTemplateArgument::ETypeCategory_ArrayArraySimpleValue, Index);
+						RegisterTypeInCategory(FRigVMTemplateArgument::ETypeCategory_ArrayArrayAnyValue, Index);
+						break;
+					}
 				}
 			}
-		}
-		else if(const UClass* Class = Cast<UClass>(CPPTypeObject))
-		{
-			if(IsAllowedType(Class))
+			else if(const UClass* Class = Cast<UClass>(CPPTypeObject))
 			{
 				switch(ArrayDimension)
 				{
@@ -283,11 +309,8 @@ TRigVMTypeIndex FRigVMRegistry::FindOrAddType(const FRigVMTemplateArgumentType&
 						break;
 					}
 				}
-			}	
-		}
-		else if(const UEnum* Enum = Cast<UEnum>(CPPTypeObject))
-		{
-			if(IsAllowedType(Enum))
+			}
+			else if(const UEnum* Enum = Cast<UEnum>(CPPTypeObject))
 			{
 				switch(ArrayDimension)
 				{
@@ -312,112 +335,92 @@ TRigVMTypeIndex FRigVMRegistry::FindOrAddType(const FRigVMTemplateArgumentType&
 					}
 				}
 			}
-		}
-		else if(const UStruct* Struct = Cast<UStruct>(CPPTypeObject))
-		{
-			if(IsAllowedType(Struct))
+			else if(const UStruct* Struct = Cast<UStruct>(CPPTypeObject))
 			{
-				if(GetMathTypes().Contains(CPPTypeObject))
+				if(Struct->IsChildOf(FRigVMExecuteContext::StaticStruct()))
+				{
+					if(ArrayDimension == 0)
+					{
+						RegisterTypeInCategory(FRigVMTemplateArgument::ETypeCategory_Execute, Index);
+					}
+				}
+				else
 				{
+					if(GetMathTypes().Contains(CPPTypeObject))
+					{
+						switch(ArrayDimension)
+						{
+							default:
+							case 0:
+							{
+								RegisterTypeInCategory(FRigVMTemplateArgument::ETypeCategory_SingleMathStructValue, Index);
+								break;
+							}
+							case 1:
+							{
+								RegisterTypeInCategory(FRigVMTemplateArgument::ETypeCategory_ArrayMathStructValue, Index);
+								break;
+							}
+							case 2:
+							{
+								RegisterTypeInCategory(FRigVMTemplateArgument::ETypeCategory_ArrayArrayMathStructValue, Index);
+								break;
+							}
+						}
+					}
+					
 					switch(ArrayDimension)
 					{
 						default:
 						case 0:
 						{
-							RegisterTypeInCategory(FRigVMTemplateArgument::ETypeCategory_SingleMathStructValue, Index);
+							RegisterTypeInCategory(FRigVMTemplateArgument::ETypeCategory_SingleScriptStructValue, Index);
+							RegisterTypeInCategory(FRigVMTemplateArgument::ETypeCategory_SingleAnyValue, Index);
 							break;
 						}
 						case 1:
 						{
-							RegisterTypeInCategory(FRigVMTemplateArgument::ETypeCategory_ArrayMathStructValue, Index);
+							RegisterTypeInCategory(FRigVMTemplateArgument::ETypeCategory_ArrayScriptStructValue, Index);
+							RegisterTypeInCategory(FRigVMTemplateArgument::ETypeCategory_ArrayAnyValue, Index);
 							break;
 						}
 						case 2:
 						{
-							RegisterTypeInCategory(FRigVMTemplateArgument::ETypeCategory_ArrayArrayMathStructValue, Index);
+							RegisterTypeInCategory(FRigVMTemplateArgument::ETypeCategory_ArrayArrayScriptStructValue, Index);
+							RegisterTypeInCategory(FRigVMTemplateArgument::ETypeCategory_ArrayArrayAnyValue, Index);
 							break;
 						}
 					}
 				}
-				
-				switch(ArrayDimension)
-				{
-					default:
-					case 0:
-					{
-						RegisterTypeInCategory(FRigVMTemplateArgument::ETypeCategory_SingleScriptStructValue, Index);
-						RegisterTypeInCategory(FRigVMTemplateArgument::ETypeCategory_SingleAnyValue, Index);
-						break;
-					}
-					case 1:
-					{
-						RegisterTypeInCategory(FRigVMTemplateArgument::ETypeCategory_ArrayScriptStructValue, Index);
-						RegisterTypeInCategory(FRigVMTemplateArgument::ETypeCategory_ArrayAnyValue, Index);
-						break;
-					}
-					case 2:
-					{
-						RegisterTypeInCategory(FRigVMTemplateArgument::ETypeCategory_ArrayArrayScriptStructValue, Index);
-						RegisterTypeInCategory(FRigVMTemplateArgument::ETypeCategory_ArrayArrayAnyValue, Index);
-						break;
-					}
-				}
-			}
-			else if(Struct->IsChildOf(FRigVMExecuteContext::StaticStruct()))
-			{
-				if(!Types[Index].bIsArray)
-				{
-					RegisterTypeInCategory(FRigVMTemplateArgument::ETypeCategory_Execute, Index);
-				}
 			}
 		}
 
-		// register the opposing type
-		if(!Types[Index].bIsArray)
+		if (!bIsExecute)
 		{
-			Types[Index].ArrayTypeIndex = FindOrAddType(ArrayType);
-			if(Types[Index].ArrayTypeIndex != INDEX_NONE)
-			{
-				Types[Types[Index].ArrayTypeIndex].BaseTypeIndex = Index;
-			}
-		}
-		else
-		{
-			Types[Index].BaseTypeIndex = FindOrAddType(ElementType);
-			if(Types[Index].BaseTypeIndex != INDEX_NONE)
-			{
-				Types[Types[Index].BaseTypeIndex].ArrayTypeIndex = Index;
+			Types[Indices[1]].BaseTypeIndex = Indices[0];
+			Types[Indices[2]].BaseTypeIndex = Indices[1];
 
-				// also automatically do the two dimensional array
-				if(GetArrayDimensionsForType(Index) == 1)
-				{
-					Types[Index].ArrayTypeIndex = FindOrAddType(ArrayType);
-					if(Types[Index].ArrayTypeIndex != INDEX_NONE)
-					{
-						Types[Types[Index].ArrayTypeIndex].BaseTypeIndex = Index;
-					}
-				}
-			}
+			Types[Indices[0]].ArrayTypeIndex = Indices[1];
+			Types[Indices[1]].ArrayTypeIndex = Indices[2];
 		}
 
 		// if the type is a structure
 		// then add all of its sub property types
-		if(!Types[Index].bIsArray)
+		if(const UStruct* Struct = Cast<UStruct>(CPPTypeObject))
 		{
-			if(const UStruct* Struct = Cast<UStruct>(Types[Index].Type.CPPTypeObject))
+			for (TFieldIterator<FProperty> It(Struct); It; ++It)
 			{
-				for (TFieldIterator<FProperty> It(Struct); It; ++It)
+				FProperty* Property = *It;
+				if(IsAllowedType(Property))
 				{
-					FProperty* Property = *It;
-					if(IsAllowedType(Property, true))
-					{
-						// by creating a template argument for the child property
-						// the type will be added by calling ::FindOrAddType recursively.
-						FRigVMTemplateArgument DummyArgument(Property);
-					}
+					// by creating a template argument for the child property
+					// the type will be added by calling ::FindOrAddType recursively.
+					FRigVMTemplateArgument DummyArgument(Property);
 				}
-			}
+			}			
 		}
+
+		return GetTypeIndex(InType);
 	}
 	
 	return Index;
@@ -602,19 +605,8 @@ TRigVMTypeIndex FRigVMRegistry::GetBaseTypeFromArrayTypeIndex(TRigVMTypeIndex In
 	return INDEX_NONE;
 }
 
-bool FRigVMRegistry::IsAllowedType(const FProperty* InProperty, bool bCheckFlags)
+bool FRigVMRegistry::IsAllowedType(const FProperty* InProperty)
 {
-	if(bCheckFlags)
-	{
-		if(!InProperty->HasAnyPropertyFlags(
-			CPF_BlueprintVisible |
-			CPF_BlueprintReadOnly |
-			CPF_Edit))
-		{
-			return false;
-		}
-	}
-
 	if(InProperty->IsA<FBoolProperty>() ||
 		InProperty->IsA<FUInt32Property>() ||
 		InProperty->IsA<FInt8Property>() ||
@@ -632,7 +624,7 @@ bool FRigVMRegistry::IsAllowedType(const FProperty* InProperty, bool bCheckFlags
 
 	if(const FArrayProperty* ArrayProperty  = CastField<FArrayProperty>(InProperty))
 	{
-		return IsAllowedType(ArrayProperty->Inner, false);
+		return IsAllowedType(ArrayProperty->Inner);
 	}
 	if(const FStructProperty* StructProperty = CastField<FStructProperty>(InProperty))
 	{
@@ -672,14 +664,6 @@ bool FRigVMRegistry::IsAllowedType(const UStruct* InStruct)
 	{
 		return false;
 	}
-	if(InStruct->IsChildOf(FRigVMUnknownType::StaticStruct()))
-	{
-		return false;
-	}
-	if(InStruct->IsChildOf(FRigVMExecuteContext::StaticStruct()))
-	{
-		return false;
-	}
 	if(InStruct->IsChildOf(FRigVMDispatchFactory::StaticStruct()))
 	{
 		return false;
diff --git a/Engine/Source/Runtime/RigVM/Public/RigVMCore/RigVMRegistry.h b/Engine/Source/Runtime/RigVM/Public/RigVMCore/RigVMRegistry.h
--- a/Engine/Source/Runtime/RigVM/Public/RigVMCore/RigVMRegistry.h
+++ b/Engine/Source/Runtime/RigVM/Public/RigVMCore/RigVMRegistry.h
@@ -224,7 +224,7 @@ private:
 		return RF_Public;
 	}
 
-	static bool IsAllowedType(const FProperty* InProperty, bool bCheckFlags = true);
+	static bool IsAllowedType(const FProperty* InProperty);
 	static bool IsAllowedType(const UEnum* InEnum);
 	static bool IsAllowedType(const UStruct* InStruct);
 	static bool IsAllowedType(const UClass* InClass);
