Control Rig: Fix Register Auxiliary Template unit test not passing

diff --git a/Engine/Source/Runtime/RigVM/Private/RigVMCore/RigVMRegistry.cpp b/Engine/Source/Runtime/RigVM/Private/RigVMCore/RigVMRegistry.cpp
--- a/Engine/Source/Runtime/RigVM/Private/RigVMCore/RigVMRegistry.cpp
+++ b/Engine/Source/Runtime/RigVM/Private/RigVMCore/RigVMRegistry.cpp
@@ -257,7 +257,25 @@ TRigVMTypeIndex FRigVMRegistry::FindOrAddType(const FRigVMTemplateArgumentType&
 			TypeToIndex.Add(CurType, Index);
 
 			Indices.Add(Index);
+		}
+
+		if (!bIsExecute)
+		{
+			Types[Indices[1]].BaseTypeIndex = Indices[0];
+			Types[Indices[2]].BaseTypeIndex = Indices[1];
 
+			Types[Indices[0]].ArrayTypeIndex = Indices[1];
+			Types[Indices[1]].ArrayTypeIndex = Indices[2];
+		}
+
+		for (int32 ArrayDimension=0; ArrayDimension<3; ++ArrayDimension)
+		{
+			if(bIsExecute && ArrayDimension > 0)
+			{
+				break;
+			}
+			Index = Indices[ArrayDimension];
+			
 			// Add to category
 			// simple types
 			if(CPPTypeObject == nullptr)
@@ -395,15 +413,6 @@ TRigVMTypeIndex FRigVMRegistry::FindOrAddType(const FRigVMTemplateArgumentType&
 			}
 		}
 
-		if (!bIsExecute)
-		{
-			Types[Indices[1]].BaseTypeIndex = Indices[0];
-			Types[Indices[2]].BaseTypeIndex = Indices[1];
-
-			Types[Indices[0]].ArrayTypeIndex = Indices[1];
-			Types[Indices[1]].ArrayTypeIndex = Indices[2];
-		}
-
 		// if the type is a structure
 		// then add all of its sub property types
 		if(const UStruct* Struct = Cast<UStruct>(CPPTypeObject))
diff --git a/Engine/Source/Runtime/RigVM/Private/RigVMCore/RigVMTemplate.cpp b/Engine/Source/Runtime/RigVM/Private/RigVMCore/RigVMTemplate.cpp
--- a/Engine/Source/Runtime/RigVM/Private/RigVMCore/RigVMTemplate.cpp
+++ b/Engine/Source/Runtime/RigVM/Private/RigVMCore/RigVMTemplate.cpp
@@ -1392,6 +1392,15 @@ bool FRigVMTemplate::AddTypeForArgument(const FName& InArgumentName, TRigVMTypeI
 					return false;
 				}
 			}
+
+			// Find if these types were already registered
+			FRigVMTemplateTypeMap TestTypes = Types;
+			TArray<int32> TestPermutations;
+			if (Resolve(TestTypes, TestPermutations, false))
+			{
+				return false;
+			}
+			
 			for(FRigVMTemplateArgument& Argument : Arguments)
 			{
 				const TRigVMTypeIndex TypeIndex = Types.FindChecked(Argument.Name);
@@ -1400,6 +1409,7 @@ bool FRigVMTemplate::AddTypeForArgument(const FName& InArgumentName, TRigVMTypeI
 			}
 
 			Permutations.Add(INDEX_NONE);
+			return true;
 		}
 	}
 	return false;
