Control Rig: Fix RigVMRegister allowing dynamically generated structs as types

diff --git a/Engine/Source/Runtime/RigVM/Private/RigVMCore/RigVMRegistry.cpp b/Engine/Source/Runtime/RigVM/Private/RigVMCore/RigVMRegistry.cpp
--- a/Engine/Source/Runtime/RigVM/Private/RigVMCore/RigVMRegistry.cpp
+++ b/Engine/Source/Runtime/RigVM/Private/RigVMCore/RigVMRegistry.cpp
@@ -6,6 +6,7 @@
 #include "RigVMModule.h"
 #include "UObject/UObjectIterator.h"
 #include "Misc/CoreDelegates.h"
+#include "Engine/UserDefinedStruct.h"
 
 FRigVMRegistry FRigVMRegistry::s_RigVMRegistry;
 const FName FRigVMRegistry::TemplateNameMetaName = TEXT("TemplateName");
@@ -135,7 +136,7 @@ void FRigVMRegistry::Refresh()
 		if(IsAllowedType(ScriptStruct))
 		{
 			// if this is a C++ type - skip it
-			if(!ScriptStruct->IsNative() || ScriptStruct->IsChildOf(FRigVMExecuteContext::StaticStruct()))
+			if(ScriptStruct->IsA<UUserDefinedStruct>() || ScriptStruct->IsChildOf(FRigVMExecuteContext::StaticStruct()))
 			{
 				const FString CPPType = ScriptStruct->GetStructCPPName();
 				FindOrAddType(FRigVMTemplateArgumentType(ScriptStruct));
