RigVM: Fix bad merge in FindOrAddType

diff --git a/Engine/Source/Runtime/RigVM/Private/RigVMCore/RigVMRegistry.cpp b/Engine/Source/Runtime/RigVM/Private/RigVMCore/RigVMRegistry.cpp
--- a/Engine/Source/Runtime/RigVM/Private/RigVMCore/RigVMRegistry.cpp
+++ b/Engine/Source/Runtime/RigVM/Private/RigVMCore/RigVMRegistry.cpp
@@ -263,12 +263,6 @@ TRigVMTypeIndex FRigVMRegistry::FindOrAddType_Internal(const FRigVMTemplateArgum
 		Types[Indices[1]].BaseTypeIndex = Indices[0];
 		Types[Indices[0]].ArrayTypeIndex = Indices[1];
 
-		if (!bIsExecute)
-		{
-			Types[Indices[1]].BaseTypeIndex = Indices[0];
-			Types[Indices[2]].BaseTypeIndex = Indices[1];
-		}
-
 		if (!bIsExecute)
 		{
 			Types[Indices[2]].BaseTypeIndex = Indices[1];
