Control Rig: Fix crash when resolving equals node to TMap

diff --git a/Engine/Source/Runtime/RigVM/Private/RigVMCore/RigVMRegistry.cpp b/Engine/Source/Runtime/RigVM/Private/RigVMCore/RigVMRegistry.cpp
--- a/Engine/Source/Runtime/RigVM/Private/RigVMCore/RigVMRegistry.cpp
+++ b/Engine/Source/Runtime/RigVM/Private/RigVMCore/RigVMRegistry.cpp
@@ -421,12 +421,20 @@ TRigVMTypeIndex FRigVMRegistry::FindOrAddType_Internal(const FRigVMTemplateArgum
 			for (TFieldIterator<FProperty> It(Struct); It; ++It)
 			{
 				FProperty* Property = *It;
-				if(bForce || IsAllowedType(Property))
+				if(IsAllowedType(Property))
 				{
 					// by creating a template argument for the child property
 					// the type will be added by calling ::FindOrAddType_Internal recursively.
 					FRigVMTemplateArgument DummyArgument(Property);
 				}
+#if WITH_EDITOR
+				else
+				{
+					// If the subproperty is not allowed, let's make sure it's hidden. Otherwise we end up with
+					// subpins with invalid types 
+					check(FRigVMStruct::GetPinDirectionFromProperty(Property) == ERigVMPinDirection::Hidden);
+				}
+#endif
 			}			
 		}
 
