RigVM: 1. Added missing execute context type patching during argument type registration.The same patching process was already in place for when arguments were initially created using FRigVMTemplateArgument 2. fixed a typo in Array node

diff --git a/Engine/Source/Developer/RigVMDeveloper/Private/RigVMModel/Nodes/RigVMArrayNode.cpp b/Engine/Source/Developer/RigVMDeveloper/Private/RigVMModel/Nodes/RigVMArrayNode.cpp
--- a/Engine/Source/Developer/RigVMDeveloper/Private/RigVMModel/Nodes/RigVMArrayNode.cpp
+++ b/Engine/Source/Developer/RigVMDeveloper/Private/RigVMModel/Nodes/RigVMArrayNode.cpp
@@ -379,7 +379,7 @@ const FRigVMTemplate* URigVMArrayNode::GetTemplate() const
 					Arguments.Emplace(IndexFName, ERigVMPinDirection::Output, RigVMTypeUtils::TypeIndex::Int32);
 					Arguments.Emplace(CountFName, ERigVMPinDirection::Output, RigVMTypeUtils::TypeIndex::Int32);
 					Arguments.Emplace(RatioFName, ERigVMPinDirection::Output, RigVMTypeUtils::TypeIndex::Float);
-					Arguments.Emplace(ContinueFName, ERigVMPinDirection::Hidden, ExecuteTypeIndex);
+					Arguments.Emplace(ContinueFName, ERigVMPinDirection::Hidden, RigVMTypeUtils::TypeIndex::Bool);
 					Arguments.Emplace(CompletedFName, ERigVMPinDirection::Output, ExecuteTypeIndex);
 					break;
 				}
@@ -525,7 +525,7 @@ const FRigVMTemplate* URigVMArrayNode::GetTemplate() const
 					case ERigVMOpCode::ArrayIterator:
 					{
 						Types.Add(ExecuteFName, ExecuteTypeIndex);
-						Types.Add(ContinueFName, ExecuteTypeIndex);
+						Types.Add(ContinueFName, RigVMTypeUtils::TypeIndex::Bool);
 						Types.Add(CompletedFName, ExecuteTypeIndex);
 						Types.Add(IndexFName, RigVMTypeUtils::TypeIndex::Int32);
 						Types.Add(CountFName, RigVMTypeUtils::TypeIndex::Int32);
diff --git a/Engine/Source/Runtime/RigVM/Private/RigVMCore/RigVMTemplate.cpp b/Engine/Source/Runtime/RigVM/Private/RigVMCore/RigVMTemplate.cpp
--- a/Engine/Source/Runtime/RigVM/Private/RigVMCore/RigVMTemplate.cpp
+++ b/Engine/Source/Runtime/RigVM/Private/RigVMCore/RigVMTemplate.cpp
@@ -9,6 +9,39 @@
 
 ////////////////////////////////////////////////////////////////////////////////////////////////
 
+static void EnsureValidExecuteType_Internal(TRigVMTypeIndex& InTypeIndex)
+{
+	const FRigVMRegistry& Registry = FRigVMRegistry::Get();
+	if(InTypeIndex == INDEX_NONE)
+	{
+		return;
+	}
+		
+	if(InTypeIndex == RigVMTypeUtils::TypeIndex::Execute) 
+	{
+		return;
+	}
+
+	if(!Registry.IsExecuteType(InTypeIndex))
+	{
+		return;
+	}
+
+	// execute arguments can have various execute context types. but we always
+	// convert them to the base execute type to make matching types easier later.
+	// this means that the execute argument in every permutations shares 
+	// the same type index of RigVMTypeUtils::TypeIndex::Execute
+	if(Registry.IsArrayType(InTypeIndex))
+	{
+		InTypeIndex = Registry.GetArrayTypeFromBaseTypeIndex(RigVMTypeUtils::TypeIndex::Execute);
+	}
+	else
+	{
+		InTypeIndex = RigVMTypeUtils::TypeIndex::Execute;
+	}	
+}
+
+
 FRigVMTemplateArgument::FRigVMTemplateArgument()
 	: Index(INDEX_NONE)
 	, Name(NAME_None)
@@ -93,39 +126,6 @@ FRigVMTemplateArgument::FRigVMTemplateArgument(const FName& InName, ERigVMPinDir
 	UpdateTypeToPermutations();
 }
 
-void FRigVMTemplateArgument::EnsureValidExecuteType()
-{
-	const FRigVMRegistry& Registry = FRigVMRegistry::Get();
-	for(TRigVMTypeIndex& TypeIndex : TypeIndices)
-	{
-		if(TypeIndex == INDEX_NONE)
-		{
-			continue;
-		}
-		
-		if(TypeIndex == RigVMTypeUtils::TypeIndex::Execute) 
-		{
-			continue;
-		}
-
-		if(!Registry.IsExecuteType(TypeIndex))
-		{
-			continue;
-		}
-
-		// execute arguments are converted to the base type to make
-		// matching types easier later
-		if(Registry.IsArrayType(TypeIndex))
-		{
-			TypeIndex = Registry.GetArrayTypeFromBaseTypeIndex(RigVMTypeUtils::TypeIndex::Execute);
-		}
-		else
-		{
-			TypeIndex = RigVMTypeUtils::TypeIndex::Execute;
-		}
-	}
-}
-
 void FRigVMTemplateArgument::Serialize(FArchive& Ar)
 {
 	Ar << Index;
@@ -200,6 +200,14 @@ void FRigVMTemplateArgument::Serialize(FArchive& Ar)
 	}	
 }
 
+void FRigVMTemplateArgument::EnsureValidExecuteType()
+{
+	for(TRigVMTypeIndex& TypeIndex : TypeIndices)
+	{
+		EnsureValidExecuteType_Internal(TypeIndex);
+	}
+}
+
 void FRigVMTemplateArgument::UpdateTypeToPermutations()
 {
 	for(int32 TypeIndex=0;TypeIndex<TypeIndices.Num();TypeIndex++)
@@ -1431,9 +1439,15 @@ bool FRigVMTemplate::AddTypeForArgument(const FName& InArgumentName, TRigVMTypeI
 {
 	if(OnNewArgumentType().IsBound())
 	{
-		const FRigVMTemplateTypeMap Types = OnNewArgumentType().Execute(this, InArgumentName, InTypeIndex);
+		FRigVMTemplateTypeMap Types = OnNewArgumentType().Execute(this, InArgumentName, InTypeIndex);
 		if(Types.Num() == Arguments.Num())
 		{
+			for (TPair<FName, TRigVMTypeIndex>& ArgumentAndType : Types)
+			{
+				// similar to FRigVMTemplateArgument::EnsureValidExecuteType
+				EnsureValidExecuteType_Internal(ArgumentAndType.Value);
+			}
+			
 			// make sure this permutation doesn't exist yet
 			FRigVMTemplateTypeMap TempTypes = Types;
 			TArray<int32> ExistingPermutations;
