Control Rig: Fix cannot connect execute pin in entry node

diff --git a/Engine/Plugins/Animation/ControlRig/Source/ControlRigDeveloper/Private/ControlRigBlueprint.cpp b/Engine/Plugins/Animation/ControlRig/Source/ControlRigDeveloper/Private/ControlRigBlueprint.cpp
--- a/Engine/Plugins/Animation/ControlRig/Source/ControlRigDeveloper/Private/ControlRigBlueprint.cpp
+++ b/Engine/Plugins/Animation/ControlRig/Source/ControlRigDeveloper/Private/ControlRigBlueprint.cpp
@@ -1452,9 +1452,32 @@ void UControlRigBlueprint::RefreshAllModels(EControlRigBlueprintLoadType InLoadT
 					TemplateNode->InvalidateCache();
 					TemplateNode->PostLoad();
 
+					// Fix preferred type of execute context which are not the base class
+					{
+						const FRigVMRegistry& Registry = FRigVMRegistry::Get();
+						for (FRigVMTemplatePreferredType& PreferredType : TemplateNode->PreferredPermutationPairs)
+						{
+							if (!Registry.IsExecuteType(PreferredType.TypeIndex))
+							{
+								continue;
+							}
+							
+							if (URigVMPin* Pin = TemplateNode->FindPin(PreferredType.GetArgument().ToString()))
+							{
+								if (Pin->GetTypeIndex() == PreferredType.GetTypeIndex())
+								{
+									continue;
+								}
+
+								bDirtyDuringLoad = true;
+								Registry.ConvertExecuteContextToBaseType(PreferredType.TypeIndex);								
+							}
+						}
+					}
+
 					// Remove preferred types that do not match the pin type
 					{
-						TemplateNode->PreferredPermutationPairs.RemoveAll([TemplateNode](const FRigVMTemplatePreferredType& PreferredType)
+						int32 NumRemoved = TemplateNode->PreferredPermutationPairs.RemoveAll([TemplateNode](const FRigVMTemplatePreferredType& PreferredType)
 						{
 							if (URigVMPin* Pin = TemplateNode->FindPin(PreferredType.GetArgument().ToString()))
 							{
@@ -1464,7 +1487,12 @@ void UControlRigBlueprint::RefreshAllModels(EControlRigBlueprintLoadType InLoadT
 								}
 							}
 							return true;
-						});						
+						});
+
+						if (NumRemoved > 0)
+						{
+							bDirtyDuringLoad = true;
+						}
 					}
 
 					if(!TemplateNode->HasWildCardPin() &&
diff --git a/Engine/Source/Developer/RigVMDeveloper/Public/RigVMModel/Nodes/RigVMTemplateNode.h b/Engine/Source/Developer/RigVMDeveloper/Public/RigVMModel/Nodes/RigVMTemplateNode.h
--- a/Engine/Source/Developer/RigVMDeveloper/Public/RigVMModel/Nodes/RigVMTemplateNode.h
+++ b/Engine/Source/Developer/RigVMDeveloper/Public/RigVMModel/Nodes/RigVMTemplateNode.h
@@ -59,6 +59,7 @@ protected:
 
 	friend class URigVMTemplateNode;
 	friend class URigVMController;
+	friend class UControlRigBlueprint;
 };
 
 /**
diff --git a/Engine/Source/Runtime/RigVM/Private/RigVMCore/RigVMRegistry.cpp b/Engine/Source/Runtime/RigVM/Private/RigVMCore/RigVMRegistry.cpp
--- a/Engine/Source/Runtime/RigVM/Private/RigVMCore/RigVMRegistry.cpp
+++ b/Engine/Source/Runtime/RigVM/Private/RigVMCore/RigVMRegistry.cpp
@@ -522,6 +522,39 @@ bool FRigVMRegistry::IsExecuteType(TRigVMTypeIndex InTypeIndex) const
 	return false;
 }
 
+bool FRigVMRegistry::ConvertExecuteContextToBaseType(TRigVMTypeIndex& InOutTypeIndex) const
+{
+	if(InOutTypeIndex == INDEX_NONE)
+	{
+		return false;
+	}
+		
+	if(InOutTypeIndex == RigVMTypeUtils::TypeIndex::Execute) 
+	{
+		return true;
+	}
+
+	if(!IsExecuteType(InOutTypeIndex))
+	{
+		return false;
+	}
+
+	// execute arguments can have various execute context types. but we always
+	// convert them to the base execute type to make matching types easier later.
+	// this means that the execute argument in every permutations shares 
+	// the same type index of RigVMTypeUtils::TypeIndex::Execute
+	if(IsArrayType(InOutTypeIndex))
+	{
+		InOutTypeIndex = GetArrayTypeFromBaseTypeIndex(RigVMTypeUtils::TypeIndex::Execute);
+	}
+	else
+	{
+		InOutTypeIndex = RigVMTypeUtils::TypeIndex::Execute;
+	}
+
+	return true;
+}
+
 int32 FRigVMRegistry::GetArrayDimensionsForType(TRigVMTypeIndex InTypeIndex) const
 {
 	if(ensure(Types.IsValidIndex(InTypeIndex)))
diff --git a/Engine/Source/Runtime/RigVM/Private/RigVMCore/RigVMTemplate.cpp b/Engine/Source/Runtime/RigVM/Private/RigVMCore/RigVMTemplate.cpp
--- a/Engine/Source/Runtime/RigVM/Private/RigVMCore/RigVMTemplate.cpp
+++ b/Engine/Source/Runtime/RigVM/Private/RigVMCore/RigVMTemplate.cpp
@@ -9,39 +9,6 @@
 
 ////////////////////////////////////////////////////////////////////////////////////////////////
 
-static void EnsureValidExecuteType_Internal(TRigVMTypeIndex& InTypeIndex)
-{
-	const FRigVMRegistry& Registry = FRigVMRegistry::Get();
-	if(InTypeIndex == INDEX_NONE)
-	{
-		return;
-	}
-		
-	if(InTypeIndex == RigVMTypeUtils::TypeIndex::Execute) 
-	{
-		return;
-	}
-
-	if(!Registry.IsExecuteType(InTypeIndex))
-	{
-		return;
-	}
-
-	// execute arguments can have various execute context types. but we always
-	// convert them to the base execute type to make matching types easier later.
-	// this means that the execute argument in every permutations shares 
-	// the same type index of RigVMTypeUtils::TypeIndex::Execute
-	if(Registry.IsArrayType(InTypeIndex))
-	{
-		InTypeIndex = Registry.GetArrayTypeFromBaseTypeIndex(RigVMTypeUtils::TypeIndex::Execute);
-	}
-	else
-	{
-		InTypeIndex = RigVMTypeUtils::TypeIndex::Execute;
-	}	
-}
-
-
 FRigVMTemplateArgument::FRigVMTemplateArgument()
 	: Index(INDEX_NONE)
 	, Name(NAME_None)
@@ -202,9 +169,10 @@ void FRigVMTemplateArgument::Serialize(FArchive& Ar)
 
 void FRigVMTemplateArgument::EnsureValidExecuteType()
 {
+	const FRigVMRegistry& Registry = FRigVMRegistry::Get();
 	for(TRigVMTypeIndex& TypeIndex : TypeIndices)
 	{
-		EnsureValidExecuteType_Internal(TypeIndex);
+		Registry.ConvertExecuteContextToBaseType(TypeIndex);
 	}
 }
 
@@ -1442,10 +1410,11 @@ bool FRigVMTemplate::AddTypeForArgument(const FName& InArgumentName, TRigVMTypeI
 		FRigVMTemplateTypeMap Types = OnNewArgumentType().Execute(this, InArgumentName, InTypeIndex);
 		if(Types.Num() == Arguments.Num())
 		{
+			const FRigVMRegistry& Registry = FRigVMRegistry::Get();
 			for (TPair<FName, TRigVMTypeIndex>& ArgumentAndType : Types)
 			{
 				// similar to FRigVMTemplateArgument::EnsureValidExecuteType
-				EnsureValidExecuteType_Internal(ArgumentAndType.Value);
+				Registry.ConvertExecuteContextToBaseType(ArgumentAndType.Value);
 			}
 			
 			// make sure this permutation doesn't exist yet
diff --git a/Engine/Source/Runtime/RigVM/Public/RigVMCore/RigVMRegistry.h b/Engine/Source/Runtime/RigVM/Public/RigVMCore/RigVMRegistry.h
--- a/Engine/Source/Runtime/RigVM/Public/RigVMCore/RigVMRegistry.h
+++ b/Engine/Source/Runtime/RigVM/Public/RigVMCore/RigVMRegistry.h
@@ -154,6 +154,9 @@ public:
 	// Returns true if the type is an execute type
 	bool IsExecuteType(TRigVMTypeIndex InTypeIndex) const;
 
+	// Converts the given execute context type to the base execute context type
+	bool ConvertExecuteContextToBaseType(TRigVMTypeIndex& InOutTypeIndex) const;
+
 	// Returns the dimensions of the array 
 	int32 GetArrayDimensionsForType(TRigVMTypeIndex InTypeIndex) const;
 
