RigVM: making sure that we find the old memory class and post load it before trying to replace it.

diff --git a/Engine/Source/Runtime/RigVM/Private/RigVMCore/RigVMMemoryStorage.cpp b/Engine/Source/Runtime/RigVM/Private/RigVMCore/RigVMMemoryStorage.cpp
--- a/Engine/Source/Runtime/RigVM/Private/RigVMCore/RigVMMemoryStorage.cpp
+++ b/Engine/Source/Runtime/RigVM/Private/RigVMCore/RigVMMemoryStorage.cpp
@@ -473,9 +473,23 @@ URigVMMemoryStorageGeneratorClass* URigVMMemoryStorageGeneratorClass::CreateStor
 	const FString& ClassName = GetClassName(InMemoryType);
 
 	// if there's an old class - remove it from the package and mark it to destroy
-	URigVMMemoryStorageGeneratorClass* OldClass = FindObject<URigVMMemoryStorageGeneratorClass>(Package, *ClassName);
+	//
+	// we need to use StaticFindObjectFastInternal with ExclusiveInternalFlags = EInternalObjectFlags::None
+	// here because objects with RF_NeedPostLoad cannot be found by regular FindObject calls as they will have
+	// ExclusiveInternalFlags = EInternalObjectFlags::AsyncLoading when called from game thread.
+	//
+	// in theory, whenever this function is called as part of OnEndLoadPackage -> CRBP::HandlePackageDone,
+	// all objects within the package should have been fully loaded, but it turns out that OnEndLoadPackage
+	// does not guarantee that condition. As a result we will do OldClass->ConditionalPostLoad() to ensure the OldClass
+	// is fully loaded, ready to be replaced.
+	
+	UObject* OldClass = StaticFindObjectFastInternal( /*Class=*/ NULL, Package, *ClassName, true, RF_NoFlags, EInternalObjectFlags::None);
 	if(OldClass)
 	{
+		// ensure the OldClass is completely loaded so that we can remove it and replace it with the new class
+		// otherwise we get an assertion trying to replace objects currently being loaded
+		OldClass->ConditionalPostLoad();
+		
 		FString DiscardedMemoryClassName;
 		static const TCHAR DiscardedMemoryClassTemplate[] = TEXT("DiscardedMemoryClassTemplate_%d");
 		static int32 DiscardedMemoryClassIndex = 0;
