Control Rig: Patch nodes referencing deprecated templates

diff --git a/Engine/Plugins/Animation/ControlRig/Source/ControlRigDeveloper/Private/ControlRigBlueprint.cpp b/Engine/Plugins/Animation/ControlRig/Source/ControlRigDeveloper/Private/ControlRigBlueprint.cpp
--- a/Engine/Plugins/Animation/ControlRig/Source/ControlRigDeveloper/Private/ControlRigBlueprint.cpp
+++ b/Engine/Plugins/Animation/ControlRig/Source/ControlRigDeveloper/Private/ControlRigBlueprint.cpp
@@ -1540,6 +1540,25 @@ void UControlRigBlueprint::RefreshAllModels(EControlRigBlueprintLoadType InLoadT
 							Controller->FullyResolveTemplateNode(UnitNode, INDEX_NONE, false);
 						}
 
+						// Try to find a deprecated template
+						if (UnitNode->GetScriptStruct() == nullptr && !UnitNode->TemplateNotation.IsNone())
+						{
+							const FRigVMTemplate* Template = FRigVMRegistry::Get().FindTemplate(UnitNode->TemplateNotation, true);
+							FRigVMTemplate::FTypeMap TypeMap;
+							for (URigVMPin* Pin : UnitNode->GetPins())
+							{
+								check(!Pin->IsWildCard());
+								TypeMap.Add(Pin->GetFName(), Pin->GetTypeIndex());
+							}
+
+							int32 Permutation;
+							if (Template->FullyResolve(TypeMap, Permutation))
+							{
+								const FRigVMFunction* Function = Template->GetPermutation(Permutation);
+								UnitNode->ResolvedFunctionName = Function->GetName();
+							}
+						}
+
 						if (UnitNode->GetScriptStruct() == nullptr)
 						{
 							static const TCHAR UnresolvedUnitNodeMessage[] = TEXT("Node %s could not be resolved.");
diff --git a/Engine/Source/Runtime/RigVM/Private/RigVMCore/RigVMRegistry.cpp b/Engine/Source/Runtime/RigVM/Private/RigVMCore/RigVMRegistry.cpp
--- a/Engine/Source/Runtime/RigVM/Private/RigVMCore/RigVMRegistry.cpp
+++ b/Engine/Source/Runtime/RigVM/Private/RigVMCore/RigVMRegistry.cpp
@@ -771,11 +771,10 @@ void FRigVMRegistry::Register(const TCHAR* InName, FRigVMFunctionPtr InFunctionP
 	FString TemplateMetadata;
 	if (InStruct->GetStringMetaDataHierarchical(TemplateNameMetaName, &TemplateMetadata))
 	{
-		if(InStruct->HasMetaData(FRigVMStruct::DeprecatedMetaName))
-		{
-			return;
-		}
-
+		bool bIsDeprecated = InStruct->HasMetaData(FRigVMStruct::DeprecatedMetaName);
+		TChunkedArray<FRigVMTemplate>& TemplateArray = (bIsDeprecated) ? DeprecatedTemplates : Templates;
+		TMap<FName, int32>& NotationToIndex = (bIsDeprecated) ? DeprecatedTemplateNotationToIndex : TemplateNotationToIndex;
+		
 		FString MethodName;
 		if (FString(InName).Split(TEXT("::"), nullptr, &MethodName))
 		{
@@ -785,10 +784,10 @@ void FRigVMRegistry::Register(const TCHAR* InName, FRigVMFunctionPtr InFunctionP
 			{
 				bool bWasMerged = false;
 
-				const int32* ExistingTemplateIndexPtr = TemplateNotationToIndex.Find(Template.GetNotation());
+				const int32* ExistingTemplateIndexPtr = NotationToIndex.Find(Template.GetNotation());
 				if(ExistingTemplateIndexPtr)
 				{
-					FRigVMTemplate& ExistingTemplate = Templates[*ExistingTemplateIndexPtr];
+					FRigVMTemplate& ExistingTemplate = TemplateArray[*ExistingTemplateIndexPtr];
 					if (ExistingTemplate.Merge(Template))
 					{
 						Functions[Function.Index].TemplateIndex = ExistingTemplate.Index;
@@ -798,13 +797,13 @@ void FRigVMRegistry::Register(const TCHAR* InName, FRigVMFunctionPtr InFunctionP
 
 				if (!bWasMerged)
 				{
-					Template.Index = Templates.Num();
+					Template.Index = TemplateArray.Num();
 					Functions[Function.Index].TemplateIndex = Template.Index;
-					Templates.AddElement(Template);
+					TemplateArray.AddElement(Template);
 					
 					if(ExistingTemplateIndexPtr == nullptr)
 					{
-						TemplateNotationToIndex.Add(Template.GetNotation(), Template.Index);
+						NotationToIndex.Add(Template.GetNotation(), Template.Index);
 					}
 				}
 			}
@@ -920,7 +919,7 @@ const TChunkedArray<FRigVMFunction>& FRigVMRegistry::GetFunctions() const
 	return Functions;
 }
 
-const FRigVMTemplate* FRigVMRegistry::FindTemplate(const FName& InNotation) const
+const FRigVMTemplate* FRigVMRegistry::FindTemplate(const FName& InNotation, bool bIncludeDeprecated) const
 {
 	if (InNotation.IsNone())
 	{
@@ -952,6 +951,14 @@ const FRigVMTemplate* FRigVMRegistry::FindTemplate(const FName& InNotation) cons
 		}
 	}
 
+	if (bIncludeDeprecated)
+	{
+		if(const int32* TemplateIndexPtr = DeprecatedTemplateNotationToIndex.Find(InNotation))
+		{
+			return &DeprecatedTemplates[*TemplateIndexPtr];
+		}
+	}
+
 	return nullptr;
 }
 
diff --git a/Engine/Source/Runtime/RigVM/Public/RigVMCore/RigVMRegistry.h b/Engine/Source/Runtime/RigVM/Public/RigVMCore/RigVMRegistry.h
--- a/Engine/Source/Runtime/RigVM/Public/RigVMCore/RigVMRegistry.h
+++ b/Engine/Source/Runtime/RigVM/Public/RigVMCore/RigVMRegistry.h
@@ -188,7 +188,7 @@ public:
 	const TChunkedArray<FRigVMFunction>& GetFunctions() const;
 
 	// Returns a template pointer given its notation (or nullptr)
-	const FRigVMTemplate* FindTemplate(const FName& InNotation) const;
+	const FRigVMTemplate* FindTemplate(const FName& InNotation, bool bIncludeDeprecated = false) const;
 
 	// Returns all current rigvm functions
 	const TChunkedArray<FRigVMTemplate>& GetTemplates() const;
@@ -265,18 +265,24 @@ private:
 	// We use TChunkedArray because we need the memory locations to be stable, since we only ever add and never remove.
 	TChunkedArray<FRigVMFunction> Functions;
 
-	// memory for all templates
+	// memory for all non-deprecated templates
 	TChunkedArray<FRigVMTemplate> Templates;
 
+	// memory for all deprecated templates
+	TChunkedArray<FRigVMTemplate> DeprecatedTemplates;
+
 	// memory for all dispatch factories
 	TArray<FRigVMDispatchFactory*> Factories;
 
 	// name lookup for functions
 	TMap<FName, int32> FunctionNameToIndex;
 
-	// name lookup for templates
+	// name lookup for non-deprecated templates
 	TMap<FName, int32> TemplateNotationToIndex;
 
+	// name lookup for deprecated templates
+	TMap<FName, int32> DeprecatedTemplateNotationToIndex;
+
 	// Maps storing the default types per type category
 	TMap<FRigVMTemplateArgument::ETypeCategory, TArray<TRigVMTypeIndex>> TypesPerCategory;
 
