Control Rig: Fix crash when adding deprecated template node through python

diff --git a/Engine/Source/Runtime/RigVM/Private/RigVMCore/RigVMRegistry.cpp b/Engine/Source/Runtime/RigVM/Private/RigVMCore/RigVMRegistry.cpp
--- a/Engine/Source/Runtime/RigVM/Private/RigVMCore/RigVMRegistry.cpp
+++ b/Engine/Source/Runtime/RigVM/Private/RigVMCore/RigVMRegistry.cpp
@@ -790,7 +790,10 @@ void FRigVMRegistry::Register(const TCHAR* InName, FRigVMFunctionPtr InFunctionP
 					FRigVMTemplate& ExistingTemplate = TemplateArray[*ExistingTemplateIndexPtr];
 					if (ExistingTemplate.Merge(Template))
 					{
-						Functions[Function.Index].TemplateIndex = ExistingTemplate.Index;
+						if (!bIsDeprecated)
+						{
+							Functions[Function.Index].TemplateIndex = ExistingTemplate.Index;
+						}
 						bWasMerged = true;
 					}
 				}
@@ -798,7 +801,10 @@ void FRigVMRegistry::Register(const TCHAR* InName, FRigVMFunctionPtr InFunctionP
 				if (!bWasMerged)
 				{
 					Template.Index = TemplateArray.Num();
-					Functions[Function.Index].TemplateIndex = Template.Index;
+					if (!bIsDeprecated)
+					{
+						Functions[Function.Index].TemplateIndex = Template.Index;
+					}
 					TemplateArray.AddElement(Template);
 					
 					if(ExistingTemplateIndexPtr == nullptr)
