Control rig: Print literal and work memory after each vm execution instruction

diff --git a/Engine/Source/Runtime/RigVM/Private/RigVMCore/RigVM.cpp b/Engine/Source/Runtime/RigVM/Private/RigVMCore/RigVM.cpp
--- a/Engine/Source/Runtime/RigVM/Private/RigVMCore/RigVM.cpp
+++ b/Engine/Source/Runtime/RigVM/Private/RigVMCore/RigVM.cpp
@@ -1796,6 +1796,21 @@ bool URigVM::Execute(TArrayView<URigVMMemoryStorage*> Memory, TArrayView<void*>
 	{
 		StartCycles = FPlatformTime::Cycles64();
 	}
+	
+#if UE_RIGVM_DEBUG_EXECUTION
+	FString DebugMemoryString;
+	TArray<FString> PreviousWorkMemory;
+	UEnum* InstanceOpCodeEnum = StaticEnum<ERigVMOpCode>();
+	URigVMMemoryStorage* LiteralMemory = GetLiteralMemory(false);
+	DebugMemoryString = FString("\n\nLiteral Memory\n\n");
+	for (int32 PropertyIndex=0; PropertyIndex<LiteralMemory->Num(); ++PropertyIndex)
+	{
+		DebugMemoryString += FString::Printf(TEXT("%s: %s\n"), *LiteralMemory->GetProperties()[PropertyIndex]->GetFullName(), *LiteralMemory->GetDataAsString(PropertyIndex));				
+	}
+	DebugMemoryString += FString(TEXT("\n\nWork Memory\n\n"));
+	
+#endif
+	
 #endif
 
 	while (Instructions.IsValidIndex(Context.PublicData.InstructionIndex))
@@ -1816,6 +1831,34 @@ bool URigVM::Execute(TArrayView<URigVMMemoryStorage*> Memory, TArrayView<void*>
 
 		const FRigVMInstruction& Instruction = Instructions[Context.PublicData.InstructionIndex];
 
+#if WITH_EDITOR
+#if UE_RIGVM_DEBUG_EXECUTION
+		if (Instruction.OpCode >= ERigVMOpCode::Execute_0_Operands && Instruction.OpCode <= ERigVMOpCode::Execute_64_Operands)
+		{
+			const FRigVMExecuteOp& Op = ByteCode.GetOpAt<FRigVMExecuteOp>(Instruction);
+			FRigVMOperandArray Operands = ByteCode.GetOperandsForExecuteOp(Instructions[Context.PublicData.InstructionIndex]);
+
+			TArray<FString> Labels;
+			for (const FRigVMOperand& Operand : Operands)
+			{
+				Labels.Add(GetOperandLabel(Operand));
+			}
+
+			DebugMemoryString += FString::Printf(TEXT("Instruction %d: %s(%s)\n"), Context.PublicData.InstructionIndex, *FunctionNames[Op.FunctionIndex].ToString(), *FString::Join(Labels, TEXT(", ")));				
+		}
+		else if(Instruction.OpCode == ERigVMOpCode::Copy)
+		{
+			const FRigVMCopyOp& Op = ByteCode.GetOpAt<FRigVMCopyOp>(Instruction);
+			DebugMemoryString += FString::Printf(TEXT("Instruction %d: Copy %s -> %s\n"), Context.PublicData.InstructionIndex, *GetOperandLabel(Op.Source), *GetOperandLabel(Op.Target));				
+		}
+		else
+		{
+			DebugMemoryString += FString::Printf(TEXT("Instruction %d: %s\n"), Context.PublicData.InstructionIndex, *InstanceOpCodeEnum->GetNameByIndex((uint8)Instruction.OpCode).ToString());				
+		}
+		
+#endif
+#endif
+
 		switch (Instruction.OpCode)
 		{
 			case ERigVMOpCode::Execute_0_Operands:
@@ -2071,6 +2114,9 @@ bool URigVM::Execute(TArrayView<URigVMMemoryStorage*> Memory, TArrayView<void*>
 						DebugInfo->SetCurrentActiveBreakpoint(FRigVMBreakpoint());
 						ExecutionHalted().Broadcast(INDEX_NONE, nullptr, InEntryName);
 					}
+#if UE_RIGVM_DEBUG_EXECUTION
+					Context.PublicData.Log(EMessageSeverity::Info, DebugMemoryString);					
+#endif
 #endif
 				}
 				return true;
@@ -2681,6 +2727,29 @@ bool URigVM::Execute(TArrayView<URigVMMemoryStorage*> Memory, TArrayView<void*>
 			StartCycles = EndCycles;
 			OverallCycles += Cycles;
 		}
+
+#if UE_RIGVM_DEBUG_EXECUTION
+		TArray<FString> CurrentWorkMemory;
+		URigVMMemoryStorage* WorkMemory = GetWorkMemory(false);
+		for (int32 PropertyIndex=0; PropertyIndex<WorkMemory->Num(); ++PropertyIndex)
+		{
+			FString Line = FString::Printf(TEXT("%s: %s"), *WorkMemory->GetProperties()[PropertyIndex]->GetFullName(), *WorkMemory->GetDataAsString(PropertyIndex));
+			if (PreviousWorkMemory.Num() > 0 && PreviousWorkMemory[PropertyIndex].StartsWith(TEXT(" -- ")))
+			{
+				PreviousWorkMemory[PropertyIndex].RightChopInline(4);
+			}
+			if (PreviousWorkMemory.Num() == 0 || Line == PreviousWorkMemory[PropertyIndex])
+			{
+				CurrentWorkMemory.Add(Line);
+			}
+			else
+			{
+				CurrentWorkMemory.Add(FString::Printf(TEXT(" -- %s"), *Line));
+			}
+		}
+		DebugMemoryString += FString::Join(CurrentWorkMemory, TEXT("\n")) + FString(TEXT("\n\n"));
+		PreviousWorkMemory = CurrentWorkMemory;		
+#endif
 #endif
 	}
 
diff --git a/Engine/Source/Runtime/RigVM/Public/RigVMDefines.h b/Engine/Source/Runtime/RigVM/Public/RigVMDefines.h
--- a/Engine/Source/Runtime/RigVM/Public/RigVMDefines.h
+++ b/Engine/Source/Runtime/RigVM/Public/RigVMDefines.h
@@ -25,3 +25,7 @@
 #ifndef UE_RIGVM_DEBUG_TYPEINDEX
 #define UE_RIGVM_DEBUG_TYPEINDEX 0
 #endif
+
+#ifndef UE_RIGVM_DEBUG_EXECUTION
+#define UE_RIGVM_DEBUG_EXECUTION 0
+#endif
