Control Rig: Fix templates with duplicate permutations

diff --git a/Engine/Source/Developer/RigVMDeveloper/Private/RigVMModel/RigVMController.cpp b/Engine/Source/Developer/RigVMDeveloper/Private/RigVMModel/RigVMController.cpp
--- a/Engine/Source/Developer/RigVMDeveloper/Private/RigVMModel/RigVMController.cpp
+++ b/Engine/Source/Developer/RigVMDeveloper/Private/RigVMModel/RigVMController.cpp
@@ -17045,8 +17045,7 @@ bool URigVMController::UpdateTemplateNodePinTypes(URigVMTemplateNode* InNode, bo
 			}
 			
 			PreferredPermutations = InNode->FindPermutationsForTypes(PreferredPermutationPairs, false);
-			// ensure(PreferredPermutations.Num() == 1);
-			check(!PreferredPermutations.IsEmpty());
+			ensureMsgf(PreferredPermutations.Num() == 1, TEXT("Number of preferred permutations is not 1 in %s"), *GetPackage()->GetPathName());
 			if (EntryNode)
 			{
 				EntryNode->ResolvedPermutation = PreferredPermutations[0];
@@ -17059,8 +17058,7 @@ bool URigVMController::UpdateTemplateNodePinTypes(URigVMTemplateNode* InNode, bo
 		else
 		{
 			PreferredPermutations = InNode->FindPermutationsForTypes(InNode->PreferredPermutationPairs, false);
-			// ensure(PreferredPermutations.Num() == 1);
-			check(!PreferredPermutations.IsEmpty());
+			ensureMsgf(PreferredPermutations.Num() == 1, TEXT("Number of preferred permutations is not 1 in %s"), *GetPackage()->GetPathName());
 			InNode->ResolvedPermutation = PreferredPermutations[0];
 		}			
 	}
diff --git a/Engine/Source/Runtime/RigVM/Private/RigVMCore/RigVMRegistry.cpp b/Engine/Source/Runtime/RigVM/Private/RigVMCore/RigVMRegistry.cpp
--- a/Engine/Source/Runtime/RigVM/Private/RigVMCore/RigVMRegistry.cpp
+++ b/Engine/Source/Runtime/RigVM/Private/RigVMCore/RigVMRegistry.cpp
@@ -794,7 +794,7 @@ const TArray<TRigVMTypeIndex>& FRigVMRegistry::GetCompatibleTypes(TRigVMTypeInde
 	return EmptyTypes;
 }
 
-const TArray<TRigVMTypeIndex>& FRigVMRegistry::GetTypesForCategory(FRigVMTemplateArgument::ETypeCategory InCategory)
+const TArray<TRigVMTypeIndex>& FRigVMRegistry::GetTypesForCategory(FRigVMTemplateArgument::ETypeCategory InCategory) const
 {
 	check(InCategory != FRigVMTemplateArgument::ETypeCategory_Invalid);
 	return TypesPerCategory.FindChecked(InCategory);
@@ -1159,11 +1159,6 @@ const FRigVMTemplate* FRigVMRegistry::GetOrAddTemplateFromArguments(const FName&
 				Argument.TypeCategories.Add(FRigVMTemplateArgument::ETypeCategory_SingleAnyValue);
 			}
 
-			for (int32 i = 0; i < Argument.TypeIndices.Num(); ++i)
-			{
-				Argument.TypeToPermutations.Add(Argument.TypeIndices[i], {i});				
-			}
-			
 			NumPermutations = FMath::Max(NumPermutations, Argument.TypeIndices.Num()); 
 		}
 	}
@@ -1177,18 +1172,52 @@ const FRigVMTemplate* FRigVMRegistry::GetOrAddTemplateFromArguments(const FName&
 			{
 				const int32 TypeIndex = Argument.TypeIndices[0];
 				Argument.TypeIndices.SetNum(NumPermutations);
-				TArray<int32> ArgTypePermutations;
-				ArgTypePermutations.SetNum(NumPermutations);
 				for(int32 Index=0;Index<NumPermutations;Index++)
 				{
 					Argument.TypeIndices[Index] = TypeIndex;
-					ArgTypePermutations[Index] = Index;
 				}
-				Argument.TypeToPermutations.Add(TypeIndex, ArgTypePermutations);
 			}
 		}
 	}
 
+	// Remove duplicate permutations
+	{
+		TArray<int32> ToRemove;
+		TArray<TArray<TRigVMTypeIndex>> PermutationTypes;
+		PermutationTypes.Reserve(NumPermutations);
+		for(int32 Index=0;Index<NumPermutations;Index++)
+		{
+			TArray<TRigVMTypeIndex> ArgTypes;
+			ArgTypes.Reserve(Template.Arguments.Num());
+			for(FRigVMTemplateArgument& Argument : Template.Arguments)
+			{
+				ArgTypes.Add(Argument.TypeIndices[Index]);
+			}
+
+			if (PermutationTypes.Contains(ArgTypes))
+			{
+				ToRemove.Add(Index);
+			}
+			else
+			{
+				PermutationTypes.Add(ArgTypes);
+			}
+		}
+		for (int32 i=ToRemove.Num()-1; i>=0; --i)
+		{
+			for(FRigVMTemplateArgument& Argument : Template.Arguments)
+			{
+				Argument.TypeIndices.RemoveAt(ToRemove[i]);
+			}
+		}
+		NumPermutations -= ToRemove.Num();
+	}
+	
+	for(FRigVMTemplateArgument& Argument : Template.Arguments)
+	{
+		Argument.UpdateTypeToPermutations();
+	}
+
 	Template.Permutations.SetNum(NumPermutations);
 	for(int32 Index=0;Index<NumPermutations;Index++)
 	{
diff --git a/Engine/Source/Runtime/RigVM/Private/RigVMCore/RigVMTemplate.cpp b/Engine/Source/Runtime/RigVM/Private/RigVMCore/RigVMTemplate.cpp
--- a/Engine/Source/Runtime/RigVM/Private/RigVMCore/RigVMTemplate.cpp
+++ b/Engine/Source/Runtime/RigVM/Private/RigVMCore/RigVMTemplate.cpp
@@ -190,6 +190,7 @@ void FRigVMTemplateArgument::EnsureValidExecuteType()
 
 void FRigVMTemplateArgument::UpdateTypeToPermutations()
 {
+	TypeToPermutations.Reset();
 	for(int32 TypeIndex=0;TypeIndex<TypeIndices.Num();TypeIndex++)
 	{
 		if (TArray<int32>* Permutations = TypeToPermutations.Find(TypeIndices[TypeIndex]))
diff --git a/Engine/Source/Runtime/RigVM/Public/RigVMCore/RigVMRegistry.h b/Engine/Source/Runtime/RigVM/Public/RigVMCore/RigVMRegistry.h
--- a/Engine/Source/Runtime/RigVM/Public/RigVMCore/RigVMRegistry.h
+++ b/Engine/Source/Runtime/RigVM/Public/RigVMCore/RigVMRegistry.h
@@ -179,7 +179,7 @@ public:
 	const TArray<TRigVMTypeIndex>& GetCompatibleTypes(TRigVMTypeIndex InTypeIndex) const;
 
 	// Returns all compatible types given a category
-	const TArray<TRigVMTypeIndex>& GetTypesForCategory(FRigVMTemplateArgument::ETypeCategory InCategory);
+	const TArray<TRigVMTypeIndex>& GetTypesForCategory(FRigVMTemplateArgument::ETypeCategory InCategory) const;
 
 	// Returns the type index of the array matching the given element type index
 	TRigVMTypeIndex GetArrayTypeFromBaseTypeIndex(TRigVMTypeIndex InTypeIndex) const;
