Fix FRigVMRegistry trying to load every asset in a plugin being unmounted or unloaded

diff --git a/Engine/Source/Runtime/RigVM/Private/RigVMCore/RigVMRegistry.cpp b/Engine/Source/Runtime/RigVM/Private/RigVMCore/RigVMRegistry.cpp
--- a/Engine/Source/Runtime/RigVM/Private/RigVMCore/RigVMRegistry.cpp
+++ b/Engine/Source/Runtime/RigVM/Private/RigVMCore/RigVMRegistry.cpp
@@ -197,7 +197,8 @@ void FRigVMRegistry::Refresh()
 
 void FRigVMRegistry::OnAssetRemoved(const FAssetData& InAssetData)
 {
-	UObject* Asset = InAssetData.GetAsset();
+	const bool bLoadAsset = false;
+	UObject* Asset = InAssetData.FastGetAsset(bLoadAsset);
 
 	if (UObjectRedirector* RedirectedStruct = Cast<UObjectRedirector>(Asset))
 	{
