Control Rig: Print debug vm when halted by breakpoint

diff --git a/Engine/Source/Runtime/RigVM/Private/RigVMCore/RigVM.cpp b/Engine/Source/Runtime/RigVM/Private/RigVMCore/RigVM.cpp
--- a/Engine/Source/Runtime/RigVM/Private/RigVMCore/RigVM.cpp
+++ b/Engine/Source/Runtime/RigVM/Private/RigVMCore/RigVM.cpp
@@ -1818,6 +1818,9 @@ bool URigVM::Execute(TArrayView<URigVMMemoryStorage*> Memory, TArrayView<void*>
 #if WITH_EDITOR
 		if (ShouldHaltAtInstruction(InEntryName, Context.PublicData.InstructionIndex))
 		{
+#if UE_RIGVM_DEBUG_EXECUTION
+			Context.PublicData.Log(EMessageSeverity::Info, DebugMemoryString);					
+#endif
 			// we'll recursively exit all invoked
 			// entries here.
 			return true;
