ControlRig: RigVMRegistry avoid calling GetAsset OnAssetRemoved, which could load otherwise unloaded assets. Changed to use string lookup based solution to know which type to remove from the registry

diff --git a/Engine/Source/Runtime/RigVM/Private/RigVMCore/RigVMRegistry.cpp b/Engine/Source/Runtime/RigVM/Private/RigVMCore/RigVMRegistry.cpp
--- a/Engine/Source/Runtime/RigVM/Private/RigVMCore/RigVMRegistry.cpp
+++ b/Engine/Source/Runtime/RigVM/Private/RigVMCore/RigVMRegistry.cpp
@@ -133,6 +133,7 @@ void FRigVMRegistry::InitializeIfNeeded()
 			if (AssetRegistryModule->TryGet())
 			{
 				AssetRegistryModule->Get().OnAssetRemoved().RemoveAll(&s_RigVMRegistry);
+				AssetRegistryModule->Get().OnAssetRenamed().RemoveAll(&s_RigVMRegistry);
 			}
 		}
 
@@ -141,6 +142,7 @@ void FRigVMRegistry::InitializeIfNeeded()
 	
 	FAssetRegistryModule& AssetRegistryModule = FModuleManager::LoadModuleChecked<FAssetRegistryModule>(TEXT("AssetRegistry"));
 	AssetRegistryModule.Get().OnAssetRemoved().AddRaw(this, &FRigVMRegistry::OnAssetRemoved);
+	AssetRegistryModule.Get().OnAssetRenamed().AddRaw(this, &FRigVMRegistry::OnAssetRenamed);
 
 	UE::Anim::AttributeTypes::GetOnAttributeTypesChanged().AddRaw(this, &FRigVMRegistry::OnAnimationAttributeTypesChanged);
 }
@@ -195,28 +197,21 @@ void FRigVMRegistry::Refresh()
 	}
 }
 
-void FRigVMRegistry::OnAssetRemoved(const FAssetData& InAssetData)
+void FRigVMRegistry::OnAssetRenamed(const FAssetData& InAssetData, const FString& InOldObjectPath)
 {
-	const bool bLoadAsset = false;
-	UObject* Asset = InAssetData.FastGetAsset(bLoadAsset);
-
-	if (UObjectRedirector* RedirectedStruct = Cast<UObjectRedirector>(Asset))
-	{
-		Asset = RedirectedStruct->DestinationObject;
-	}
-
-	bool bTypeRemoved = false;
+	const FSoftObjectPath OldPath(InOldObjectPath);
 	
-	if (UScriptStruct* ScriptStruct = Cast<UScriptStruct>(Asset))
+	if (const TRigVMTypeIndex* TypeIndexPtr = UserDefinedTypeToIndex.Find(OldPath))
 	{
-		bTypeRemoved = RemoveType(ScriptStruct);
-	}
-	else if (UEnum* Enum = Cast<UEnum>(Asset))
-	{
-		bTypeRemoved = RemoveType(Enum);
+		const TRigVMTypeIndex TypeIndex = *TypeIndexPtr;
+		UserDefinedTypeToIndex.Remove(OldPath);
+		UserDefinedTypeToIndex.Add(InAssetData.ToSoftObjectPath()) = TypeIndex;
 	}
+}
 
-	if (bTypeRemoved)
+void FRigVMRegistry::OnAssetRemoved(const FAssetData& InAssetData)
+{
+	if (RemoveType(InAssetData))
 	{
 		OnRigVMRegistryChangedDelegate.Broadcast();
 	}
@@ -499,8 +494,18 @@ TRigVMTypeIndex FRigVMRegistry::FindOrAddType_Internal(const FRigVMTemplateArgum
 #endif
 			}			
 		}
-
-		return GetTypeIndex(InType);
+		
+		Index = GetTypeIndex(InType);
+		if (IsValid(CPPTypeObject))
+		{
+			if (CPPTypeObject->IsA<UUserDefinedStruct>() || CPPTypeObject->IsA<UUserDefinedEnum>())
+			{
+				// used to track name changes to user defined types
+				UserDefinedTypeToIndex.FindOrAdd(CPPTypeObject) = Index;
+			}
+		}
+		
+		return Index;
 	}
 	
 	return Index;
@@ -522,96 +527,104 @@ void FRigVMRegistry::RegisterTypeInCategory(FRigVMTemplateArgument::ETypeCategor
 	}
 }
 
-bool FRigVMRegistry::RemoveType(const FRigVMTemplateArgumentType& InType)
+bool FRigVMRegistry::RemoveType(const FAssetData& InAssetData)
 {
-	check(!InType.IsArray())
-	
-	TRigVMTypeIndex Index = GetTypeIndex(InType);
+	const FSoftObjectPath AssetPath = InAssetData.ToSoftObjectPath();
+	const UClass* TypeClass = InAssetData.GetClass();
 	
-	if(Index == INDEX_NONE)
+	if (const TRigVMTypeIndex* TypeIndexPtr = UserDefinedTypeToIndex.Find(AssetPath))
 	{
-		return false;
-	}
+		const TRigVMTypeIndex TypeIndex = *TypeIndexPtr;
+		
+		UserDefinedTypeToIndex.Remove(AssetPath);
+		
+		if(TypeIndex == INDEX_NONE)
+		{
+			return false;
+		}
 
-	TArray<TRigVMTypeIndex> Indices;
-	Indices.Init(INDEX_NONE, 3);
-	Indices[0] = Index;
-	Indices[1] = GetArrayTypeFromBaseTypeIndex(Indices[0]);
+		check(!IsArrayType(TypeIndex));
 
-	// any type that can be removed should have 3 entries in the registry
-	if (ensure(Indices[1] != INDEX_NONE))
-	{
-		Indices[2] = GetArrayTypeFromBaseTypeIndex(Indices[1]);
-	}
-	
-	const UObject* CPPTypeObject = InType.CPPTypeObject;
+		TArray<TRigVMTypeIndex> Indices;
+		Indices.Init(INDEX_NONE, 3);
+		Indices[0] = TypeIndex;
+		Indices[1] = GetArrayTypeFromBaseTypeIndex(Indices[0]);
 
-	for (int32 ArrayDimension=0; ArrayDimension<3; ++ArrayDimension)
-	{
-		Index = Indices[ArrayDimension];
-		
-		if (Index == INDEX_NONE)
+		// any type that can be removed should have 3 entries in the registry
+		if (ensure(Indices[1] != INDEX_NONE))
 		{
-			break;
+			Indices[2] = GetArrayTypeFromBaseTypeIndex(Indices[1]);
 		}
 		
-		if(const UUserDefinedEnum* Enum = Cast<UUserDefinedEnum>(CPPTypeObject))
+		for (int32 ArrayDimension=0; ArrayDimension<3; ++ArrayDimension)
 		{
-			switch(ArrayDimension)
+			const TRigVMTypeIndex Index = Indices[ArrayDimension];
+			
+			if (Index == INDEX_NONE)
 			{
-			default:
-			case 0:
-				{
-					RemoveTypeInCategory(FRigVMTemplateArgument::ETypeCategory_SingleEnumValue, Index);
-					RemoveTypeInCategory(FRigVMTemplateArgument::ETypeCategory_SingleAnyValue, Index);
-					break;
-				}
-			case 1:
-				{
-					RemoveTypeInCategory(FRigVMTemplateArgument::ETypeCategory_ArrayEnumValue, Index);
-					RemoveTypeInCategory(FRigVMTemplateArgument::ETypeCategory_ArrayAnyValue, Index);
-					break;
-				}
-			case 2:
-				{
-					RemoveTypeInCategory(FRigVMTemplateArgument::ETypeCategory_ArrayArrayEnumValue, Index);
-					RemoveTypeInCategory(FRigVMTemplateArgument::ETypeCategory_ArrayArrayAnyValue, Index);
-					break;
-				}
+				break;
 			}
-		}
-		else if(const UUserDefinedStruct* Struct = Cast<UUserDefinedStruct>(CPPTypeObject))
-		{
-			switch(ArrayDimension)
+			
+			if(TypeClass == UUserDefinedEnum::StaticClass())
 			{
-			default:
-			case 0:
-				{
-					RemoveTypeInCategory(FRigVMTemplateArgument::ETypeCategory_SingleScriptStructValue, Index);
-					RemoveTypeInCategory(FRigVMTemplateArgument::ETypeCategory_SingleAnyValue, Index);
-					break;
-				}
-			case 1:
+				switch(ArrayDimension)
 				{
-					RemoveTypeInCategory(FRigVMTemplateArgument::ETypeCategory_ArrayScriptStructValue, Index);
-					RemoveTypeInCategory(FRigVMTemplateArgument::ETypeCategory_ArrayAnyValue, Index);
-					break;
+				default:
+				case 0:
+					{
+						RemoveTypeInCategory(FRigVMTemplateArgument::ETypeCategory_SingleEnumValue, Index);
+						RemoveTypeInCategory(FRigVMTemplateArgument::ETypeCategory_SingleAnyValue, Index);
+						break;
+					}
+				case 1:
+					{
+						RemoveTypeInCategory(FRigVMTemplateArgument::ETypeCategory_ArrayEnumValue, Index);
+						RemoveTypeInCategory(FRigVMTemplateArgument::ETypeCategory_ArrayAnyValue, Index);
+						break;
+					}
+				case 2:
+					{
+						RemoveTypeInCategory(FRigVMTemplateArgument::ETypeCategory_ArrayArrayEnumValue, Index);
+						RemoveTypeInCategory(FRigVMTemplateArgument::ETypeCategory_ArrayArrayAnyValue, Index);
+						break;
+					}
 				}
-			case 2:
+			}
+			else if(TypeClass == UUserDefinedStruct::StaticClass())
+			{
+				switch(ArrayDimension)
 				{
-					RemoveTypeInCategory(FRigVMTemplateArgument::ETypeCategory_ArrayArrayScriptStructValue, Index);
-					RemoveTypeInCategory(FRigVMTemplateArgument::ETypeCategory_ArrayArrayAnyValue, Index);
-					break;
+				default:
+				case 0:
+					{
+						RemoveTypeInCategory(FRigVMTemplateArgument::ETypeCategory_SingleScriptStructValue, Index);
+						RemoveTypeInCategory(FRigVMTemplateArgument::ETypeCategory_SingleAnyValue, Index);
+						break;
+					}
+				case 1:
+					{
+						RemoveTypeInCategory(FRigVMTemplateArgument::ETypeCategory_ArrayScriptStructValue, Index);
+						RemoveTypeInCategory(FRigVMTemplateArgument::ETypeCategory_ArrayAnyValue, Index);
+						break;
+					}
+				case 2:
+					{
+						RemoveTypeInCategory(FRigVMTemplateArgument::ETypeCategory_ArrayArrayScriptStructValue, Index);
+						RemoveTypeInCategory(FRigVMTemplateArgument::ETypeCategory_ArrayArrayAnyValue, Index);
+						break;
+					}
 				}
 			}
+
+			// remove the type from the registry entirely
+			TypeToIndex.Remove(GetType(Index));
+			Types[Index] = FTypeInfo();
 		}
 
-		// remove the type from the registry entirely
-		TypeToIndex.Remove(GetType(Index));
-		Types[Index] = FTypeInfo();
+		return true;	
 	}
-
-	return true;
+	
+	return false;
 }
 
 void FRigVMRegistry::RemoveTypeInCategory(FRigVMTemplateArgument::ETypeCategory InCategory, TRigVMTypeIndex InTypeIndex)
diff --git a/Engine/Source/Runtime/RigVM/Public/RigVMCore/RigVMRegistry.h b/Engine/Source/Runtime/RigVM/Public/RigVMCore/RigVMRegistry.h
--- a/Engine/Source/Runtime/RigVM/Public/RigVMCore/RigVMRegistry.h
+++ b/Engine/Source/Runtime/RigVM/Public/RigVMCore/RigVMRegistry.h
@@ -58,6 +58,9 @@ public:
 	// Refreshes the list and finds the function pointers
 	// based on the names.
 	void Refresh();
+
+	// Update the registry when types are renamed
+	void OnAssetRenamed(const FAssetData& InAssetData, const FString& InOldObjectPath);
 	
 	// Update the registry when old types are removed
     void OnAssetRemoved(const FAssetData& InAssetData);
@@ -78,7 +81,7 @@ public:
 
 	// Removes a type from the registry, and updates all dependent templates
 	// which also creates invalid permutations in templates that we should ignore
-	bool RemoveType(const FRigVMTemplateArgumentType& InType);
+	bool RemoveType(const FAssetData& InAssetData);
 
 	// Returns the type index given a type
 	TRigVMTypeIndex GetTypeIndex(const FRigVMTemplateArgumentType& InType) const;
@@ -307,6 +310,10 @@ private:
 	// Lookup per type category to know which argument to keep in sync
 	TMap<FRigVMTemplateArgument::ETypeCategory, TArray<TPair<int32,int32>>> ArgumentsPerCategory;
 
+	// Name loop up for user defined types since they can be deleted.
+	// When that happens, it won't be safe to reload deleted assets so only type names are reliable
+	TMap<FSoftObjectPath, TRigVMTypeIndex> UserDefinedTypeToIndex;
+	
 	// Notifies other system that types have been added/removed, and template permutations have been updated
 	FOnRigVMRegistryChanged OnRigVMRegistryChangedDelegate;
 	
