Control Rig: Print debug external variable values after a VM evaluation

diff --git a/Engine/Source/Runtime/RigVM/Private/RigVMCore/RigVM.cpp b/Engine/Source/Runtime/RigVM/Private/RigVMCore/RigVM.cpp
--- a/Engine/Source/Runtime/RigVM/Private/RigVMCore/RigVM.cpp
+++ b/Engine/Source/Runtime/RigVM/Private/RigVMCore/RigVM.cpp
@@ -2734,7 +2734,8 @@ bool URigVM::Execute(TArrayView<URigVMMemoryStorage*> Memory, TArrayView<void*>
 #if UE_RIGVM_DEBUG_EXECUTION
 		TArray<FString> CurrentWorkMemory;
 		URigVMMemoryStorage* WorkMemory = GetWorkMemory(false);
-		for (int32 PropertyIndex=0; PropertyIndex<WorkMemory->Num(); ++PropertyIndex)
+		int32 LineIndex = 0;
+		for (int32 PropertyIndex=0; PropertyIndex<WorkMemory->Num(); ++PropertyIndex, ++LineIndex)
 		{
 			FString Line = FString::Printf(TEXT("%s: %s"), *WorkMemory->GetProperties()[PropertyIndex]->GetFullName(), *WorkMemory->GetDataAsString(PropertyIndex));
 			if (PreviousWorkMemory.Num() > 0 && PreviousWorkMemory[PropertyIndex].StartsWith(TEXT(" -- ")))
@@ -2750,6 +2751,25 @@ bool URigVM::Execute(TArrayView<URigVMMemoryStorage*> Memory, TArrayView<void*>
 				CurrentWorkMemory.Add(FString::Printf(TEXT(" -- %s"), *Line));
 			}
 		}
+		for (const FRigVMExternalVariable& ExternalVariable : ExternalVariables)
+		{
+			FString Value;
+			ExternalVariable.Property->ExportTextItem_Direct(Value, ExternalVariable.Memory, nullptr, nullptr, PPF_None);
+			FString Line = FString::Printf(TEXT("External %s: %s"), *ExternalVariable.Name.ToString(), *Value);
+			if (PreviousWorkMemory.Num() > 0 && PreviousWorkMemory[LineIndex].StartsWith(TEXT(" -- ")))
+			{
+				PreviousWorkMemory[LineIndex].RightChopInline(4);
+			}
+			if (PreviousWorkMemory.Num() == 0 || Line == PreviousWorkMemory[LineIndex])
+			{
+				CurrentWorkMemory.Add(Line);
+			}
+			else
+			{
+				CurrentWorkMemory.Add(FString::Printf(TEXT(" -- %s"), *Line));
+			}
+			++LineIndex;
+		}
 		DebugMemoryString += FString::Join(CurrentWorkMemory, TEXT("\n")) + FString(TEXT("\n\n"));
 		PreviousWorkMemory = CurrentWorkMemory;		
 #endif
