ControlRig: RigVMRegistry should inherit from FGCObject such that the Types array can hold strong references to the type objects. Otherwise, some type object like user defined enums may get GCed if not used during cooking for a while.

diff --git a/Engine/Source/Runtime/RigVM/Private/RigVMCore/RigVMRegistry.cpp b/Engine/Source/Runtime/RigVM/Private/RigVMCore/RigVMRegistry.cpp
--- a/Engine/Source/Runtime/RigVM/Private/RigVMCore/RigVMRegistry.cpp
+++ b/Engine/Source/Runtime/RigVM/Private/RigVMCore/RigVMRegistry.cpp
@@ -11,7 +11,6 @@
 #include "UObject/UObjectIterator.h"
 #include "Misc/CoreDelegates.h"
 
-FRigVMRegistry FRigVMRegistry::s_RigVMRegistry;
 const FName FRigVMRegistry::TemplateNameMetaName = TEXT("TemplateName");
 
 FRigVMRegistry::~FRigVMRegistry()
@@ -21,10 +20,32 @@ FRigVMRegistry::~FRigVMRegistry()
 
 FRigVMRegistry& FRigVMRegistry::Get()
 {
+	// static in a function scope ensures that the GC system is initiated before 
+	// the registry constructor is called
+	static FRigVMRegistry s_RigVMRegistry;
 	s_RigVMRegistry.InitializeIfNeeded();
 	return s_RigVMRegistry;
 }
 
+void FRigVMRegistry::AddReferencedObjects(FReferenceCollector& Collector)
+{
+	// registry should hold strong references to these type objects
+	// otherwise GC may remove them without the registry known it
+	// which can happen during cook time.
+	for (FTypeInfo& Type : Types)
+	{
+		if (Type.Type.CPPTypeObject)
+		{
+			Collector.AddReferencedObject(Type.Type.CPPTypeObject);
+		}
+	}
+}
+
+FString FRigVMRegistry::GetReferencerName() const
+{
+	return TEXT("FRigVMRegistry");
+}
+
 const TArray<UScriptStruct*>& FRigVMRegistry::GetMathTypes()
 {
 	// The list of base math types to automatically register 
@@ -124,16 +145,16 @@ void FRigVMRegistry::InitializeIfNeeded()
 	Refresh();
 
 	// hook the registry to prepare for engine shutdown
-	FCoreDelegates::OnExit.AddLambda([]()
+	FCoreDelegates::OnExit.AddLambda([&]()
 	{
-		s_RigVMRegistry.Reset();
+		Reset();
 
 		if (FAssetRegistryModule* AssetRegistryModule = FModuleManager::GetModulePtr<FAssetRegistryModule>(TEXT("AssetRegistry")))
 		{
 			if (AssetRegistryModule->TryGet())
 			{
-				AssetRegistryModule->Get().OnAssetRemoved().RemoveAll(&s_RigVMRegistry);
-				AssetRegistryModule->Get().OnAssetRenamed().RemoveAll(&s_RigVMRegistry);
+				AssetRegistryModule->Get().OnAssetRemoved().RemoveAll(this);
+				AssetRegistryModule->Get().OnAssetRenamed().RemoveAll(this);
 			}
 		}
 
diff --git a/Engine/Source/Runtime/RigVM/Public/RigVMCore/RigVMRegistry.h b/Engine/Source/Runtime/RigVM/Public/RigVMCore/RigVMRegistry.h
--- a/Engine/Source/Runtime/RigVM/Public/RigVMCore/RigVMRegistry.h
+++ b/Engine/Source/Runtime/RigVM/Public/RigVMCore/RigVMRegistry.h
@@ -22,6 +22,7 @@
 #include "UObject/NameTypes.h"
 #include "UObject/ObjectMacros.h"
 #include "UObject/ReflectedTypeAccessors.h"
+#include "UObject/GCObject.h"
 
 class FProperty;
 class UObject;
@@ -32,8 +33,10 @@ struct FRigVMDispatchFactory;
  * for use in the RigVM. The Register method is called automatically
  * when the static struct is initially constructed for each USTRUCT
  * hosting a RIGVM_METHOD enabled virtual function.
+ * 
+ * Inheriting from FGCObject to ensure that all type objects cannot be GCed
  */
-struct RIGVM_API FRigVMRegistry
+struct RIGVM_API FRigVMRegistry : public FGCObject
 {
 public:
 
@@ -44,6 +47,10 @@ public:
 	// Returns the singleton registry
 	static FRigVMRegistry& Get();
 
+	// FGCObject overrides
+	virtual void AddReferencedObjects(FReferenceCollector& Collector) override;
+	virtual FString GetReferencerName() const override;
+	
 	// Registers a function given its name.
 	// The name will be the name of the struct and virtual method,
 	// for example "FMyStruct::MyVirtualMethod"
@@ -317,7 +324,6 @@ private:
 	// Notifies other system that types have been added/removed, and template permutations have been updated
 	FOnRigVMRegistryChanged OnRigVMRegistryChangedDelegate;
 	
-	static FRigVMRegistry s_RigVMRegistry;
 	friend struct FRigVMStruct;
 	friend struct FRigVMTemplate;
 	friend struct FRigVMTemplateArgument;
