Control Rig: fix compilation error

diff --git a/Engine/Source/Runtime/RigVM/Private/RigVMCore/RigVMRegistry.cpp b/Engine/Source/Runtime/RigVM/Private/RigVMCore/RigVMRegistry.cpp
--- a/Engine/Source/Runtime/RigVM/Private/RigVMCore/RigVMRegistry.cpp
+++ b/Engine/Source/Runtime/RigVM/Private/RigVMCore/RigVMRegistry.cpp
@@ -158,7 +158,7 @@ void FRigVMRegistry::InitializeIfNeeded()
 			}
 		}
 
-		UE::Anim::AttributeTypes::GetOnAttributeTypesChanged().RemoveAll(&s_RigVMRegistry);
+		UE::Anim::AttributeTypes::GetOnAttributeTypesChanged().RemoveAll(this);
 	});
 	
 	FAssetRegistryModule& AssetRegistryModule = FModuleManager::LoadModuleChecked<FAssetRegistryModule>(TEXT("AssetRegistry"));
