Control Rig: Fix CI unit test "Metadata Dispatch" failing in Linux

diff --git a/Engine/Source/Developer/RigVMDeveloper/Private/RigVMModel/RigVMController.cpp b/Engine/Source/Developer/RigVMDeveloper/Private/RigVMModel/RigVMController.cpp
--- a/Engine/Source/Developer/RigVMDeveloper/Private/RigVMModel/RigVMController.cpp
+++ b/Engine/Source/Developer/RigVMDeveloper/Private/RigVMModel/RigVMController.cpp
@@ -16055,10 +16055,10 @@ bool URigVMController::PrepareTemplatePinForType(URigVMPin* InPin, const TArray<
 	// the type required matches the one in the pin
 	if (InPin->IsStructMember())
 	{
-		int32 PinTypeIndex = InPin->GetTypeIndex();
+		TRigVMTypeIndex PinTypeIndex = InPin->GetTypeIndex();
 		
 		// Check if any of the types allows the connection
-		for (const int32& Type : InTypeIndices)
+		for (const TRigVMTypeIndex& Type : InTypeIndices)
 		{
 			if (Registry.CanMatchTypes(Type, PinTypeIndex, true))
 			{
diff --git a/Engine/Source/Runtime/RigVM/Private/RigVMTypeUtils.cpp b/Engine/Source/Runtime/RigVM/Private/RigVMTypeUtils.cpp
new file mode 100644
--- /dev/null
+++ b/Engine/Source/Runtime/RigVM/Private/RigVMTypeUtils.cpp
@@ -0,0 +1,28 @@
+﻿// Copyright Epic Games, Inc. All Rights Reserved.
+
+/*=============================================================================
+	RigVMTypeUtils.h: Module implementation.
+=============================================================================*/
+
+#include "RigVMTypeUtils.h"
+
+namespace RigVMTypeUtils
+{
+	TRigVMTypeIndex TypeIndex::Execute = INDEX_NONE;	
+	TRigVMTypeIndex TypeIndex::Bool = INDEX_NONE;	
+	TRigVMTypeIndex TypeIndex::Float = INDEX_NONE;	
+	TRigVMTypeIndex TypeIndex::Double = INDEX_NONE;	
+	TRigVMTypeIndex TypeIndex::Int32 = INDEX_NONE;	
+	TRigVMTypeIndex TypeIndex::UInt8 = INDEX_NONE;	
+	TRigVMTypeIndex TypeIndex::FName = INDEX_NONE;	
+	TRigVMTypeIndex TypeIndex::FString = INDEX_NONE;
+	TRigVMTypeIndex TypeIndex::WildCard = INDEX_NONE;	
+	TRigVMTypeIndex TypeIndex::BoolArray = INDEX_NONE;	
+	TRigVMTypeIndex TypeIndex::FloatArray = INDEX_NONE;	
+	TRigVMTypeIndex TypeIndex::DoubleArray = INDEX_NONE;	
+	TRigVMTypeIndex TypeIndex::Int32Array = INDEX_NONE;	
+	TRigVMTypeIndex TypeIndex::UInt8Array = INDEX_NONE;	
+	TRigVMTypeIndex TypeIndex::FNameArray = INDEX_NONE;	
+	TRigVMTypeIndex TypeIndex::FStringArray = INDEX_NONE;	
+	TRigVMTypeIndex TypeIndex::WildCardArray = INDEX_NONE;	
+}
\ No newline at end of file
diff --git a/Engine/Source/Runtime/RigVM/Public/RigVMTypeUtils.h b/Engine/Source/Runtime/RigVM/Public/RigVMTypeUtils.h
--- a/Engine/Source/Runtime/RigVM/Public/RigVMTypeUtils.h
+++ b/Engine/Source/Runtime/RigVM/Public/RigVMTypeUtils.h
@@ -50,23 +50,23 @@ namespace RigVMTypeUtils
 	class RIGVM_API TypeIndex
 	{
 	public:
-		static inline TRigVMTypeIndex Execute = INDEX_NONE;	
-		static inline TRigVMTypeIndex Bool = INDEX_NONE;	
-		static inline TRigVMTypeIndex Float = INDEX_NONE;	
-		static inline TRigVMTypeIndex Double = INDEX_NONE;	
-		static inline TRigVMTypeIndex Int32 = INDEX_NONE;	
-		static inline TRigVMTypeIndex UInt8 = INDEX_NONE;	
-		static inline TRigVMTypeIndex FName = INDEX_NONE;	
-		static inline TRigVMTypeIndex FString = INDEX_NONE;
-		static inline TRigVMTypeIndex WildCard = INDEX_NONE;	
-		static inline TRigVMTypeIndex BoolArray = INDEX_NONE;	
-		static inline TRigVMTypeIndex FloatArray = INDEX_NONE;	
-		static inline TRigVMTypeIndex DoubleArray = INDEX_NONE;	
-		static inline TRigVMTypeIndex Int32Array = INDEX_NONE;	
-		static inline TRigVMTypeIndex UInt8Array = INDEX_NONE;	
-		static inline TRigVMTypeIndex FNameArray = INDEX_NONE;	
-		static inline TRigVMTypeIndex FStringArray = INDEX_NONE;	
-		static inline TRigVMTypeIndex WildCardArray = INDEX_NONE;	
+		static TRigVMTypeIndex Execute;	
+		static TRigVMTypeIndex Bool;	
+		static TRigVMTypeIndex Float;	
+		static TRigVMTypeIndex Double;	
+		static TRigVMTypeIndex Int32;	
+		static TRigVMTypeIndex UInt8;	
+		static TRigVMTypeIndex FName;	
+		static TRigVMTypeIndex FString;
+		static TRigVMTypeIndex WildCard;	
+		static TRigVMTypeIndex BoolArray;	
+		static TRigVMTypeIndex FloatArray;	
+		static TRigVMTypeIndex DoubleArray;	
+		static TRigVMTypeIndex Int32Array;	
+		static TRigVMTypeIndex UInt8Array;	
+		static TRigVMTypeIndex FNameArray;	
+		static TRigVMTypeIndex FStringArray;	
+		static TRigVMTypeIndex WildCardArray;	
 	};
 
 	// Returns true if the type specified is an array
