RigVM: always register the dispatch factory on demand when resolving functions

diff --git a/Engine/Source/Runtime/RigVM/Private/RigVMCore/RigVMRegistry.cpp b/Engine/Source/Runtime/RigVM/Private/RigVMCore/RigVMRegistry.cpp
--- a/Engine/Source/Runtime/RigVM/Private/RigVMCore/RigVMRegistry.cpp
+++ b/Engine/Source/Runtime/RigVM/Private/RigVMCore/RigVMRegistry.cpp
@@ -1079,8 +1079,8 @@ const FRigVMFunction* FRigVMRegistry::FindFunction(const TCHAR* InName) const
 	{
 		return &Functions[*FunctionIndexPtr];
 	}
-	
-	if(Factories.Num() > 0)
+
+	// look for the function from a dispatch factory, will register the corresponding factory if needed
 	{
 		static bool IsDispatchingFunction = false;
 		if(IsDispatchingFunction)
