Control Rig: Ensure custom types and dispatch factories only get registered after the object system is ready. Improve error reporting for unresolved functions.

#rnx

diff --git a/Engine/Plugins/Animation/ControlRig/Source/ControlRigEditor/Private/ControlRigEditorModule.cpp b/Engine/Plugins/Animation/ControlRig/Source/ControlRigEditor/Private/ControlRigEditorModule.cpp
--- a/Engine/Plugins/Animation/ControlRig/Source/ControlRigEditor/Private/ControlRigEditorModule.cpp
+++ b/Engine/Plugins/Animation/ControlRig/Source/ControlRigEditor/Private/ControlRigEditorModule.cpp
@@ -1003,7 +1003,7 @@ void FControlRigEditorModule::GetTypeActions(UControlRigBlueprint* CRB, FBluepri
 	}
 
 	FRigVMRegistry& Registry = FRigVMRegistry::Get();
-	Registry.Refresh();
+	Registry.RefreshEngineTypes();
 
 #if UE_RIGVM_ENABLE_TEMPLATE_NODES
 	
diff --git a/Engine/Source/Developer/RigVMDeveloper/Private/RigVMModel/Nodes/RigVMDispatchNode.cpp b/Engine/Source/Developer/RigVMDeveloper/Private/RigVMModel/Nodes/RigVMDispatchNode.cpp
--- a/Engine/Source/Developer/RigVMDeveloper/Private/RigVMModel/Nodes/RigVMDispatchNode.cpp
+++ b/Engine/Source/Developer/RigVMDeveloper/Private/RigVMModel/Nodes/RigVMDispatchNode.cpp
@@ -135,7 +135,7 @@ const FRigVMDispatchFactory* URigVMDispatchNode::GetFactory() const
 	}
 
 	// Try to find factories that have not been registered yet
-	FRigVMRegistry::Get().Refresh();
+	FRigVMRegistry::Get().RefreshEngineTypes();
 	if(const FRigVMTemplate* Template = GetTemplate())
 	{
 		CachedFactory = Template->GetDispatchFactory(); 
diff --git a/Engine/Source/Runtime/RigVM/Private/RigVMCore/RigVM.cpp b/Engine/Source/Runtime/RigVM/Private/RigVMCore/RigVM.cpp
--- a/Engine/Source/Runtime/RigVM/Private/RigVMCore/RigVM.cpp
+++ b/Engine/Source/Runtime/RigVM/Private/RigVMCore/RigVM.cpp
@@ -995,14 +995,16 @@ void URigVM::ResolveFunctionsIfRequired()
 
 		for (int32 FunctionIndex = 0; FunctionIndex < GetFunctionNames().Num(); FunctionIndex++)
 		{
-			GetFunctions()[FunctionIndex] = nullptr;
-			GetFactories()[FunctionIndex] = nullptr;
 			if(const FRigVMFunction* Function = FRigVMRegistry::Get().FindFunction(*GetFunctionNames()[FunctionIndex].ToString()))
 			{
 				GetFunctions()[FunctionIndex] = Function->FunctionPtr;
 				GetFactories()[FunctionIndex] = Function->Factory;
 			}
-			ensureMsgf(GetFunctions()[FunctionIndex], TEXT("Function %s is not valid"), *GetFunctionNames()[FunctionIndex].ToString());
+			else
+			{
+				// We cannot recover from missing functions. 
+				UE_LOG(LogRigVM, Fatal, TEXT("No handler found for function '%s'"), *GetFunctionNames()[FunctionIndex].ToString());
+			}
 		}
 	}
 }
diff --git a/Engine/Source/Runtime/RigVM/Private/RigVMCore/RigVMRegistry.cpp b/Engine/Source/Runtime/RigVM/Private/RigVMCore/RigVMRegistry.cpp
--- a/Engine/Source/Runtime/RigVM/Private/RigVMCore/RigVMRegistry.cpp
+++ b/Engine/Source/Runtime/RigVM/Private/RigVMCore/RigVMRegistry.cpp
@@ -10,9 +10,18 @@
 #include "Engine/UserDefinedStruct.h"
 #include "UObject/UObjectIterator.h"
 #include "Misc/CoreDelegates.h"
+#include "Misc/DelayedAutoRegister.h"
 
 const FName FRigVMRegistry::TemplateNameMetaName = TEXT("TemplateName");
 
+
+// When the object system has been completely loaded, load in all the engine types that we haven't registered already in InitializeIfNeeded 
+FDelayedAutoRegisterHelper GRigVMRegistrySingletonHelper(EDelayedRegisterRunPhase::ObjectSystemReady, []() -> void
+{
+	FRigVMRegistry::Get().RefreshEngineTypes();
+});
+
+
 FRigVMRegistry::~FRigVMRegistry()
 {
 	Reset();
@@ -142,8 +151,6 @@ void FRigVMRegistry::InitializeIfNeeded()
 		FindOrAddType(FRigVMTemplateArgumentType(MathType));
 	}
 
-	Refresh();
-
 	// hook the registry to prepare for engine shutdown
 	FCoreDelegates::OnExit.AddLambda([&]()
 	{
@@ -168,53 +175,47 @@ void FRigVMRegistry::InitializeIfNeeded()
 	UE::Anim::AttributeTypes::GetOnAttributeTypesChanged().AddRaw(this, &FRigVMRegistry::OnAnimationAttributeTypesChanged);
 }
 
-void FRigVMRegistry::Refresh()
+void FRigVMRegistry::RefreshEngineTypes()
 {
-	// refresh only tries to register all loaded types
-	// trying to load all types is very slow for large projects so we don't want to do that
-	
-	// add all user defined structs
-	for (TObjectIterator<UScriptStruct> ScriptIt; ScriptIt; ++ScriptIt)
-	{
-		UScriptStruct* ScriptStruct = *ScriptIt;
-		if(IsAllowedType(ScriptStruct))
+	// Register all user-defined types that the engine knows about. Enumerating over the entire object hierarchy is
+	// slow, so we do it for structs, enums and dispatch factories in one shot.
+	TArray<UScriptStruct*> DispatchFactoriesToRegister;
+	DispatchFactoriesToRegister.Reserve(32);
+
+	static TArray<const UClass *> ClassesToEnumerate{UScriptStruct::StaticClass(), UEnum::StaticClass()};
+	ForEachObjectOfClasses(
+		ClassesToEnumerate,
+		[this, &DispatchFactoriesToRegister](UObject* InObject) -> void
 		{
-			// if this is a C++ type - skip it
-			if(ScriptStruct->IsA<UUserDefinedStruct>() || ScriptStruct->IsChildOf(FRigVMExecuteContext::StaticStruct()))
+			if (UScriptStruct* ScriptStruct = Cast<UScriptStruct>(InObject))
 			{
-				FindOrAddType(FRigVMTemplateArgumentType(ScriptStruct));
+				// if this is a C++ type - skip it
+				if(ScriptStruct->IsA<UUserDefinedStruct>() || ScriptStruct->IsChildOf(FRigVMExecuteContext::StaticStruct()))
+				{
+					FindOrAddType(FRigVMTemplateArgumentType(ScriptStruct));
+				}
+				else if (ScriptStruct != FRigVMDispatchFactory::StaticStruct() &&
+						 ScriptStruct->IsChildOf(FRigVMDispatchFactory::StaticStruct()))
+				{
+					DispatchFactoriesToRegister.Add(ScriptStruct);
+				}
+			}
+			else if (UEnum* Enum = Cast<UEnum>(InObject))
+			{
+				// Ignore reflected C++ enums.
+				if(IsAllowedType(Enum) && !Enum->IsNative())
+				{
+					const FString CPPType = Enum->CppType.IsEmpty() ? Enum->GetName() : Enum->CppType;
+					FindOrAddType(FRigVMTemplateArgumentType(*CPPType, Enum));
+				}
 			}
 		}
-	}
-
-	// add all user defined enums
-	for (TObjectIterator<UEnum> EnumIt; EnumIt; ++EnumIt)
-	{
-		UEnum* Enum = (*EnumIt);
-		if(!IsAllowedType(Enum))
-		{
-			continue;
-		}
-
-		// if this is a C++ type - skip it
-		if(Enum->IsNative())
-		{
-			continue;
-		}
-
-		const FString CPPType = Enum->CppType.IsEmpty() ? Enum->GetName() : Enum->CppType;
-		FindOrAddType(FRigVMTemplateArgumentType(*CPPType, Enum));
-	}
+	);
 	
-	// register all dispatch factories
-	for (TObjectIterator<UScriptStruct> ScriptIt; ScriptIt; ++ScriptIt)
+	// Register all dispatch factories only after all other types have been registered.
+	for (UScriptStruct* DispatchFactoryStruct: DispatchFactoriesToRegister)
 	{
-		UScriptStruct* ScriptStruct = *ScriptIt;
-		if(ScriptStruct != FRigVMDispatchFactory::StaticStruct() &&
-			ScriptStruct->IsChildOf(FRigVMDispatchFactory::StaticStruct()))
-		{
-			RegisterFactory(ScriptStruct);
-		}
+		RegisterFactory(DispatchFactoryStruct);
 	}
 }
 
@@ -1075,37 +1076,38 @@ const FRigVMDispatchFactory* FRigVMRegistry::RegisterFactory(UScriptStruct* InFa
 
 const FRigVMFunction* FRigVMRegistry::FindFunction(const TCHAR* InName) const
 {
+	// Check first if the function is provided by internally registered rig units. 
 	if(const int32* FunctionIndexPtr = FunctionNameToIndex.Find(InName))
 	{
 		return &Functions[*FunctionIndexPtr];
 	}
 
-	// look for the function from a dispatch factory, will register the corresponding factory if needed
+#if !UE_BUILD_SHIPPING
+	static bool IsDispatchingFunction = false;
+	if(IsDispatchingFunction)
 	{
-		static bool IsDispatchingFunction = false;
-		if(IsDispatchingFunction)
-		{
-			return nullptr;
-		}
-		
-		const FString NameString(InName);
-		FString FactoryName, ArgumentsString;
-		if(NameString.Split(TEXT("::"), &FactoryName, &ArgumentsString))
+		return nullptr;
+	}
+#endif
+
+	// Otherwise ask the associated dispatch factory for a function matching this signature.
+	const FString NameString(InName);
+	FString FactoryName, ArgumentsString;
+	if(NameString.Split(TEXT("::"), &FactoryName, &ArgumentsString))
+	{
+		// if the factory has never been registered - FindDispatchFactory will try to look it up and register
+		if(const FRigVMDispatchFactory* Factory = FindDispatchFactory(*FactoryName))
 		{
-			// if the factory has never been registered - FindDispatchFactory will try to look it up and register
-			if(const FRigVMDispatchFactory* Factory = FindDispatchFactory(*FactoryName))
+			if(const FRigVMTemplate* Template = Factory->GetTemplate())
 			{
-				if(const FRigVMTemplate* Template = Factory->GetTemplate())
+				const FRigVMTemplateTypeMap ArgumentTypes = Template->GetArgumentTypesFromString(ArgumentsString);
+				if(ArgumentTypes.Num() == Template->NumArguments())
 				{
-					const FRigVMTemplateTypeMap ArgumentTypes = Template->GetArgumentTypesFromString(ArgumentsString);
-					if(ArgumentTypes.Num() == Template->NumArguments())
+					const int32 PermutationIndex = Template->FindPermutation(ArgumentTypes);
+					if(PermutationIndex != INDEX_NONE)
 					{
-						const int32 PermutationIndex = Template->FindPermutation(ArgumentTypes);
-						if(PermutationIndex != INDEX_NONE)
-						{
-							TGuardValue<bool> ReEntryGuard(IsDispatchingFunction, true);
-							return ((FRigVMTemplate*)Template)->GetOrCreatePermutation(PermutationIndex);
-						}
+						TGuardValue<bool> ReEntryGuard(IsDispatchingFunction, true);
+						return ((FRigVMTemplate*)Template)->GetOrCreatePermutation(PermutationIndex);
 					}
 				}
 			}
diff --git a/Engine/Source/Runtime/RigVM/Public/RigVMCore/RigVMRegistry.h b/Engine/Source/Runtime/RigVM/Public/RigVMCore/RigVMRegistry.h
--- a/Engine/Source/Runtime/RigVM/Public/RigVMCore/RigVMRegistry.h
+++ b/Engine/Source/Runtime/RigVM/Public/RigVMCore/RigVMRegistry.h
@@ -64,7 +64,7 @@ public:
 	
 	// Refreshes the list and finds the function pointers
 	// based on the names.
-	void Refresh();
+	void RefreshEngineTypes();
 
 	// Update the registry when types are renamed
 	void OnAssetRenamed(const FAssetData& InAssetData, const FString& InOldObjectPath);
