Control Rig: Remove erroneous shipping checks.

#rnx

diff --git a/Engine/Source/Runtime/RigVM/Private/RigVMCore/RigVMRegistry.cpp b/Engine/Source/Runtime/RigVM/Private/RigVMCore/RigVMRegistry.cpp
--- a/Engine/Source/Runtime/RigVM/Private/RigVMCore/RigVMRegistry.cpp
+++ b/Engine/Source/Runtime/RigVM/Private/RigVMCore/RigVMRegistry.cpp
@@ -1082,13 +1082,11 @@ const FRigVMFunction* FRigVMRegistry::FindFunction(const TCHAR* InName) const
 		return &Functions[*FunctionIndexPtr];
 	}
 
-#if !UE_BUILD_SHIPPING
 	static bool IsDispatchingFunction = false;
 	if(IsDispatchingFunction)
 	{
 		return nullptr;
 	}
-#endif
 
 	// Otherwise ask the associated dispatch factory for a function matching this signature.
 	const FString NameString(InName);
