Control Rig: Move the initialization of the Control Rig type system to after engine init is done, including plugins. Fixes startup issues with some editor derivatives that don't initialize RigVM directly.

#rnx

diff --git a/Engine/Source/Runtime/RigVM/Private/RigVMCore/RigVMRegistry.cpp b/Engine/Source/Runtime/RigVM/Private/RigVMCore/RigVMRegistry.cpp
--- a/Engine/Source/Runtime/RigVM/Private/RigVMCore/RigVMRegistry.cpp
+++ b/Engine/Source/Runtime/RigVM/Private/RigVMCore/RigVMRegistry.cpp
@@ -16,7 +16,7 @@ const FName FRigVMRegistry::TemplateNameMetaName = TEXT("TemplateName");
 
 
 // When the object system has been completely loaded, load in all the engine types that we haven't registered already in InitializeIfNeeded 
-FDelayedAutoRegisterHelper GRigVMRegistrySingletonHelper(EDelayedRegisterRunPhase::ObjectSystemReady, []() -> void
+static FDelayedAutoRegisterHelper GRigVMRegistrySingletonHelper(EDelayedRegisterRunPhase::EndOfEngineInit, []() -> void
 {
 	FRigVMRegistry::Get().RefreshEngineTypes();
 });
