RigVM: Move NameCaching to RigVM module

Also move base class for struct tests to RigVM module

diff --git a/Engine/Plugins/Animation/ControlRig/Source/ControlRig/Private/ControlRig.cpp b/Engine/Plugins/Animation/ControlRig/Source/ControlRig/Private/ControlRig.cpp
--- a/Engine/Plugins/Animation/ControlRig/Source/ControlRig/Private/ControlRig.cpp
+++ b/Engine/Plugins/Animation/ControlRig/Source/ControlRig/Private/ControlRig.cpp
@@ -845,7 +845,11 @@ bool UControlRig::Execute(const EControlRigState InState, const FName& InEventNa
 	if (bIsInitializingMemory)
 	{
 		AbsoluteTime = DeltaTime = 0.f;
-		NameCache.Reset();
+		
+		if(VM)
+		{
+			VM->GetPublicData<>().GetNameCache()->Reset();
+		}
 	}
 
 	// setup the context with further fields
@@ -881,8 +885,7 @@ bool UControlRig::Execute(const EControlRigState InState, const FName& InEventNa
 	Context.OwningComponent = nullptr;
 	Context.OwningActor = nullptr;
 	Context.World = nullptr;
-	Context.NameCache = &NameCache;
-
+	
 	if (!OuterSceneComponent.IsValid())
 	{
 		USceneComponent* SceneComponentFromRegistry = Context.DataSourceRegistry->RequestSource<USceneComponent>(UControlRig::OwnerComponent);
diff --git a/Engine/Plugins/Animation/ControlRig/Source/ControlRig/Private/Rigs/RigNameCache.cpp b/Engine/Plugins/Animation/ControlRig/Source/ControlRig/Private/Rigs/RigNameCache.cpp
deleted file mode 100644
diff --git a/Engine/Plugins/Animation/ControlRig/Source/ControlRig/Private/Units/Core/RigUnit_Name.cpp b/Engine/Plugins/Animation/ControlRig/Source/ControlRig/Private/Units/Core/RigUnit_Name.cpp
--- a/Engine/Plugins/Animation/ControlRig/Source/ControlRig/Private/Units/Core/RigUnit_Name.cpp
+++ b/Engine/Plugins/Animation/ControlRig/Source/ControlRig/Private/Units/Core/RigUnit_Name.cpp
@@ -7,9 +7,9 @@
 
 FRigUnit_NameConcat_Execute()
 {
-	check(Context.NameCache);
+	check(ExecuteContext.GetNameCache());
 	
-	Result = Context.NameCache->Concat(A, B);
+	Result = ExecuteContext.GetNameCache()->Concat(A, B);
 }
 
 #if WITH_DEV_AUTOMATION_TESTS
@@ -42,7 +42,7 @@ IMPLEMENT_RIGUNIT_AUTOMATION_TEST(FRigUnit_NameConcat)
 
 FRigUnit_NameTruncate_Execute()
 {
-	check(Context.NameCache);
+	check(ExecuteContext.GetNameCache());
 
 	Remainder = Name;
 	Chopped = NAME_None;
@@ -54,13 +54,13 @@ FRigUnit_NameTruncate_Execute()
 
 	if (FromEnd)
 	{
-		Remainder = Context.NameCache->LeftChop(Name, Count);
-		Chopped = Context.NameCache->Right(Name, Count);
+		Remainder = ExecuteContext.GetNameCache()->LeftChop(Name, Count);
+		Chopped = ExecuteContext.GetNameCache()->Right(Name, Count);
 	}
 	else
 	{
-		Remainder = Context.NameCache->RightChop(Name, Count);
-		Chopped = Context.NameCache->Left(Name, Count);
+		Remainder = ExecuteContext.GetNameCache()->RightChop(Name, Count);
+		Chopped = ExecuteContext.GetNameCache()->Left(Name, Count);
 	}
 }
 
@@ -115,9 +115,9 @@ IMPLEMENT_RIGUNIT_AUTOMATION_TEST(FRigUnit_NameTruncate)
 
 FRigUnit_NameReplace_Execute()
 {
-	check(Context.NameCache);
+	check(ExecuteContext.GetNameCache());
 
-	Result = Context.NameCache->Replace(Name, Old, New, ESearchCase::CaseSensitive);
+	Result = ExecuteContext.GetNameCache()->Replace(Name, Old, New, ESearchCase::CaseSensitive);
 }
 
 #if WITH_DEV_AUTOMATION_TESTS
@@ -162,9 +162,9 @@ IMPLEMENT_RIGUNIT_AUTOMATION_TEST(FRigUnit_NameReplace)
 
 FRigUnit_EndsWith_Execute()
 {
-	check(Context.NameCache);
+	check(ExecuteContext.GetNameCache());
 
-	Result = Context.NameCache->EndsWith(Name, Ending, ESearchCase::CaseSensitive);
+	Result = ExecuteContext.GetNameCache()->EndsWith(Name, Ending, ESearchCase::CaseSensitive);
 }
 
 #if WITH_DEV_AUTOMATION_TESTS
@@ -193,9 +193,9 @@ IMPLEMENT_RIGUNIT_AUTOMATION_TEST(FRigUnit_EndsWith)
 
 FRigUnit_StartsWith_Execute()
 {
-	check(Context.NameCache);
+	check(ExecuteContext.GetNameCache());
 
-	Result = Context.NameCache->StartsWith(Name, Start, ESearchCase::CaseSensitive);
+	Result = ExecuteContext.GetNameCache()->StartsWith(Name, Start, ESearchCase::CaseSensitive);
 }
 
 #if WITH_DEV_AUTOMATION_TESTS
@@ -223,9 +223,9 @@ IMPLEMENT_RIGUNIT_AUTOMATION_TEST(FRigUnit_StartsWith)
 
 FRigUnit_Contains_Execute()
 {
-	check(Context.NameCache);
+	check(ExecuteContext.GetNameCache());
 
-	Result = Context.NameCache->Contains(Name, Search, ESearchCase::CaseSensitive);
+	Result = ExecuteContext.GetNameCache()->Contains(Name, Search, ESearchCase::CaseSensitive);
 }
 
 #if WITH_DEV_AUTOMATION_TESTS
diff --git a/Engine/Plugins/Animation/ControlRig/Source/ControlRig/Public/ControlRig.h b/Engine/Plugins/Animation/ControlRig/Source/ControlRig/Public/ControlRig.h
--- a/Engine/Plugins/Animation/ControlRig/Source/ControlRig/Public/ControlRig.h
+++ b/Engine/Plugins/Animation/ControlRig/Source/ControlRig/Public/ControlRig.h
@@ -549,8 +549,6 @@ protected:
 	TMap<FName, FCachedPropertyPath> OutputProperties_DEPRECATED;
 #endif
 
-	FRigNameCache NameCache;
-
 private:
 	
 	void HandleOnControlModified(UControlRig* Subject, FRigControlElement* Control, const FRigControlModifiedContext& Context);
diff --git a/Engine/Plugins/Animation/ControlRig/Source/ControlRig/Public/Rigs/RigNameCache.h b/Engine/Plugins/Animation/ControlRig/Source/ControlRig/Public/Rigs/RigNameCache.h
deleted file mode 100644
diff --git a/Engine/Plugins/Animation/ControlRig/Source/ControlRig/Public/Units/RigUnitContext.h b/Engine/Plugins/Animation/ControlRig/Source/ControlRig/Public/Units/RigUnitContext.h
--- a/Engine/Plugins/Animation/ControlRig/Source/ControlRig/Public/Units/RigUnitContext.h
+++ b/Engine/Plugins/Animation/ControlRig/Source/ControlRig/Public/Units/RigUnitContext.h
@@ -5,7 +5,6 @@
 #include "CoreMinimal.h"
 #include "Rigs/RigHierarchyContainer.h"
 #include "Rigs/RigCurveContainer.h"
-#include "Rigs/RigNameCache.h"
 #include "ControlRigLog.h"
 #include "AnimationDataSource.h"
 #include "Animation/AttributesRuntime.h"
@@ -75,7 +74,6 @@ struct FRigUnitContext
 #if WITH_EDITOR
 		, Log(nullptr)
 #endif
-		, NameCache(nullptr)
 	{
 	}
 
@@ -132,9 +130,6 @@ struct FRigUnitContext
 	FControlRigLog* Log;
 #endif
 
-	/** A container to store all names */
-	FRigNameCache* NameCache;
-
 	/**
 	 * Returns a given data source and cast it to the expected class.
 	 *
diff --git a/Engine/Plugins/Animation/ControlRig/Source/ControlRig/Public/Units/RigUnitTest.h b/Engine/Plugins/Animation/ControlRig/Source/ControlRig/Public/Units/RigUnitTest.h
--- a/Engine/Plugins/Animation/ControlRig/Source/ControlRig/Public/Units/RigUnitTest.h
+++ b/Engine/Plugins/Animation/ControlRig/Source/ControlRig/Public/Units/RigUnitTest.h
@@ -7,12 +7,13 @@
 #include "RigUnitContext.h"
 #include "Rigs/RigHierarchy.h"
 #include "Rigs/RigHierarchyController.h"
+#include "RigVMCore/RigVMStructTest.h"
 
-class FControlRigUnitTestBase : public FAutomationTestBase
+class FControlRigUnitTestBase : public FRigVMStructTestBase<FControlRigExecuteContext>
 {
 public:
 	FControlRigUnitTestBase(const FString& InName, bool bIsComplex)
-		: FAutomationTestBase(InName, bIsComplex)
+		: FRigVMStructTestBase<FControlRigExecuteContext>(InName, bIsComplex)
 		, Hierarchy(nullptr)
 		, Controller(nullptr)
 	{
@@ -27,8 +28,10 @@ public:
 		}
 	}
 
-	void InitHierarchy()
+	virtual void Initialize() override
 	{
+		FRigVMStructTestBase<FControlRigExecuteContext>::Initialize();
+		
 		if (!Hierarchy)
 		{
 			Hierarchy = NewObject<URigHierarchy>();
@@ -36,17 +39,13 @@ public:
 
 			Hierarchy->AddToRoot();
 			// we no longer add the controller to root since controller is now part of the hierarchy
-
-			Context.Hierarchy = Hierarchy;;
-			ExecuteContext.Hierarchy = Hierarchy;
-			Context.NameCache = &NameCache;
 		}
+		Context.Hierarchy = Hierarchy;;
+		ExecuteContext.Hierarchy = Hierarchy;
 	}
 
 	URigHierarchy* Hierarchy;
 	URigHierarchyController* Controller;
-	FControlRigExecuteContext ExecuteContext;
-	FRigNameCache NameCache;
 	FRigUnitContext Context;
 };
 
@@ -72,7 +71,7 @@ public:
 		TUnitStruct Unit; \
 		virtual bool RunTest(const FString& Parameters) override \
 		{ \
-			FControlRigUnitTestBase::InitHierarchy(); \
+			Initialize(); \
 			Hierarchy->Reset(); \
 			Unit = TUnitStruct(); \
 			return RunControlRigUnitTest(Parameters); \
diff --git a/Engine/Plugins/Animation/ControlRig/Source/ControlRigDeveloper/Private/ControlRigBlueprint.cpp b/Engine/Plugins/Animation/ControlRig/Source/ControlRigDeveloper/Private/ControlRigBlueprint.cpp
--- a/Engine/Plugins/Animation/ControlRig/Source/ControlRigDeveloper/Private/ControlRigBlueprint.cpp
+++ b/Engine/Plugins/Animation/ControlRig/Source/ControlRigDeveloper/Private/ControlRigBlueprint.cpp
@@ -1226,11 +1226,9 @@ void UControlRigBlueprint::RecompileVM()
 		}
 		CDO->VM->Reset();
 
-		FRigNameCache TempNameCache;
 		FRigUnitContext InitContext;
 		InitContext.State = EControlRigState::Init;
 		InitContext.Hierarchy = CDO->DynamicHierarchy;
-		InitContext.NameCache = &TempNameCache;
 
 		FRigUnitContext UpdateContext = InitContext;
 		UpdateContext.State = EControlRigState::Update;
diff --git a/Engine/Plugins/Animation/ControlRig/Source/ControlRigEditor/Private/ControlRigEditorModule.cpp b/Engine/Plugins/Animation/ControlRig/Source/ControlRigEditor/Private/ControlRigEditorModule.cpp
--- a/Engine/Plugins/Animation/ControlRig/Source/ControlRigEditor/Private/ControlRigEditorModule.cpp
+++ b/Engine/Plugins/Animation/ControlRig/Source/ControlRigEditor/Private/ControlRigEditorModule.cpp
@@ -1857,11 +1857,9 @@ void FControlRigEditorModule::GetContextMenuActions(const UControlRigGraphSchema
 								StructOnScope = UnitNode->ConstructStructInstance(false /* default */);
 								StructMemory = (FRigUnit*)StructOnScope->GetStructMemory();
 
-								FRigNameCache NameCache;
 								FRigUnitContext RigUnitContext;
 								RigUnitContext.Hierarchy = TemporaryHierarchy;
 								RigUnitContext.State = EControlRigState::Update;
-								RigUnitContext.NameCache = &NameCache;
 								
 								StructMemory->Execute(RigUnitContext);
 							}
diff --git a/Engine/Plugins/Experimental/DataInterfaceGraph/Source/DataInterfaceGraphUncookedOnly/Private/DataInterfaceGraphUncookedOnlyUtils.cpp b/Engine/Plugins/Experimental/DataInterfaceGraph/Source/DataInterfaceGraphUncookedOnly/Private/DataInterfaceGraphUncookedOnlyUtils.cpp
--- a/Engine/Plugins/Experimental/DataInterfaceGraph/Source/DataInterfaceGraphUncookedOnly/Private/DataInterfaceGraphUncookedOnlyUtils.cpp
+++ b/Engine/Plugins/Experimental/DataInterfaceGraph/Source/DataInterfaceGraphUncookedOnly/Private/DataInterfaceGraphUncookedOnlyUtils.cpp
@@ -36,10 +36,8 @@ void FUtils::Compile(UDataInterfaceGraph* InGraph)
 
 	InGraph->VMRuntimeSettings = EditorData->VMRuntimeSettings;
 
-	FRigNameCache TempNameCache;
 	FRigUnitContext InitContext;
 	InitContext.State = EControlRigState::Init;
-	InitContext.NameCache = &TempNameCache;
 
 	FRigUnitContext UpdateContext = InitContext;
 	UpdateContext.State = EControlRigState::Update;
diff --git a/Engine/Source/Programs/Shared/EpicGames.UHT/Exporters/CodeGen/UhtHeaderCodeGeneratorCppFile.cs b/Engine/Source/Programs/Shared/EpicGames.UHT/Exporters/CodeGen/UhtHeaderCodeGeneratorCppFile.cs
--- a/Engine/Source/Programs/Shared/EpicGames.UHT/Exporters/CodeGen/UhtHeaderCodeGeneratorCppFile.cs
+++ b/Engine/Source/Programs/Shared/EpicGames.UHT/Exporters/CodeGen/UhtHeaderCodeGeneratorCppFile.cs
@@ -725,6 +725,8 @@ namespace EpicGames.UHT.Exporters.CodeGen
 						builder.Append("\t").Append(scriptStruct.RigVMStructInfo.ExecuteContextType).Append("& TemporaryExecuteContext = ").Append(scriptStruct.RigVMStructInfo.ExecuteContextMember).Append(";\r\n");
 					}
 
+					builder.Append("\tTemporaryExecuteContext.Initialize();\r\n");
+
 					bool wroteLine = false;
 					foreach (UhtRigVMParameter parameter in scriptStruct.RigVMStructInfo.Members)
 					{
diff --git a/Engine/Source/Programs/UnrealHeaderTool/Private/CodeGenerator.cpp b/Engine/Source/Programs/UnrealHeaderTool/Private/CodeGenerator.cpp
--- a/Engine/Source/Programs/UnrealHeaderTool/Private/CodeGenerator.cpp
+++ b/Engine/Source/Programs/UnrealHeaderTool/Private/CodeGenerator.cpp
@@ -3787,6 +3787,7 @@ void FNativeClassHeaderGenerator::ExportGeneratedStructBodyMacros(FOutputDevice&
 			{
 				Out.Logf(TEXT("\t%s& TemporaryExecuteContext = %s;\r\n"), *StructRigVMInfo.ExecuteContextType, *StructRigVMInfo.ExecuteContextMember);
 			}
+			Out.Log(TEXT("\tTemporaryExecuteContext.Initialize();\r\n"));
 
 			if(RigVMVirtualFuncProlog.Num() > 0)
 			{
diff --git a/Engine/Source/Runtime/RigVM/Private/RigVMCore/RigVM.cpp b/Engine/Source/Runtime/RigVM/Private/RigVMCore/RigVM.cpp
--- a/Engine/Source/Runtime/RigVM/Private/RigVMCore/RigVM.cpp
+++ b/Engine/Source/Runtime/RigVM/Private/RigVMCore/RigVM.cpp
@@ -995,9 +995,7 @@ void URigVM::SetContextPublicDataStruct(UScriptStruct* InScriptStruct)
 	{
 		return;
 	}
-	
-	check(InScriptStruct);
-	Context.PublicDataScope = FStructOnScope(InScriptStruct); 
+	Context.Initialize(InScriptStruct);
 }
 
 void URigVM::ResolveFunctionsIfRequired()
diff --git a/Engine/Source/Runtime/RigVM/Private/RigVMCore/RigVMNameCache.cpp b/Engine/Source/Runtime/RigVM/Private/RigVMCore/RigVMNameCache.cpp
new file mode 100644
--- /dev/null
+++ b/Engine/Source/Runtime/RigVM/Private/RigVMCore/RigVMNameCache.cpp
@@ -0,0 +1,357 @@
+// Copyright Epic Games, Inc. All Rights Reserved.
+
+#include "RigVMCore/RigVMNameCache.h"
+#include "HAL/IConsoleManager.h"
+
+#if WITH_EDITOR
+#include "RigVMCore/RigVM.h"
+TAutoConsoleVariable<int32> CVarRigVMNameCacheMaxSize(TEXT("RigVM.NameCacheMaxSize"), 512, TEXT("Change to control how many names are cached per VM instance."));
+#endif
+
+////////////////////////////////////////////////////////////////////////////////
+// FRigVMNameOp
+////////////////////////////////////////////////////////////////////////////////
+
+FRigVMNameOp FRigVMNameOp::Concat(const FName& InA, const FName& InB)
+{
+	FRigVMNameOp Op;
+	Op.Type = ERigVMNameOp::Concat;
+	Op.A = GetTypeHash(InA);
+	Op.B = GetTypeHash(InB);
+	return Op;
+}
+
+FRigVMNameOp FRigVMNameOp::Left(const FName& InA, const uint32 InCount)
+{
+	FRigVMNameOp Op;
+	Op.Type = ERigVMNameOp::Left;
+	Op.A = GetTypeHash(InA);
+	Op.B = InCount;
+	return Op;
+}
+
+FRigVMNameOp FRigVMNameOp::Right(const FName& InA, const uint32 InCount)
+{
+	FRigVMNameOp Op;
+	Op.Type = ERigVMNameOp::Right;
+	Op.A = GetTypeHash(InA);
+	Op.B = InCount;
+	return Op;
+}
+
+FRigVMNameOp FRigVMNameOp::LeftChop(const FName& InA, const uint32 InCount)
+{
+	FRigVMNameOp Op;
+	Op.Type = ERigVMNameOp::LeftChop;
+	Op.A = GetTypeHash(InA);
+	Op.B = InCount;
+	return Op;
+}
+
+FRigVMNameOp FRigVMNameOp::RightChop(const FName& InA, const uint32 InCount)
+{
+	FRigVMNameOp Op;
+	Op.Type = ERigVMNameOp::RightChop;
+	Op.A = GetTypeHash(InA);
+	Op.B = InCount;
+	return Op;
+}
+
+FRigVMNameOp FRigVMNameOp::Replace(const FName& InA, const FName& InB, const FName& InC, const ESearchCase::Type InSearchCase)
+{
+	FRigVMNameOp Op;
+	Op.Type = InSearchCase == ESearchCase::CaseSensitive ? ERigVMNameOp::ReplaceCase : ERigVMNameOp::ReplaceNoCase;
+	Op.A = GetTypeHash(InA);
+	Op.B = GetTypeHash(InB);
+	Op.C = GetTypeHash(InC);
+	return Op;
+}
+
+FRigVMNameOp FRigVMNameOp::EndsWith(const FName& InA, const FName& InB, const ESearchCase::Type InSearchCase)
+{
+	FRigVMNameOp Op;
+	Op.Type = InSearchCase == ESearchCase::CaseSensitive ? ERigVMNameOp::EndsWithCase : ERigVMNameOp::EndsWithNoCase;
+	Op.A = GetTypeHash(InA);
+	Op.B = GetTypeHash(InB);
+	return Op;
+}
+
+FRigVMNameOp FRigVMNameOp::StartsWith(const FName& InA, const FName& InB, const ESearchCase::Type InSearchCase)
+{
+	FRigVMNameOp Op;
+	Op.Type = InSearchCase == ESearchCase::CaseSensitive ? ERigVMNameOp::StartsWithCase : ERigVMNameOp::StartsWithNoCase;
+	Op.A = GetTypeHash(InA);
+	Op.B = GetTypeHash(InB);
+	return Op;
+}
+
+FRigVMNameOp FRigVMNameOp::Contains(const FName& InA, const FName& InB, const ESearchCase::Type InSearchCase)
+{
+	FRigVMNameOp Op;
+	Op.Type = InSearchCase == ESearchCase::CaseSensitive ? ERigVMNameOp::ContainsCase : ERigVMNameOp::ContainsNoCase;
+	Op.A = GetTypeHash(InA);
+	Op.B = GetTypeHash(InB);
+	return Op;
+}
+
+////////////////////////////////////////////////////////////////////////////////
+// FRigVMNameCache
+////////////////////////////////////////////////////////////////////////////////
+
+FName FRigVMNameCache::Concat(const FName& InA, const FName& InB)
+{
+	if(InA.IsNone())
+	{
+		return InB;
+	}
+
+	if(InB.IsNone())
+	{
+		return InA;
+	}
+	
+	const FRigVMNameOp Op = FRigVMNameOp::Concat(InA, InB);
+	if(const FName* ExistingName = NameCache.Find(Op))
+	{
+		return *ExistingName;
+	}
+	const FName NewName = *(InA.ToString() + InB.ToString());
+
+	if(CheckCacheSize())
+	{
+		NameCache.Add(Op, NewName);
+	}
+	return NewName;
+}
+
+FName FRigVMNameCache::Left(const FName& InA, const uint32 InCount)
+{
+	if(InCount == 0 || InA.IsNone())
+	{
+		return NAME_None;
+	}
+
+	if(InCount >= InA.GetStringLength())
+	{
+		return InA;
+	}
+	
+	const FRigVMNameOp Op = FRigVMNameOp::Left(InA, InCount);
+	if(const FName* ExistingName = NameCache.Find(Op))
+	{
+		return *ExistingName;
+	}
+	const FName NewName = *InA.ToString().Left(InCount);
+
+	if(CheckCacheSize())
+	{
+		NameCache.Add(Op, NewName);
+	}
+	
+	return NewName;
+}
+
+FName FRigVMNameCache::Right(const FName& InA, const uint32 InCount)
+{
+	if(InCount == 0 || InA.IsNone())
+	{
+		return NAME_None;
+	}
+
+	if(InCount >= InA.GetStringLength())
+	{
+		return InA;
+	}
+
+	const FRigVMNameOp Op = FRigVMNameOp::Right(InA, InCount);
+	if(const FName* ExistingName = NameCache.Find(Op))
+	{
+		return *ExistingName;
+	}
+	const FName NewName = *InA.ToString().Right(InCount);
+
+	if(CheckCacheSize())
+	{
+		NameCache.Add(Op, NewName);
+	}
+	
+	return NewName;
+}
+
+FName FRigVMNameCache::LeftChop(const FName& InA, const uint32 InCount)
+{
+	if(InCount == 0 || InA.IsNone())
+	{
+		return InA;
+	}
+	
+	const FRigVMNameOp Op = FRigVMNameOp::LeftChop(InA, InCount);
+	if(const FName* ExistingName = NameCache.Find(Op))
+	{
+		return *ExistingName;
+	}
+	const FName NewName = *InA.ToString().LeftChop(InCount);
+
+	if(CheckCacheSize())
+	{
+		NameCache.Add(Op, NewName);
+	}
+	
+	return NewName;
+}
+
+FName FRigVMNameCache::RightChop(const FName& InA, const uint32 InCount)
+{
+	if(InCount == 0 || InA.IsNone())
+	{
+		return InA;
+	}
+
+	const FRigVMNameOp Op = FRigVMNameOp::RightChop(InA, InCount);
+	if(const FName* ExistingName = NameCache.Find(Op))
+	{
+		return *ExistingName;
+	}
+	const FName NewName = *InA.ToString().RightChop(InCount);
+
+	if(CheckCacheSize())
+	{
+		NameCache.Add(Op, NewName);
+	}
+	
+	return NewName;
+}
+
+FName FRigVMNameCache::Replace(const FName& InA, const FName& InB, const FName& InC, const ESearchCase::Type InSearchCase)
+{
+	if(InA.IsNone() || InB.IsNone())
+	{
+		return InA;
+	}
+
+	const FRigVMNameOp Op = FRigVMNameOp::Replace(InA, InB, InC, InSearchCase);
+	if(const FName* ExistingName = NameCache.Find(Op))
+	{
+		return *ExistingName;
+	}
+
+	FString NewString(TEXT(""));
+	if (!InC.IsNone())
+	{
+		NewString = *InC.ToString();
+	}
+	const FName NewName = *InA.ToString().Replace(*InB.ToString(), *NewString, InSearchCase);
+
+	if(CheckCacheSize())
+	{
+		NameCache.Add(Op, NewName);
+	}
+	
+	return NewName;
+}
+
+bool FRigVMNameCache::EndsWith(const FName& InA, const FName& InB, const ESearchCase::Type InSearchCase)
+{
+	if(InA.IsNone() || InB.IsNone())
+	{
+		return false;
+	}
+	
+	const FRigVMNameOp Op = FRigVMNameOp::EndsWith(InA, InB, InSearchCase);
+	if(const bool* ExistingBool = BoolCache.Find(Op))
+	{
+		return *ExistingBool;
+	}
+	const bool bResult = InA.ToString().EndsWith(InB.ToString(), InSearchCase);
+
+	if(CheckCacheSize())
+	{
+		BoolCache.Add(Op, bResult);
+	}
+	
+	return bResult;
+}
+
+bool FRigVMNameCache::StartsWith(const FName& InA, const FName& InB, const ESearchCase::Type InSearchCase)
+{
+	if(InA.IsNone() || InB.IsNone())
+	{
+		return false;
+	}
+
+	const FRigVMNameOp Op = FRigVMNameOp::StartsWith(InA, InB, InSearchCase);
+	if(const bool* ExistingBool = BoolCache.Find(Op))
+	{
+		return *ExistingBool;
+	}
+	const bool bResult = InA.ToString().StartsWith(InB.ToString(), InSearchCase);
+
+	if(CheckCacheSize())
+	{
+		BoolCache.Add(Op, bResult);
+	}
+	
+	return bResult;
+}
+
+bool FRigVMNameCache::Contains(const FName& InA, const FName& InB, const ESearchCase::Type InSearchCase)
+{
+	if(InA.IsNone() || InB.IsNone())
+	{
+		return false;
+	}
+
+	const FRigVMNameOp Op = FRigVMNameOp::Contains(InA, InB, InSearchCase);
+	if(const bool* ExistingBool = BoolCache.Find(Op))
+	{
+		return *ExistingBool;
+	}
+	const bool bResult = InA.ToString().Contains(InB.ToString(), InSearchCase);
+
+	if(CheckCacheSize())
+	{
+		BoolCache.Add(Op, bResult);
+	}
+	
+	return bResult;
+}
+
+TArray<FRigVMNameOp> FRigVMNameCache::GetNameOps() const
+{
+	TArray<FRigVMNameOp> Result;
+	NameCache.GenerateKeyArray(Result);
+	return Result;
+}
+
+TArray<FName> FRigVMNameCache::GetNameValues() const
+{
+	TArray<FName> Result;
+	NameCache.GenerateValueArray(Result);
+	return Result;
+}
+
+TArray<FRigVMNameOp> FRigVMNameCache::GetBoolOps() const
+{
+	TArray<FRigVMNameOp> Result;
+	BoolCache.GenerateKeyArray(Result);
+	return Result;
+}
+
+TArray<bool> FRigVMNameCache::GetBoolValues() const
+{
+	TArray<bool> Result;
+	BoolCache.GenerateValueArray(Result);
+	return Result;
+}
+
+#if WITH_EDITOR
+bool FRigVMNameCache::CheckCacheSize() const
+{
+	if(CVarRigVMNameCacheMaxSize.GetValueOnAnyThread() == NameCache.Num() || CVarRigVMNameCacheMaxSize.GetValueOnAnyThread() == BoolCache.Num())
+	{
+		UE_LOG(LogRigVM, Warning, TEXT("FRigVMNameCache exceeded maximum size of %d. You can change it using the 'RigVM.NameCacheMaxSize' console variable."), CVarRigVMNameCacheMaxSize.GetValueOnAnyThread());
+		return false;
+	}
+	return true;
+}
+#endif
diff --git a/Engine/Source/Runtime/RigVM/Public/RigVMCore/RigVM.h b/Engine/Source/Runtime/RigVM/Public/RigVMCore/RigVM.h
--- a/Engine/Source/Runtime/RigVM/Public/RigVMCore/RigVM.h
+++ b/Engine/Source/Runtime/RigVM/Public/RigVMCore/RigVM.h
@@ -668,6 +668,7 @@ public:
 
 	uint32 GetNumExecutions() const { return NumExecutions; }
 	const FRigVMExtendedExecuteContext& GetContext() const { return Context; }
+	FRigVMExtendedExecuteContext& GetContext() { return Context; }
 
 	template<typename ExecuteContextType = FRigVMExecuteContext>
 	FORCEINLINE const ExecuteContextType& GetPublicData() const
diff --git a/Engine/Source/Runtime/RigVM/Public/RigVMCore/RigVMExecuteContext.h b/Engine/Source/Runtime/RigVM/Public/RigVMCore/RigVMExecuteContext.h
--- a/Engine/Source/Runtime/RigVM/Public/RigVMCore/RigVMExecuteContext.h
+++ b/Engine/Source/Runtime/RigVM/Public/RigVMCore/RigVMExecuteContext.h
@@ -20,7 +20,7 @@
 #include "UObject/ObjectMacros.h"
 #include "UObject/UnrealNames.h"
 #include "UObject/UnrealType.h"
-//#include "RigVMCore/RigVMNameCache.h"
+#include "RigVMCore/RigVMNameCache.h"
 
 #include "RigVMExecuteContext.generated.h"
 
@@ -29,9 +29,6 @@ struct FRigVMDispatchFactory;
 struct FRigVMExecuteContext;
 struct FRigVMExtendedExecuteContext;
 
-// temporary definition until we implement the namecache
-typedef int32 FRigVMNameCache;
-
 USTRUCT()
 struct RIGVM_API FRigVMSlice
 {
@@ -270,6 +267,15 @@ struct RIGVM_API FRigVMExecuteContext
 
 	FORCEINLINE FRigVMNameCache* GetNameCache() const { return NameCache; }
 
+	virtual void Initialize()
+	{
+		if(NameCache == nullptr)
+		{
+			static FRigVMNameCache StaticNameCache;
+			NameCache = &StaticNameCache;
+		}
+	}
+
 protected:
 
 	FORCEINLINE virtual void Reset()
@@ -317,6 +323,7 @@ struct RIGVM_API FRigVMExtendedExecuteContext
 	, Factory(nullptr)
 	{
 		Reset();
+		SetDefaultNameCache();
 	}
 
 	FORCEINLINE FRigVMExtendedExecuteContext(const UScriptStruct* InExecuteContextStruct)
@@ -325,10 +332,16 @@ struct RIGVM_API FRigVMExtendedExecuteContext
 		, LastExecutionMicroSeconds()
 		, Factory(nullptr)
 	{
-		check(InExecuteContextStruct->IsChildOf(FRigVMExecuteContext::StaticStruct()));
-		Reset();
+		if(InExecuteContextStruct)
+		{
+			check(InExecuteContextStruct->IsChildOf(FRigVMExecuteContext::StaticStruct()));
+			Reset();
+			SetDefaultNameCache();
+		}
 	}
 
+	virtual ~FRigVMExtendedExecuteContext() {}
+
 	FORCEINLINE void Reset()
 	{
 		((FRigVMExecuteContext*)PublicDataScope.GetStructMemory())->Reset();
@@ -345,7 +358,8 @@ struct RIGVM_API FRigVMExtendedExecuteContext
 		check(PublicDataStruct);
 		PublicDataScope = FStructOnScope(PublicDataStruct);
 		PublicDataStruct->CopyScriptStruct(PublicDataScope.GetStructMemory(), Other.PublicDataScope.GetStructMemory());
-		
+		SetDefaultNameCache();
+
 		VM = Other.VM;
 		OpaqueArguments = Other.OpaqueArguments;
 		Slices = Other.Slices;
@@ -353,6 +367,13 @@ struct RIGVM_API FRigVMExtendedExecuteContext
 		return *this;
 	}
 
+	virtual void Initialize(UScriptStruct* InScriptStruct)
+	{
+		check(InScriptStruct->IsChildOf(FRigVMExecuteContext::StaticStruct()));
+		PublicDataScope = FStructOnScope(InScriptStruct);
+		SetDefaultNameCache();
+	}
+
 	template<typename ExecuteContextType = FRigVMExecuteContext>
 	FORCEINLINE const ExecuteContextType& GetPublicData() const
 	{
@@ -443,6 +464,11 @@ struct RIGVM_API FRigVMExtendedExecuteContext
 		check(GetPublicData<>().RuntimeSettings.MaximumArraySize > 0);
 	}
 
+	FORCEINLINE void SetDefaultNameCache()
+	{
+		SetNameCache(&NameCache);
+	}
+
 	FORCEINLINE void SetNameCache(FRigVMNameCache* InNameCache)
 	{
 		GetPublicData<>().NameCache = InNameCache;
@@ -455,4 +481,5 @@ struct RIGVM_API FRigVMExtendedExecuteContext
 	TArray<uint16> SliceOffsets;
 	double LastExecutionMicroSeconds;
 	const FRigVMDispatchFactory* Factory;
+	FRigVMNameCache NameCache;
 };
diff --git a/Engine/Source/Runtime/RigVM/Public/RigVMCore/RigVMNameCache.h b/Engine/Source/Runtime/RigVM/Public/RigVMCore/RigVMNameCache.h
new file mode 100644
--- /dev/null
+++ b/Engine/Source/Runtime/RigVM/Public/RigVMCore/RigVMNameCache.h
@@ -0,0 +1,146 @@
+// Copyright Epic Games, Inc. All Rights Reserved.
+
+#pragma once
+
+#include "CoreMinimal.h"
+
+enum class ERigVMNameOp : uint8
+{
+	None,
+	Concat,
+	Left,
+    Right,
+    LeftChop,
+	RightChop,
+	ReplaceCase,
+	ReplaceNoCase,
+	EndsWithCase,
+	EndsWithNoCase,
+    StartsWithCase,
+    StartsWithNoCase,
+    ContainsCase,
+    ContainsNoCase
+};
+
+struct RIGVM_API FRigVMNameOp
+{
+public:
+
+	FORCEINLINE FRigVMNameOp()
+	: A(INDEX_NONE)
+	, B(INDEX_NONE)
+	, C(INDEX_NONE)
+	, Type(ERigVMNameOp::None)
+	{}
+
+	FORCEINLINE static uint32 GetTypeHash(const FName& InName)
+	{
+		return HashCombine(InName.GetComparisonIndex().ToUnstableInt(), uint32(InName.GetNumber()));
+	}
+	
+	static FRigVMNameOp Concat(const FName& InA, const FName& InB);
+	static FRigVMNameOp Left(const FName& InA, const uint32 InCount);
+	static FRigVMNameOp Right(const FName& InA, const uint32 InCount);
+	static FRigVMNameOp LeftChop(const FName& InA, const uint32 InCount);
+	static FRigVMNameOp RightChop(const FName& InA, const uint32 InCount);
+	static FRigVMNameOp Replace(const FName& InA, const FName& InB, const FName& InC, const ESearchCase::Type InSearchCase);
+    static FRigVMNameOp EndsWith(const FName& InA, const FName& InB, const ESearchCase::Type InSearchCase);
+    static FRigVMNameOp StartsWith(const FName& InA, const FName& InB, const ESearchCase::Type InSearchCase);
+    static FRigVMNameOp Contains(const FName& InA, const FName& InB, const ESearchCase::Type InSearchCase);
+
+	friend FORCEINLINE uint32 GetTypeHash(const FRigVMNameOp& Op)
+	{
+		return HashCombine(Op.A, HashCombine(Op.B, HashCombine(Op.C, (uint32)Op.Type)));
+	}
+
+	FORCEINLINE bool operator ==(const FRigVMNameOp& Other) const
+	{
+		return Type == Other.Type && A == Other.A && B == Other.B && C == Other.C;
+	}
+
+	FORCEINLINE bool operator !=(const FRigVMNameOp& Other) const
+	{
+		return Type != Other.Type || A != Other.A || B != Other.B || C != Other.C;
+	}
+
+	FORCEINLINE bool operator <(const FRigVMNameOp& Other) const
+	{
+		if (Type < Other.Type)
+		{
+			return true;
+		}
+		if(A < Other.A)
+		{
+			return true;
+		}
+		if(B < Other.B)
+		{
+			return true;
+		}
+		return C < Other.C;
+	}
+
+	FORCEINLINE bool operator >(const FRigVMNameOp& Other) const
+	{
+		if (Type > Other.Type)
+		{
+			return true;
+		}
+		if(A > Other.A)
+		{
+			return true;
+		}
+		if(B > Other.B)
+		{
+			return true;
+		}
+		return C > Other.C;
+	}
+
+
+private:
+
+	uint32 A;
+	uint32 B;
+	uint32 C;
+	ERigVMNameOp Type;
+};
+
+class RIGVM_API FRigVMNameCache
+{
+public:
+
+	FORCEINLINE void Reset()
+	{
+		NameCache.Reset();
+		BoolCache.Reset();
+	}
+
+	FName Concat(const FName& InA, const FName& InB);
+	FName Left(const FName& InA, const uint32 InCount);
+	FName Right(const FName& InA, const uint32 InCount);
+	FName LeftChop(const FName& InA, const uint32 InCount);
+	FName RightChop(const FName& InA, const uint32 InCount);
+	FName Replace(const FName& InA, const FName& InB, const FName& InC, const ESearchCase::Type InSearchCase);
+	bool EndsWith(const FName& InA, const FName& InB, const ESearchCase::Type InSearchCase);
+	bool StartsWith(const FName& InA, const FName& InB, const ESearchCase::Type InSearchCase);
+	bool Contains(const FName& InA, const FName& InB, const ESearchCase::Type InSearchCase);
+
+	FORCEINLINE const TMap<FRigVMNameOp, FName>& GetNameCache() const { return NameCache; }
+	FORCEINLINE const TMap<FRigVMNameOp, bool>&  GetBoolCache() const { return BoolCache; }
+	TArray<FRigVMNameOp> GetNameOps() const;
+	TArray<FName> GetNameValues() const;
+	TArray<FRigVMNameOp> GetBoolOps() const;
+	TArray<bool> GetBoolValues() const;
+
+private:
+
+#if WITH_EDITOR
+	bool CheckCacheSize() const;
+#else
+	FORCEINLINE bool CheckCacheSize() const { return true; } 
+#endif
+
+	TMap<FRigVMNameOp, FName> NameCache;
+	TMap<FRigVMNameOp, bool> BoolCache;
+};
diff --git a/Engine/Source/Runtime/RigVM/Public/RigVMCore/RigVMStructTest.h b/Engine/Source/Runtime/RigVM/Public/RigVMCore/RigVMStructTest.h
new file mode 100644
--- /dev/null
+++ b/Engine/Source/Runtime/RigVM/Public/RigVMCore/RigVMStructTest.h
@@ -0,0 +1,83 @@
+// Copyright Epic Games, Inc. All Rights Reserved.
+
+#pragma once
+
+#include "Misc/AutomationTest.h"
+#include "RigVMCore/RigVMStruct.h"
+
+template<typename ExecuteContextType = FRigVMExecuteContext>
+class FRigVMStructTestBase : public FAutomationTestBase
+{
+public:
+	FRigVMStructTestBase(const FString& InName, bool bIsComplex)
+		: FAutomationTestBase(InName, bIsComplex)
+		, ExtendedExecuteContext(nullptr)
+		, ExecuteContext(DefaultExecuteContext)
+	{
+	}
+
+	~FRigVMStructTestBase()
+	{
+	}
+
+	virtual void Initialize()
+	{
+		ExtendedExecuteContext = FRigVMExtendedExecuteContext(ExecuteContextType::StaticStruct());
+		ExecuteContext = ExtendedExecuteContext.GetPublicData<ExecuteContextType>(); 
+	}
+	
+protected:
+	FRigVMExtendedExecuteContext ExtendedExecuteContext;
+
+public:
+	ExecuteContextType& ExecuteContext;
+
+private:
+	ExecuteContextType DefaultExecuteContext;
+};
+
+#define RIGVMSTRUCT_TEST_STRINGIFY(Content) #Content
+#define IMPLEMENT_RIGVMSTRUCT_AUTOMATION_TEST(TRigVMStruct) \
+	class TRigVMStruct##Test : public FRigVMStructTestBase \
+	{ \
+	public: \
+		TRigVMStruct##Test( const FString& InName ) \
+		:FRigVMStructTestBase( InName, false ) {\
+		} \
+		virtual uint32 GetTestFlags() const override { return EAutomationTestFlags::EditorContext | EAutomationTestFlags::EngineFilter; } \
+		virtual bool IsStressTest() const { return false; } \
+		virtual uint32 GetRequiredDeviceNum() const override { return 1; } \
+		virtual FString GetTestSourceFileName() const override { return __FILE__; } \
+		virtual int32 GetTestSourceFileLine() const override { return __LINE__; } \
+	protected: \
+		virtual void GetTests(TArray<FString>& OutBeautifiedNames, TArray <FString>& OutTestCommands) const override \
+		{ \
+			OutBeautifiedNames.Add(TEXT(RIGVMSTRUCT_TEST_STRINGIFY(RigVM.RigVMFunctions.TRigVMStruct))); \
+			OutTestCommands.Add(FString()); \
+		} \
+		TRigVMStruct Unit; \
+		virtual bool RunTest(const FString& Parameters) override \
+		{ \
+			Initialize(); \
+			Unit = TRigVMStruct(); \
+			return RunRigVMStructTest(Parameters); \
+		} \
+		virtual bool RunRigVMStructTest(const FString& Parameters); \
+		virtual FString GetBeautifiedTestName() const override { return TEXT(RIGVMSTRUCT_TEST_STRINGIFY(RigVM.RigVMFunctions.TRigVMStruct)); } \
+		void Init() \
+		{ \
+			Context.State = EControlRigState::Init; \
+			Unit.Execute(Context); \
+		} \
+		void Execute() \
+		{ \
+			Context.State = EControlRigState::Update; \
+			Unit.Execute(Context); \
+		} \
+		void InitAndExecute() { Init(); Execute(); } \
+	}; \
+	namespace\
+	{\
+		TRigVMStruct##Test TRigVMStruct##AutomationTestInstance(TEXT(RIGVMSTRUCT_TEST_STRINGIFY(TRigVMStruct##Test))); \
+	} \
+	bool TRigVMStruct##Test::RunRigVMStructTest(const FString& Parameters)
\ No newline at end of file
