RigVM: Move name related units to RigVM module

diff --git a/Engine/Plugins/Animation/ControlRig/Config/BaseControlRig.ini b/Engine/Plugins/Animation/ControlRig/Config/BaseControlRig.ini
--- a/Engine/Plugins/Animation/ControlRig/Config/BaseControlRig.ini
+++ b/Engine/Plugins/Animation/ControlRig/Config/BaseControlRig.ini
@@ -48,3 +48,12 @@
 +PropertyRedirects=(OldName="/Script/ControlRig.ControlRigComponent.OnPreSetupDelegate",NewName="OnPreConstructionDelegate")
 +PropertyRedirects=(OldName="/Script/ControlRig.ControlRigComponent.OnPostSetupDelegate",NewName="OnPostConstructionDelegate")
 +PropertyRedirects=(OldName="/Script/ControlRig.UControlRigSettings.SetupEventBorderColor",NewName="ConstructionEventBorderColor")
+
+; moving core units to RigVM runtime
++StructRedirects=(OldName="/Script/ControlRig.RigUnit_NameBase",NewName="/Script/RigVM.RigVMFunction_NameBase")
++StructRedirects=(OldName="/Script/ControlRig.RigUnit_NameConcat",NewName="/Script/RigVM.RigVMFunction_NameConcat")
++StructRedirects=(OldName="/Script/ControlRig.RigUnit_NameTruncate",NewName="/Script/RigVM.RigVMFunction_NameTruncate")
++StructRedirects=(OldName="/Script/ControlRig.RigUnit_NameReplace",NewName="/Script/RigVM.RigVMFunction_NameReplace")
++StructRedirects=(OldName="/Script/ControlRig.RigUnit_EndsWith",NewName="/Script/RigVM.RigVMFunction_EndsWith")
++StructRedirects=(OldName="/Script/ControlRig.RigUnit_StartsWith",NewName="/Script/RigVM.RigVMFunction_StartsWith")
++StructRedirects=(OldName="/Script/ControlRig.RigUnit_Contains",NewName="/Script/RigVM.RigVMFunction_Contains")
diff --git a/Engine/Plugins/Animation/ControlRig/Source/ControlRig/Private/Units/Core/RigUnit_Name.cpp b/Engine/Plugins/Animation/ControlRig/Source/ControlRig/Private/Units/Core/RigUnit_Name.cpp
deleted file mode 100644
diff --git a/Engine/Plugins/Animation/ControlRig/Source/ControlRig/Private/Units/Core/RigUnit_Name.h b/Engine/Plugins/Animation/ControlRig/Source/ControlRig/Private/Units/Core/RigUnit_Name.h
deleted file mode 100644
diff --git a/Engine/Plugins/Animation/ControlRig/Source/ControlRig/Private/Units/Execution/RigUnit_Item.cpp b/Engine/Plugins/Animation/ControlRig/Source/ControlRig/Private/Units/Execution/RigUnit_Item.cpp
--- a/Engine/Plugins/Animation/ControlRig/Source/ControlRig/Private/Units/Execution/RigUnit_Item.cpp
+++ b/Engine/Plugins/Animation/ControlRig/Source/ControlRig/Private/Units/Execution/RigUnit_Item.cpp
@@ -1,7 +1,7 @@
 // Copyright Epic Games, Inc. All Rights Reserved.
 
 #include "RigUnit_Item.h"
-#include "Units/Core/RigUnit_Name.h"
+#include "RigVMFunctions/RigVMFunction_Name.h"
 #include "Units/RigUnitContext.h"
 #include "Units/Core/RigUnit_CoreDispatch.h"
 
@@ -37,7 +37,7 @@ FRigUnit_ItemReplace_Execute()
     DECLARE_SCOPE_HIERARCHICAL_COUNTER_RIGUNIT()
 
 	Result = Item;
-	FRigUnit_NameReplace::StaticExecute(ExecuteContext, Item.Name, Old, New, Result.Name, Context);
+	FRigVMFunction_NameReplace::StaticExecute(ExecuteContext, Item.Name, Old, New, Result.Name);
 }
 
 FRigUnit_ItemEquals_Execute()
diff --git a/Engine/Source/Runtime/RigVM/Private/RigVMFunctions/RigVMFunction_Name.cpp b/Engine/Source/Runtime/RigVM/Private/RigVMFunctions/RigVMFunction_Name.cpp
new file mode 100644
--- /dev/null
+++ b/Engine/Source/Runtime/RigVM/Private/RigVMFunctions/RigVMFunction_Name.cpp
@@ -0,0 +1,254 @@
+// Copyright Epic Games, Inc. All Rights Reserved.
+
+#include "RigVMFunctions/RigVMFunction_Name.h"
+#if WITH_DEV_AUTOMATION_TESTS
+#include "RigVMCore/RigVMStructTest.h"
+#endif
+
+#include UE_INLINE_GENERATED_CPP_BY_NAME(RigVMFunction_Name)
+
+FRigVMFunction_NameConcat_Execute()
+{
+	check(ExecuteContext.GetNameCache());
+	
+	Result = ExecuteContext.GetNameCache()->Concat(A, B);
+}
+
+#if WITH_DEV_AUTOMATION_TESTS
+
+IMPLEMENT_RIGVMSTRUCT_AUTOMATION_TEST(FRigVMFunction_NameConcat)
+{
+	const FName Control(TEXT("Control"));
+	const FName Rig(TEXT("Rig"));
+	const FName ControlRig(TEXT("ControlRig"));
+	
+	Unit.A = Control;
+	Unit.B = Rig;
+	Execute();
+	AddErrorIfFalse(Unit.Result == ControlRig, TEXT("unexpected name concat result"));
+
+	Unit.A = Control;
+	Unit.B = NAME_None;
+	Execute();
+	AddErrorIfFalse(Unit.Result == Control, TEXT("unexpected name concat result"));
+
+	Unit.A = NAME_None;
+	Unit.B = Rig;
+	Execute();
+	AddErrorIfFalse(Unit.Result == Rig, TEXT("unexpected name concat result"));
+	return true;
+}
+
+#endif
+
+FRigVMFunction_NameTruncate_Execute()
+{
+	check(ExecuteContext.GetNameCache());
+
+	Remainder = Name;
+	Chopped = NAME_None;
+
+	if(Name.IsNone() || Count <= 0)
+	{
+		return;
+	}
+
+	if (FromEnd)
+	{
+		Remainder = ExecuteContext.GetNameCache()->LeftChop(Name, Count);
+		Chopped = ExecuteContext.GetNameCache()->Right(Name, Count);
+	}
+	else
+	{
+		Remainder = ExecuteContext.GetNameCache()->RightChop(Name, Count);
+		Chopped = ExecuteContext.GetNameCache()->Left(Name, Count);
+	}
+}
+
+#if WITH_DEV_AUTOMATION_TESTS
+
+IMPLEMENT_RIGVMSTRUCT_AUTOMATION_TEST(FRigVMFunction_NameTruncate)
+{
+	const FName Control(TEXT("Control"));
+	const FName Rig(TEXT("Rig"));
+	const FName ControlRig(TEXT("ControlRig"));
+	const FName Con(TEXT("Con"));
+	const FName trolRig(TEXT("trolRig"));
+	
+	Unit.Name = ControlRig;
+	Unit.Count = 3;
+	Unit.FromEnd = true;
+	Execute();
+	AddErrorIfFalse(Unit.Chopped == Rig, *FString::Printf(TEXT("unexpected name truncate result 1 '%s'"), *Unit.Chopped.ToString()));
+	AddErrorIfFalse(Unit.Remainder == Control, *FString::Printf(TEXT("unexpected name truncate result 2 '%s'"), *Unit.Remainder.ToString()));
+
+	Unit.Name = ControlRig;
+	Unit.Count = 7;
+	Unit.FromEnd = false;
+	Execute();
+	AddErrorIfFalse(Unit.Chopped == Control, *FString::Printf(TEXT("unexpected name truncate result 3 '%s'"), *Unit.Chopped.ToString()));
+	AddErrorIfFalse(Unit.Remainder == Rig, *FString::Printf(TEXT("unexpected name truncate result 4 '%s'"), *Unit.Remainder.ToString()));
+
+	Unit.Name = ControlRig;
+	Unit.Count = 3;
+	Unit.FromEnd = false;
+	Execute();
+	AddErrorIfFalse(Unit.Chopped == Con, *FString::Printf(TEXT("unexpected name truncate result 5 '%s'"), *Unit.Chopped.ToString()));
+	AddErrorIfFalse(Unit.Remainder == trolRig, *FString::Printf(TEXT("unexpected name truncate result 6 '%s'"), *Unit.Remainder.ToString()));
+
+	Unit.Name = ControlRig;
+	Unit.Count = 0;
+	Unit.FromEnd = false;
+	Execute();
+	AddErrorIfFalse(Unit.Chopped == NAME_None, *FString::Printf(TEXT("unexpected name truncate result 7 '%s'"), *Unit.Chopped.ToString()));
+	AddErrorIfFalse(Unit.Remainder == ControlRig, *FString::Printf(TEXT("unexpected name truncate result 8 '%s'"), *Unit.Remainder.ToString()));
+
+	Unit.Name = NAME_None;
+	Unit.Count = 6;
+	Unit.FromEnd = false;
+	Execute();
+	AddErrorIfFalse(Unit.Chopped == NAME_None, *FString::Printf(TEXT("unexpected name truncate result 9 '%s'"), *Unit.Chopped.ToString()));
+	AddErrorIfFalse(Unit.Remainder == NAME_None, *FString::Printf(TEXT("unexpected name truncate result 10 '%s'"), *Unit.Remainder.ToString()));
+	return true;
+}
+#endif
+
+FRigVMFunction_NameReplace_Execute()
+{
+	check(ExecuteContext.GetNameCache());
+
+	Result = ExecuteContext.GetNameCache()->Replace(Name, Old, New, ESearchCase::CaseSensitive);
+}
+
+#if WITH_DEV_AUTOMATION_TESTS
+
+IMPLEMENT_RIGVMSTRUCT_AUTOMATION_TEST(FRigVMFunction_NameReplace)
+{
+	const FName Control(TEXT("Control"));
+	const FName Rig(TEXT("Rig"));
+	const FName ControlRig(TEXT("ControlRig"));
+	const FName Foo(TEXT("Foo"));
+	const FName FooRig(TEXT("FooRig"));
+	const FName ControlFoo(TEXT("ControlFoo"));
+	
+	Unit.Name = ControlRig;
+	Unit.Old = Control;
+	Unit.New = Foo;
+	Execute();
+	AddErrorIfFalse(Unit.Result == FooRig, TEXT("unexpected name replace result"));
+
+	Unit.Name = ControlRig;
+	Unit.Old = Rig;
+	Unit.New = Foo;
+	Execute();
+	AddErrorIfFalse(Unit.Result == ControlFoo, TEXT("unexpected name replace result"));
+
+	Unit.Name = ControlRig;
+	Unit.Old = NAME_None;
+	Unit.New = Foo;
+	Execute();
+	AddErrorIfFalse(Unit.Result == ControlRig, TEXT("unexpected name replace result"));
+
+	Unit.Name = ControlRig;
+	Unit.Old = Control;
+	Unit.New = NAME_None;
+	Execute();
+	AddErrorIfFalse(Unit.Result == Rig, TEXT("unexpected name replace result"));
+
+	return true;
+}
+#endif
+
+FRigVMFunction_EndsWith_Execute()
+{
+	check(ExecuteContext.GetNameCache());
+
+	Result = ExecuteContext.GetNameCache()->EndsWith(Name, Ending, ESearchCase::CaseSensitive);
+}
+
+#if WITH_DEV_AUTOMATION_TESTS
+
+IMPLEMENT_RIGVMSTRUCT_AUTOMATION_TEST(FRigVMFunction_EndsWith)
+{
+	const FName Control(TEXT("Control"));
+	const FName Rig(TEXT("Rig"));
+	const FName ControlRig(TEXT("ControlRig"));
+	
+	Unit.Name = ControlRig;
+	Unit.Ending = Control;
+	Execute();
+	AddErrorIfFalse(Unit.Result == false, TEXT("unexpected name endswith result"));
+
+	Unit.Name = ControlRig;
+	Unit.Ending = Rig;
+	Execute();
+	AddErrorIfFalse(Unit.Result == true, TEXT("unexpected name endswith result"));
+
+	return true;
+}
+
+#endif
+
+FRigVMFunction_StartsWith_Execute()
+{
+	check(ExecuteContext.GetNameCache());
+
+	Result = ExecuteContext.GetNameCache()->StartsWith(Name, Start, ESearchCase::CaseSensitive);
+}
+
+#if WITH_DEV_AUTOMATION_TESTS
+
+IMPLEMENT_RIGVMSTRUCT_AUTOMATION_TEST(FRigVMFunction_StartsWith)
+{
+	const FName Control(TEXT("Control"));
+	const FName Rig(TEXT("Rig"));
+	const FName ControlRig(TEXT("ControlRig"));
+	
+	Unit.Name = ControlRig;
+	Unit.Start = Control;
+	Execute();
+	AddErrorIfFalse(Unit.Result == true, TEXT("unexpected name startswith result"));
+
+	Unit.Name = ControlRig;
+	Unit.Start = Rig;
+	Execute();
+	AddErrorIfFalse(Unit.Result == false, TEXT("unexpected name startswith result"));
+
+	return true;
+}
+#endif
+
+FRigVMFunction_Contains_Execute()
+{
+	check(ExecuteContext.GetNameCache());
+
+	Result = ExecuteContext.GetNameCache()->Contains(Name, Search, ESearchCase::CaseSensitive);
+}
+
+#if WITH_DEV_AUTOMATION_TESTS
+
+IMPLEMENT_RIGVMSTRUCT_AUTOMATION_TEST(FRigVMFunction_Contains)
+{
+	const FName Control(TEXT("Control"));
+	const FName Rig(TEXT("Rig"));
+	const FName Foo(TEXT("Foo"));
+	const FName ControlRig(TEXT("ControlRig"));
+	
+	Unit.Name = ControlRig;
+	Unit.Search = Control;
+	Execute();
+	AddErrorIfFalse(Unit.Result == true, TEXT("unexpected name contains result"));
+
+	Unit.Name = ControlRig;
+	Unit.Search = Rig;
+	Execute();
+	AddErrorIfFalse(Unit.Result == true, TEXT("unexpected name contains result"));
+
+	Unit.Name = ControlRig;
+	Unit.Search = Foo;
+	Execute();
+	AddErrorIfFalse(Unit.Result == false, TEXT("unexpected name contains result"));
+
+	return true;
+}
+#endif
diff --git a/Engine/Source/Runtime/RigVM/Public/RigVMCore/RigVMStructTest.h b/Engine/Source/Runtime/RigVM/Public/RigVMCore/RigVMStructTest.h
--- a/Engine/Source/Runtime/RigVM/Public/RigVMCore/RigVMStructTest.h
+++ b/Engine/Source/Runtime/RigVM/Public/RigVMCore/RigVMStructTest.h
@@ -38,11 +38,11 @@ private:
 
 #define RIGVMSTRUCT_TEST_STRINGIFY(Content) #Content
 #define IMPLEMENT_RIGVMSTRUCT_AUTOMATION_TEST(TRigVMStruct) \
-	class TRigVMStruct##Test : public FRigVMStructTestBase \
+	class TRigVMStruct##Test : public FRigVMStructTestBase<FRigVMExecuteContext> \
 	{ \
 	public: \
 		TRigVMStruct##Test( const FString& InName ) \
-		:FRigVMStructTestBase( InName, false ) {\
+		:FRigVMStructTestBase<FRigVMExecuteContext>( InName, false ) {\
 		} \
 		virtual uint32 GetTestFlags() const override { return EAutomationTestFlags::EditorContext | EAutomationTestFlags::EngineFilter; } \
 		virtual bool IsStressTest() const { return false; } \
@@ -64,17 +64,10 @@ private:
 		} \
 		virtual bool RunRigVMStructTest(const FString& Parameters); \
 		virtual FString GetBeautifiedTestName() const override { return TEXT(RIGVMSTRUCT_TEST_STRINGIFY(RigVM.RigVMFunctions.TRigVMStruct)); } \
-		void Init() \
-		{ \
-			Context.State = EControlRigState::Init; \
-			Unit.Execute(Context); \
-		} \
 		void Execute() \
 		{ \
-			Context.State = EControlRigState::Update; \
-			Unit.Execute(Context); \
+			Unit.Execute(); \
 		} \
-		void InitAndExecute() { Init(); Execute(); } \
 	}; \
 	namespace\
 	{\
diff --git a/Engine/Source/Runtime/RigVM/Public/RigVMFunctions/RigVMFunction_Name.h b/Engine/Source/Runtime/RigVM/Public/RigVMFunctions/RigVMFunction_Name.h
new file mode 100644
--- /dev/null
+++ b/Engine/Source/Runtime/RigVM/Public/RigVMFunctions/RigVMFunction_Name.h
@@ -0,0 +1,188 @@
+// Copyright Epic Games, Inc. All Rights Reserved.
+
+#pragma once
+
+#include "RigVMCore/RigVMStruct.h"
+#include "RigVMFunction_Name.generated.h"
+
+USTRUCT(meta=(Abstract, Category="Core|Name", NodeColor = "0.462745, 1,0, 0.329412"))
+struct RIGVM_API FRigVMFunction_NameBase : public FRigVMStruct
+{
+	GENERATED_BODY()
+
+	virtual void Execute() {}
+};
+
+/**
+ * Concatenates two strings together to make a new string
+ */
+USTRUCT(meta = (DisplayName = "Concat", TemplateName = "Concat", Keywords = "Add,+,Combine,Merge,Append"))
+struct RIGVM_API FRigVMFunction_NameConcat : public FRigVMFunction_NameBase
+{
+	GENERATED_BODY()
+
+	FRigVMFunction_NameConcat()
+	{
+		A = B = Result = NAME_None;
+	}
+
+	RIGVM_METHOD()
+	virtual void Execute() override;
+
+	UPROPERTY(meta=(Input, Aggregate))
+	FName A;
+
+	UPROPERTY(meta=(Input, Aggregate))
+	FName B;
+
+	UPROPERTY(meta=(Output, Aggregate))
+	FName Result;
+};
+
+/**
+ * Returns the left or right most characters from the string chopping the given number of characters from the start or the end
+ */
+USTRUCT(meta = (DisplayName = "Chop", TemplateName = "Chop", Keywords = "Truncate,-,Remove,Subtract,Split"))
+struct RIGVM_API FRigVMFunction_NameTruncate : public FRigVMFunction_NameBase
+{
+	GENERATED_BODY()
+
+	FRigVMFunction_NameTruncate()
+	{
+		Name = NAME_None;
+		Count = 1;
+		FromEnd = true;
+	}
+
+	RIGVM_METHOD()
+	virtual void Execute() override;
+
+	UPROPERTY(meta=(Input))
+	FName Name;
+
+	// Number of characters to remove from left or right
+	UPROPERTY(meta=(Input))
+	int32 Count;
+
+	// if set to true the characters will be removed from the end
+	UPROPERTY(meta=(Input))
+	bool FromEnd;
+
+	// the part of the string without the chopped characters
+	UPROPERTY(meta=(Output))
+	FName Remainder;
+
+	// the part of the name that has been chopped off
+	UPROPERTY(meta = (Output))
+	FName Chopped;
+};
+
+/**
+ * Replace all occurrences of a substring in this string
+ */
+USTRUCT(meta = (DisplayName = "Replace", TemplateName = "Replace", Keywords = "Search,Emplace,Find"))
+struct RIGVM_API FRigVMFunction_NameReplace : public FRigVMFunction_NameBase
+{
+	GENERATED_BODY()
+
+	FRigVMFunction_NameReplace()
+	{
+		Name = Old = New = NAME_None;
+	}
+
+	RIGVM_METHOD()
+	virtual void Execute() override;
+
+	UPROPERTY(meta = (Input))
+	FName Name;
+
+	UPROPERTY(meta = (Input))
+	FName Old;
+
+	UPROPERTY(meta = (Input))
+	FName New;
+
+	UPROPERTY(meta = (Output))
+	FName Result;
+};
+
+/**
+ * Tests whether this string ends with given string
+ */
+USTRUCT(meta = (DisplayName = "Ends With", TemplateName = "EndsWith", Keywords = "Right"))
+struct RIGVM_API FRigVMFunction_EndsWith : public FRigVMFunction_NameBase
+{
+	GENERATED_BODY()
+
+	FRigVMFunction_EndsWith()
+	{
+		Name = Ending = NAME_None;
+		Result = false;
+	}
+
+	RIGVM_METHOD()
+		virtual void Execute() override;
+
+	UPROPERTY(meta = (Input))
+	FName Name;
+
+	UPROPERTY(meta = (Input))
+	FName Ending;
+
+	UPROPERTY(meta = (Output))
+	bool Result;
+};
+
+/**
+ * Tests whether this string starts with given string
+ */
+USTRUCT(meta = (DisplayName = "Starts With", TemplateName = "StartsWith", Keywords = "Left"))
+struct RIGVM_API FRigVMFunction_StartsWith : public FRigVMFunction_NameBase
+{
+	GENERATED_BODY()
+
+	FRigVMFunction_StartsWith()
+	{
+		Name = Start = NAME_None;
+		Result = false;
+	}
+
+	RIGVM_METHOD()
+	virtual void Execute() override;
+
+	UPROPERTY(meta = (Input))
+	FName Name;
+
+	UPROPERTY(meta = (Input))
+	FName Start;
+
+	UPROPERTY(meta = (Output))
+	bool Result;
+};
+
+/**
+ * Returns true or false if a given name exists in another given name
+ */
+USTRUCT(meta = (DisplayName = "Contains", TemplateName = "Contains", Keywords = "Contains,Find,Has,Search"))
+struct RIGVM_API FRigVMFunction_Contains : public FRigVMFunction_NameBase
+{
+	GENERATED_BODY()
+
+		FRigVMFunction_Contains()
+	{
+		Name = Search = NAME_None;
+		Result = false;
+	}
+
+	RIGVM_METHOD()
+	virtual void Execute() override;
+
+	UPROPERTY(meta = (Input))
+	FName Name;
+
+	UPROPERTY(meta = (Input))
+	FName Search;
+
+	UPROPERTY(meta = (Output))
+	bool Result;
+};
\ No newline at end of file
