RigVM: Move string and core dispatches to RigVM module

diff --git a/Engine/Plugins/Animation/ControlRig/Config/BaseControlRig.ini b/Engine/Plugins/Animation/ControlRig/Config/BaseControlRig.ini
--- a/Engine/Plugins/Animation/ControlRig/Config/BaseControlRig.ini
+++ b/Engine/Plugins/Animation/ControlRig/Config/BaseControlRig.ini
@@ -50,6 +50,8 @@
 +PropertyRedirects=(OldName="/Script/ControlRig.UControlRigSettings.SetupEventBorderColor",NewName="ConstructionEventBorderColor")
 
 ; moving core units to RigVM runtime
+;+StructRedirects=(OldName="/Script/ControlRig.RigUnit_",NewName="/Script/RigVM.RigVMFunction_", MatchSubstring=true)
+;+StructRedirects=(OldName="/Script/ControlRig.RigDispatch_",NewName="/Script/RigVM.RigVMDispatch_", MatchSubstring=true)
 +StructRedirects=(OldName="/Script/ControlRig.RigUnit_NameBase",NewName="/Script/RigVM.RigVMFunction_NameBase")
 +StructRedirects=(OldName="/Script/ControlRig.RigUnit_NameConcat",NewName="/Script/RigVM.RigVMFunction_NameConcat")
 +StructRedirects=(OldName="/Script/ControlRig.RigUnit_NameTruncate",NewName="/Script/RigVM.RigVMFunction_NameTruncate")
@@ -57,3 +59,26 @@
 +StructRedirects=(OldName="/Script/ControlRig.RigUnit_EndsWith",NewName="/Script/RigVM.RigVMFunction_EndsWith")
 +StructRedirects=(OldName="/Script/ControlRig.RigUnit_StartsWith",NewName="/Script/RigVM.RigVMFunction_StartsWith")
 +StructRedirects=(OldName="/Script/ControlRig.RigUnit_Contains",NewName="/Script/RigVM.RigVMFunction_Contains")
++StructRedirects=(OldName="/Script/ControlRig.RigDispatch_CoreBase",NewName="/Script/RigVM.RigVMDispatch_CoreBase")
++StructRedirects=(OldName="/Script/ControlRig.RigDispatch_CoreEquals",NewName="/Script/RigVM.RigVMDispatch_CoreEquals")
++StructRedirects=(OldName="/Script/ControlRig.RigDispatch_CoreNotEquals",NewName="/Script/RigVM.RigVMDispatch_CoreNotEquals")
++StructRedirects=(OldName="/Script/ControlRig.RigDispatch_Print",NewName="/Script/RigVM.RigVMDispatch_Print")
++StructRedirects=(OldName="/Script/ControlRig.RigUnit_StringBase",NewName="/Script/RigVM.RigVMFunction_StringBase")
++StructRedirects=(OldName="/Script/ControlRig.RigUnit_StringConcat",NewName="/Script/RigVM.RigVMFunction_StringConcat")
++StructRedirects=(OldName="/Script/ControlRig.RigUnit_StringTruncate",NewName="/Script/RigVM.RigVMFunction_StringTruncate")
++StructRedirects=(OldName="/Script/ControlRig.RigUnit_StringReplace",NewName="/Script/RigVM.RigVMFunction_StringReplace")
++StructRedirects=(OldName="/Script/ControlRig.RigUnit_StringEndsWith",NewName="/Script/RigVM.RigVMFunction_StringEndsWith")
++StructRedirects=(OldName="/Script/ControlRig.RigUnit_StringStartsWith",NewName="/Script/RigVM.RigVMFunction_StringStartsWith")
++StructRedirects=(OldName="/Script/ControlRig.RigUnit_StringContains",NewName="/Script/RigVM.RigVMFunction_StringContains")
++StructRedirects=(OldName="/Script/ControlRig.RigUnit_StringLength",NewName="/Script/RigVM.RigVMFunction_StringLength")
++StructRedirects=(OldName="/Script/ControlRig.RigUnit_StringTrimWhitespace",NewName="/Script/RigVM.RigVMFunction_StringTrimWhitespace")
++StructRedirects=(OldName="/Script/ControlRig.RigUnit_StringToUppercase",NewName="/Script/RigVM.RigVMFunction_StringToUppercase")
++StructRedirects=(OldName="/Script/ControlRig.RigUnit_StringToLowercase",NewName="/Script/RigVM.RigVMFunction_StringToLowercase")
++StructRedirects=(OldName="/Script/ControlRig.RigUnit_StringReverse",NewName="/Script/RigVM.RigVMFunction_StringReverse")
++StructRedirects=(OldName="/Script/ControlRig.RigUnit_StringLeft",NewName="/Script/RigVM.RigVMFunction_StringLeft")
++StructRedirects=(OldName="/Script/ControlRig.RigUnit_StringRight",NewName="/Script/RigVM.RigVMFunction_StringRight")
++StructRedirects=(OldName="/Script/ControlRig.RigUnit_StringMiddle",NewName="/Script/RigVM.RigVMFunction_StringMiddle")
++StructRedirects=(OldName="/Script/ControlRig.RigUnit_StringFind",NewName="/Script/RigVM.RigVMFunction_StringFind")
++StructRedirects=(OldName="/Script/ControlRig.RigUnit_StringSplit",NewName="/Script/RigVM.RigVMFunction_StringSplit")
++StructRedirects=(OldName="/Script/ControlRig.RigUnit_StringJoin",NewName="/Script/RigVM.RigVMFunction_StringJoin")
++StructRedirects=(OldName="/Script/ControlRig.RigUnit_StringPadInteger",NewName="/Script/RigVM.RigVMFunction_StringPadInteger")
diff --git a/Engine/Plugins/Animation/ControlRig/Source/ControlRig/Private/Units/Core/RigUnit_CoreDispatch.cpp b/Engine/Plugins/Animation/ControlRig/Source/ControlRig/Private/Units/Core/RigUnit_CoreDispatch.cpp
deleted file mode 100644
diff --git a/Engine/Plugins/Animation/ControlRig/Source/ControlRig/Private/Units/Core/RigUnit_CoreDispatch.h b/Engine/Plugins/Animation/ControlRig/Source/ControlRig/Private/Units/Core/RigUnit_CoreDispatch.h
deleted file mode 100644
diff --git a/Engine/Plugins/Animation/ControlRig/Source/ControlRig/Private/Units/Core/RigUnit_Print.cpp b/Engine/Plugins/Animation/ControlRig/Source/ControlRig/Private/Units/Core/RigUnit_Print.cpp
deleted file mode 100644
diff --git a/Engine/Plugins/Animation/ControlRig/Source/ControlRig/Private/Units/Core/RigUnit_Print.h b/Engine/Plugins/Animation/ControlRig/Source/ControlRig/Private/Units/Core/RigUnit_Print.h
deleted file mode 100644
diff --git a/Engine/Plugins/Animation/ControlRig/Source/ControlRig/Private/Units/Core/RigUnit_String.cpp b/Engine/Plugins/Animation/ControlRig/Source/ControlRig/Private/Units/Core/RigUnit_String.cpp
deleted file mode 100644
diff --git a/Engine/Plugins/Animation/ControlRig/Source/ControlRig/Private/Units/Core/RigUnit_String.h b/Engine/Plugins/Animation/ControlRig/Source/ControlRig/Private/Units/Core/RigUnit_String.h
deleted file mode 100644
diff --git a/Engine/Plugins/Animation/ControlRig/Source/ControlRig/Private/Units/Execution/RigUnit_Item.cpp b/Engine/Plugins/Animation/ControlRig/Source/ControlRig/Private/Units/Execution/RigUnit_Item.cpp
--- a/Engine/Plugins/Animation/ControlRig/Source/ControlRig/Private/Units/Execution/RigUnit_Item.cpp
+++ b/Engine/Plugins/Animation/ControlRig/Source/ControlRig/Private/Units/Execution/RigUnit_Item.cpp
@@ -3,7 +3,7 @@
 #include "RigUnit_Item.h"
 #include "RigVMFunctions/RigVMFunction_Name.h"
 #include "Units/RigUnitContext.h"
-#include "Units/Core/RigUnit_CoreDispatch.h"
+#include "RigVMFunctions/RigVMDispatch_Core.h"
 
 #include UE_INLINE_GENERATED_CPP_BY_NAME(RigUnit_Item)
 
@@ -47,7 +47,7 @@ FRigUnit_ItemEquals_Execute()
 
 FRigVMStructUpgradeInfo FRigUnit_ItemEquals::GetUpgradeInfo() const
 {
-	return FRigVMStructUpgradeInfo::MakeFromStructToFactory(StaticStruct(), FRigDispatch_CoreEquals::StaticStruct());
+	return FRigVMStructUpgradeInfo::MakeFromStructToFactory(StaticStruct(), FRigVMDispatch_CoreEquals::StaticStruct());
 }
 
 FRigUnit_ItemNotEquals_Execute()
@@ -57,7 +57,7 @@ FRigUnit_ItemNotEquals_Execute()
 
 FRigVMStructUpgradeInfo FRigUnit_ItemNotEquals::GetUpgradeInfo() const
 {
-	return FRigVMStructUpgradeInfo::MakeFromStructToFactory(StaticStruct(), FRigDispatch_CoreNotEquals::StaticStruct());
+	return FRigVMStructUpgradeInfo::MakeFromStructToFactory(StaticStruct(), FRigVMDispatch_CoreNotEquals::StaticStruct());
 }
 
 FRigUnit_ItemTypeEquals_Execute()
diff --git a/Engine/Plugins/Animation/ControlRig/Source/ControlRig/Private/Units/Math/RigUnit_MathBool.cpp b/Engine/Plugins/Animation/ControlRig/Source/ControlRig/Private/Units/Math/RigUnit_MathBool.cpp
--- a/Engine/Plugins/Animation/ControlRig/Source/ControlRig/Private/Units/Math/RigUnit_MathBool.cpp
+++ b/Engine/Plugins/Animation/ControlRig/Source/ControlRig/Private/Units/Math/RigUnit_MathBool.cpp
@@ -2,6 +2,7 @@
 
 #include "Units/Math/RigUnit_MathBool.h"
 #include "Units/RigUnitContext.h"
+#include "RigVMFunctions/RigVMDispatch_Core.h"
 
 #include UE_INLINE_GENERATED_CPP_BY_NAME(RigUnit_MathBool)
 
@@ -66,7 +67,7 @@ FRigUnit_MathBoolEquals_Execute()
 
 FRigVMStructUpgradeInfo FRigUnit_MathBoolEquals::GetUpgradeInfo() const
 {
-	return FRigVMStructUpgradeInfo::MakeFromStructToFactory(StaticStruct(), FRigDispatch_CoreEquals::StaticStruct());
+	return FRigVMStructUpgradeInfo::MakeFromStructToFactory(StaticStruct(), FRigVMDispatch_CoreEquals::StaticStruct());
 }
 
 FRigUnit_MathBoolNotEquals_Execute()
@@ -77,7 +78,7 @@ FRigUnit_MathBoolNotEquals_Execute()
 
 FRigVMStructUpgradeInfo FRigUnit_MathBoolNotEquals::GetUpgradeInfo() const
 {
-	return FRigVMStructUpgradeInfo::MakeFromStructToFactory(StaticStruct(), FRigDispatch_CoreNotEquals::StaticStruct());
+	return FRigVMStructUpgradeInfo::MakeFromStructToFactory(StaticStruct(), FRigVMDispatch_CoreNotEquals::StaticStruct());
 }
 
 FRigUnit_MathBoolToggled_Execute()
diff --git a/Engine/Plugins/Animation/ControlRig/Source/ControlRig/Private/Units/Math/RigUnit_MathBool.h b/Engine/Plugins/Animation/ControlRig/Source/ControlRig/Private/Units/Math/RigUnit_MathBool.h
--- a/Engine/Plugins/Animation/ControlRig/Source/ControlRig/Private/Units/Math/RigUnit_MathBool.h
+++ b/Engine/Plugins/Animation/ControlRig/Source/ControlRig/Private/Units/Math/RigUnit_MathBool.h
@@ -3,7 +3,6 @@
 #pragma once
 
 #include "RigUnit_MathBase.h"
-#include "Units/Core/RigUnit_CoreDispatch.h"
 #include "RigUnit_MathBool.generated.h"
 
 USTRUCT(meta=(Abstract, Category="Math|Boolean"))
diff --git a/Engine/Plugins/Animation/ControlRig/Source/ControlRig/Private/Units/Math/RigUnit_MathDouble.cpp b/Engine/Plugins/Animation/ControlRig/Source/ControlRig/Private/Units/Math/RigUnit_MathDouble.cpp
--- a/Engine/Plugins/Animation/ControlRig/Source/ControlRig/Private/Units/Math/RigUnit_MathDouble.cpp
+++ b/Engine/Plugins/Animation/ControlRig/Source/ControlRig/Private/Units/Math/RigUnit_MathDouble.cpp
@@ -2,7 +2,7 @@
 
 #include "Units/Math/RigUnit_MathDouble.h"
 #include "Units/RigUnitContext.h"
-#include "Units/Core/RigUnit_CoreDispatch.h"
+#include "RigVMFunctions/RigVMDispatch_Core.h"
 
 #include UE_INLINE_GENERATED_CPP_BY_NAME(RigUnit_MathDouble)
 
@@ -186,7 +186,7 @@ FRigUnit_MathDoubleEquals_Execute()
 
 FRigVMStructUpgradeInfo FRigUnit_MathDoubleEquals::GetUpgradeInfo() const
 {
-	return FRigVMStructUpgradeInfo::MakeFromStructToFactory(StaticStruct(), FRigDispatch_CoreEquals::StaticStruct());
+	return FRigVMStructUpgradeInfo::MakeFromStructToFactory(StaticStruct(), FRigVMDispatch_CoreEquals::StaticStruct());
 }
 
 FRigUnit_MathDoubleNotEquals_Execute()
@@ -197,7 +197,7 @@ FRigUnit_MathDoubleNotEquals_Execute()
 
 FRigVMStructUpgradeInfo FRigUnit_MathDoubleNotEquals::GetUpgradeInfo() const
 {
-	return FRigVMStructUpgradeInfo::MakeFromStructToFactory(StaticStruct(), FRigDispatch_CoreNotEquals::StaticStruct());
+	return FRigVMStructUpgradeInfo::MakeFromStructToFactory(StaticStruct(), FRigVMDispatch_CoreNotEquals::StaticStruct());
 }
 
 FRigUnit_MathDoubleGreater_Execute()
diff --git a/Engine/Plugins/Animation/ControlRig/Source/ControlRig/Private/Units/Math/RigUnit_MathFloat.cpp b/Engine/Plugins/Animation/ControlRig/Source/ControlRig/Private/Units/Math/RigUnit_MathFloat.cpp
--- a/Engine/Plugins/Animation/ControlRig/Source/ControlRig/Private/Units/Math/RigUnit_MathFloat.cpp
+++ b/Engine/Plugins/Animation/ControlRig/Source/ControlRig/Private/Units/Math/RigUnit_MathFloat.cpp
@@ -2,7 +2,7 @@
 
 #include "Units/Math/RigUnit_MathFloat.h"
 #include "Units/RigUnitContext.h"
-#include "Units/Core/RigUnit_CoreDispatch.h"
+#include "RigVMFunctions/RigVMDispatch_Core.h"
 
 #include UE_INLINE_GENERATED_CPP_BY_NAME(RigUnit_MathFloat)
 
@@ -186,7 +186,7 @@ FRigUnit_MathFloatEquals_Execute()
 
 FRigVMStructUpgradeInfo FRigUnit_MathFloatEquals::GetUpgradeInfo() const
 {
-	return FRigVMStructUpgradeInfo::MakeFromStructToFactory(StaticStruct(), FRigDispatch_CoreEquals::StaticStruct());
+	return FRigVMStructUpgradeInfo::MakeFromStructToFactory(StaticStruct(), FRigVMDispatch_CoreEquals::StaticStruct());
 }
 
 FRigUnit_MathFloatNotEquals_Execute()
@@ -197,7 +197,7 @@ FRigUnit_MathFloatNotEquals_Execute()
 
 FRigVMStructUpgradeInfo FRigUnit_MathFloatNotEquals::GetUpgradeInfo() const
 {
-	return FRigVMStructUpgradeInfo::MakeFromStructToFactory(StaticStruct(), FRigDispatch_CoreNotEquals::StaticStruct());
+	return FRigVMStructUpgradeInfo::MakeFromStructToFactory(StaticStruct(), FRigVMDispatch_CoreNotEquals::StaticStruct());
 }
 
 FRigUnit_MathFloatGreater_Execute()
diff --git a/Engine/Plugins/Animation/ControlRig/Source/ControlRig/Private/Units/Math/RigUnit_MathInt.cpp b/Engine/Plugins/Animation/ControlRig/Source/ControlRig/Private/Units/Math/RigUnit_MathInt.cpp
--- a/Engine/Plugins/Animation/ControlRig/Source/ControlRig/Private/Units/Math/RigUnit_MathInt.cpp
+++ b/Engine/Plugins/Animation/ControlRig/Source/ControlRig/Private/Units/Math/RigUnit_MathInt.cpp
@@ -2,7 +2,7 @@
 
 #include "Units/Math/RigUnit_MathInt.h"
 #include "Units/RigUnitContext.h"
-#include "Units/Core/RigUnit_CoreDispatch.h"
+#include "RigVMFunctions/RigVMDispatch_Core.h"
 
 #include UE_INLINE_GENERATED_CPP_BY_NAME(RigUnit_MathInt)
 
@@ -114,7 +114,7 @@ FRigUnit_MathIntEquals_Execute()
 
 FRigVMStructUpgradeInfo FRigUnit_MathIntEquals::GetUpgradeInfo() const
 {
-	return FRigVMStructUpgradeInfo::MakeFromStructToFactory(StaticStruct(), FRigDispatch_CoreEquals::StaticStruct());
+	return FRigVMStructUpgradeInfo::MakeFromStructToFactory(StaticStruct(), FRigVMDispatch_CoreEquals::StaticStruct());
 }
 
 FRigUnit_MathIntNotEquals_Execute()
@@ -125,7 +125,7 @@ FRigUnit_MathIntNotEquals_Execute()
 
 FRigVMStructUpgradeInfo FRigUnit_MathIntNotEquals::GetUpgradeInfo() const
 {
-	return FRigVMStructUpgradeInfo::MakeFromStructToFactory(StaticStruct(), FRigDispatch_CoreNotEquals::StaticStruct());
+	return FRigVMStructUpgradeInfo::MakeFromStructToFactory(StaticStruct(), FRigVMDispatch_CoreNotEquals::StaticStruct());
 }
 
 FRigUnit_MathIntGreater_Execute()
diff --git a/Engine/Plugins/Animation/ControlRig/Source/ControlRig/Private/Units/Math/RigUnit_MathQuaternion.cpp b/Engine/Plugins/Animation/ControlRig/Source/ControlRig/Private/Units/Math/RigUnit_MathQuaternion.cpp
--- a/Engine/Plugins/Animation/ControlRig/Source/ControlRig/Private/Units/Math/RigUnit_MathQuaternion.cpp
+++ b/Engine/Plugins/Animation/ControlRig/Source/ControlRig/Private/Units/Math/RigUnit_MathQuaternion.cpp
@@ -5,7 +5,7 @@
 #include "Units/RigUnitContext.h"
 #include "AnimationCoreLibrary.h"
 #include "Math/ControlRigMathLibrary.h"
-#include "Units/Core/RigUnit_CoreDispatch.h"
+#include "RigVMFunctions/RigVMDispatch_Core.h"
 
 #include UE_INLINE_GENERATED_CPP_BY_NAME(RigUnit_MathQuaternion)
 
@@ -160,7 +160,7 @@ FRigUnit_MathQuaternionEquals_Execute()
 
 FRigVMStructUpgradeInfo FRigUnit_MathQuaternionEquals::GetUpgradeInfo() const
 {
-	return FRigVMStructUpgradeInfo::MakeFromStructToFactory(StaticStruct(), FRigDispatch_CoreEquals::StaticStruct());
+	return FRigVMStructUpgradeInfo::MakeFromStructToFactory(StaticStruct(), FRigVMDispatch_CoreEquals::StaticStruct());
 }
 
 FRigUnit_MathQuaternionNotEquals_Execute()
@@ -171,7 +171,7 @@ FRigUnit_MathQuaternionNotEquals_Execute()
 
 FRigVMStructUpgradeInfo FRigUnit_MathQuaternionNotEquals::GetUpgradeInfo() const
 {
-	return FRigVMStructUpgradeInfo::MakeFromStructToFactory(StaticStruct(), FRigDispatch_CoreNotEquals::StaticStruct());
+	return FRigVMStructUpgradeInfo::MakeFromStructToFactory(StaticStruct(), FRigVMDispatch_CoreNotEquals::StaticStruct());
 }
 
 FRigUnit_MathQuaternionSelectBool_Execute()
diff --git a/Engine/Plugins/Animation/ControlRig/Source/ControlRig/Private/Units/Math/RigUnit_MathVector.cpp b/Engine/Plugins/Animation/ControlRig/Source/ControlRig/Private/Units/Math/RigUnit_MathVector.cpp
--- a/Engine/Plugins/Animation/ControlRig/Source/ControlRig/Private/Units/Math/RigUnit_MathVector.cpp
+++ b/Engine/Plugins/Animation/ControlRig/Source/ControlRig/Private/Units/Math/RigUnit_MathVector.cpp
@@ -4,7 +4,7 @@
 #include "Units/Math/RigUnit_MathTransform.h"
 #include "Units/RigUnitContext.h"
 #include "Math/ControlRigMathLibrary.h"
-#include "Units/Core/RigUnit_CoreDispatch.h"
+#include "RigVMFunctions/RigVMDispatch_Core.h"
 
 #include UE_INLINE_GENERATED_CPP_BY_NAME(RigUnit_MathVector)
 
@@ -212,7 +212,7 @@ FRigUnit_MathVectorEquals_Execute()
 
 FRigVMStructUpgradeInfo FRigUnit_MathVectorEquals::GetUpgradeInfo() const
 {
-	return FRigVMStructUpgradeInfo::MakeFromStructToFactory(StaticStruct(), FRigDispatch_CoreEquals::StaticStruct());
+	return FRigVMStructUpgradeInfo::MakeFromStructToFactory(StaticStruct(), FRigVMDispatch_CoreEquals::StaticStruct());
 }
 
 FRigUnit_MathVectorNotEquals_Execute()
@@ -223,7 +223,7 @@ FRigUnit_MathVectorNotEquals_Execute()
 
 FRigVMStructUpgradeInfo FRigUnit_MathVectorNotEquals::GetUpgradeInfo() const
 {
-	return FRigVMStructUpgradeInfo::MakeFromStructToFactory(StaticStruct(), FRigDispatch_CoreNotEquals::StaticStruct());
+	return FRigVMStructUpgradeInfo::MakeFromStructToFactory(StaticStruct(), FRigVMDispatch_CoreNotEquals::StaticStruct());
 }
 
 FRigUnit_MathVectorIsNearlyZero_Execute()
diff --git a/Engine/Plugins/Animation/ControlRig/Source/ControlRigDeveloper/Private/ControlRigBlueprint.cpp b/Engine/Plugins/Animation/ControlRig/Source/ControlRigDeveloper/Private/ControlRigBlueprint.cpp
--- a/Engine/Plugins/Animation/ControlRig/Source/ControlRigDeveloper/Private/ControlRigBlueprint.cpp
+++ b/Engine/Plugins/Animation/ControlRig/Source/ControlRigDeveloper/Private/ControlRigBlueprint.cpp
@@ -892,7 +892,7 @@ void UControlRigBlueprint::PostLoad()
 		}
 		
 		InitializeModelIfRequired(false /* recompile vm */);
-
+		GetRigVMClient()->PatchModelsOnLoad();
 		PatchFunctionReferencesOnLoad();
 		PatchVariableNodesOnLoad();
 		PatchVariableNodesWithIncorrectType();
diff --git a/Engine/Plugins/Animation/ControlRig/Source/ControlRigEditor/Private/ControlRigEditorModule.cpp b/Engine/Plugins/Animation/ControlRig/Source/ControlRigEditor/Private/ControlRigEditorModule.cpp
--- a/Engine/Plugins/Animation/ControlRig/Source/ControlRigEditor/Private/ControlRigEditorModule.cpp
+++ b/Engine/Plugins/Animation/ControlRig/Source/ControlRigEditor/Private/ControlRigEditorModule.cpp
@@ -1072,7 +1072,7 @@ void FControlRigEditorModule::GetTypeActions(UControlRigBlueprint* CRB, FBluepri
 	for(const FRigVMFunction& Function : Registry.GetFunctions())
 	{
 		UScriptStruct* Struct = Function.Struct;
-		if (Struct == nullptr || !Struct->IsChildOf(FRigUnit::StaticStruct()))
+		if (Struct == nullptr || !Struct->IsChildOf(FRigVMStruct::StaticStruct()))
 		{
 			continue;
 		}
diff --git a/Engine/Source/Developer/RigVMDeveloper/Private/RigVMModel/RigVMClient.cpp b/Engine/Source/Developer/RigVMDeveloper/Private/RigVMModel/RigVMClient.cpp
--- a/Engine/Source/Developer/RigVMDeveloper/Private/RigVMModel/RigVMClient.cpp
+++ b/Engine/Source/Developer/RigVMDeveloper/Private/RigVMModel/RigVMClient.cpp
@@ -641,3 +641,13 @@ void FRigVMClient::DestroyObject(UObject* InObject)
 	}
 }
 
+void FRigVMClient::PatchModelsOnLoad()
+{
+	TArray<URigVMGraph*> AllModels = GetAllModels(true, true);
+	for(URigVMGraph* Model : AllModels)
+	{
+		URigVMController* Controller = GetOrCreateController(Model);
+		Controller->PatchDispatchNodesOnLoad();
+	}
+}
+
diff --git a/Engine/Source/Developer/RigVMDeveloper/Private/RigVMModel/RigVMController.cpp b/Engine/Source/Developer/RigVMDeveloper/Private/RigVMModel/RigVMController.cpp
--- a/Engine/Source/Developer/RigVMDeveloper/Private/RigVMModel/RigVMController.cpp
+++ b/Engine/Source/Developer/RigVMDeveloper/Private/RigVMModel/RigVMController.cpp
@@ -12958,59 +12958,7 @@ URigVMTemplateNode* URigVMController::AddTemplateNode(const FName& InNotation, c
 	Node->InitializeFilteredPermutations();
 
 	FRigVMRegistry& Registry = FRigVMRegistry::Get();
-
-	for (int32 ArgIndex = 0; ArgIndex < Template->NumExecuteArguments(); ArgIndex++)
-	{
-		const FRigVMExecuteArgument* Arg = Template->GetExecuteArgument(ArgIndex);
-		URigVMPin* Pin = NewObject<URigVMPin>(Node, Arg->Name);
-		const FRigVMTemplateArgumentType Type = Registry.GetType(Arg->TypeIndex);
-		
-		Pin->CPPType = Type.CPPType.ToString();
-		Pin->CPPTypeObject = Type.CPPTypeObject;
-		if (Pin->CPPTypeObject)
-		{
-			Pin->CPPTypeObjectPath = *Pin->CPPTypeObject->GetPathName();
-		}
-		Pin->Direction = Arg->Direction;
-		Pin->LastKnownTypeIndex = Arg->TypeIndex;
-		Pin->LastKnownCPPType = Pin->CPPType;
-
-		AddNodePin(Node, Pin);
-	}
-	
-	for (int32 ArgIndex = 0; ArgIndex < Template->NumArguments(); ArgIndex++)
-	{
-		const FRigVMTemplateArgument* Arg = Template->GetArgument(ArgIndex);
-
-		URigVMPin* Pin = NewObject<URigVMPin>(Node, Arg->GetName());
-		const TRigVMTypeIndex& TypeIndex = Types.FindChecked(Arg->GetName());
-		const FRigVMTemplateArgumentType Type = Registry.GetType(TypeIndex);
-		
-		Pin->CPPType = Type.CPPType.ToString();
-		Pin->CPPTypeObject = Type.CPPTypeObject;
-		if (Pin->CPPTypeObject)
-		{
-			Pin->CPPTypeObjectPath = *Pin->CPPTypeObject->GetPathName();
-		}
-		Pin->Direction = Arg->GetDirection();
-		Pin->LastKnownTypeIndex = TypeIndex;
-		Pin->LastKnownCPPType = Pin->CPPType;
-
-		AddNodePin(Node, Pin);
-
-		if(!Pin->IsWildCard())
-		{
-			const FString DefaultValue = Node->GetInitialDefaultValueForPin(Pin->GetFName());
-			if(UScriptStruct* ScriptStruct = Cast<UScriptStruct>(Pin->CPPTypeObject))
-			{
-				AddPinsForStruct(ScriptStruct, Pin->GetNode(), Pin, Pin->Direction, DefaultValue, false, false);
-			}
-			else if(!DefaultValue.IsEmpty())
-			{
-				SetPinDefaultValue(Pin, DefaultValue, true, false, false, false);
-			}
-		}
-	}
+	AddPinsForTemplate(Template, Types, Node);
 
 	UpdateTemplateNodePinTypes(Node, false);
 
@@ -14196,6 +14144,66 @@ void URigVMController::AddPinsForArray(FArrayProperty* InArrayProperty, URigVMNo
 	}
 }
 
+void URigVMController::AddPinsForTemplate(const FRigVMTemplate* InTemplate, const FRigVMTemplateTypeMap& InPinTypeMap, URigVMNode* InNode)
+{
+	const FRigVMRegistry& Registry = FRigVMRegistry::Get();
+	
+	for (int32 ArgIndex = 0; ArgIndex < InTemplate->NumExecuteArguments(); ArgIndex++)
+    {
+        const FRigVMExecuteArgument* Arg = InTemplate->GetExecuteArgument(ArgIndex);
+        URigVMPin* Pin = NewObject<URigVMPin>(InNode, Arg->Name);
+        const FRigVMTemplateArgumentType Type = Registry.GetType(Arg->TypeIndex);
+        
+        Pin->CPPType = Type.CPPType.ToString();
+        Pin->CPPTypeObject = Type.CPPTypeObject;
+        if (Pin->CPPTypeObject)
+        {
+        	Pin->CPPTypeObjectPath = *Pin->CPPTypeObject->GetPathName();
+        }
+        Pin->Direction = Arg->Direction;
+        Pin->LastKnownTypeIndex = Arg->TypeIndex;
+        Pin->LastKnownCPPType = Pin->CPPType;
+
+        AddNodePin(InNode, Pin);
+    }
+
+	for (int32 ArgIndex = 0; ArgIndex < InTemplate->NumArguments(); ArgIndex++)
+	{
+		const FRigVMTemplateArgument* Arg = InTemplate->GetArgument(ArgIndex);
+
+		URigVMPin* Pin = NewObject<URigVMPin>(InNode, Arg->GetName());
+		const TRigVMTypeIndex& TypeIndex = InPinTypeMap.FindChecked(Arg->GetName());
+		const FRigVMTemplateArgumentType Type = FRigVMRegistry::Get().GetType(TypeIndex);
+		Pin->CPPType = Type.CPPType.ToString();
+		Pin->CPPTypeObject = Type.CPPTypeObject;
+		if (Pin->CPPTypeObject)
+		{
+			Pin->CPPTypeObjectPath = *Pin->CPPTypeObject->GetPathName();
+		}
+		Pin->Direction = Arg->GetDirection();
+
+		AddNodePin(InNode, Pin);
+
+		if(!Pin->IsWildCard() && !Pin->IsArray())
+		{
+			FString DefaultValue;
+			if(const URigVMTemplateNode* TemplateNode = Cast<URigVMTemplateNode>(InNode))
+			{
+				DefaultValue = TemplateNode->GetInitialDefaultValueForPin(Pin->GetFName());
+			}
+			
+			if(UScriptStruct* ScriptStruct = Cast<UScriptStruct>(Pin->CPPTypeObject))
+			{
+				AddPinsForStruct(ScriptStruct, Pin->GetNode(), Pin, Pin->Direction, DefaultValue, false, false);
+			}
+			else if(!DefaultValue.IsEmpty())
+			{
+				SetPinDefaultValue(Pin, DefaultValue, true, false, false, false);
+			}
+		}
+	}
+}
+
 void URigVMController::ConfigurePinFromProperty(FProperty* InProperty, URigVMPin* InOutPin, ERigVMPinDirection InPinDirection)
 {
 	if (InPinDirection == ERigVMPinDirection::Invalid)
@@ -14989,47 +14997,16 @@ void URigVMController::RepopulatePinsOnNode(URigVMNode* InNode, bool bFollowCore
 	}
 	else if (DispatchNode)
 	{
-		TMap<FString, TRigVMTypeIndex> PinTypeMap;
+		FRigVMTemplateTypeMap PinTypeMap;
 		for (const URigVMPin* Pin : DispatchNode->Pins)
 		{
-			PinTypeMap.Add(Pin->GetPinPath(), Pin->GetTypeIndex());
+			PinTypeMap.Add(Pin->GetFName(), Pin->GetTypeIndex());
 		}
 		
 		RemovePinsDuringRepopulate(DispatchNode, DispatchNode->Pins, bNotify, bSetupOrphanedPins);
 
 		const FRigVMTemplate* Template = DispatchNode->GetTemplate();
-		
-		for (int32 ArgIndex = 0; ArgIndex < Template->NumArguments(); ArgIndex++)
-		{
-			const FRigVMTemplateArgument* Arg = Template->GetArgument(ArgIndex);
-
-			URigVMPin* Pin = NewObject<URigVMPin>(DispatchNode, Arg->GetName());
-			const TRigVMTypeIndex& TypeIndex = PinTypeMap.FindChecked(Pin->GetPinPath());
-			const FRigVMTemplateArgumentType Type = FRigVMRegistry::Get().GetType(TypeIndex);
-			Pin->CPPType = Type.CPPType.ToString();
-			Pin->CPPTypeObject = Type.CPPTypeObject;
-			if (Pin->CPPTypeObject)
-			{
-				Pin->CPPTypeObjectPath = *Pin->CPPTypeObject->GetPathName();
-			}
-			Pin->Direction = Arg->GetDirection();
-
-			AddNodePin(DispatchNode, Pin);
-
-			if(!Pin->IsWildCard() && !Pin->IsArray())
-			{
-				// any serialize default values will be applied later with ApplyPinStates(...)
-				const FString DefaultValue = DispatchNode->GetInitialDefaultValueForPin(Pin->GetFName());
-				if(UScriptStruct* ScriptStruct = Cast<UScriptStruct>(Pin->CPPTypeObject))
-				{
-					AddPinsForStruct(ScriptStruct, Pin->GetNode(), Pin, Pin->Direction, DefaultValue, false, false);
-				}
-				else if(!DefaultValue.IsEmpty())
-				{
-					SetPinDefaultValue(Pin, DefaultValue, true, false, false, false);
-				}
-			}
-		}
+		AddPinsForTemplate(Template, PinTypeMap, DispatchNode);
 		
 		ResolveTemplateNodeMetaData(DispatchNode, false);
 	}
@@ -19471,6 +19448,50 @@ FRigVMTemplate::FTypeMap URigVMController::GetCommonlyUsedTypesForTemplate(
 	return Template->GetArgumentTypesFromString(TypesString);
 }
 
+void URigVMController::PatchDispatchNodesOnLoad()
+{
+	if (const URigVMGraph* Graph = GetGraph())
+	{
+		for(URigVMNode* Node : Graph->GetNodes())
+		{
+			if(URigVMDispatchNode* DispatchNode = Cast<URigVMDispatchNode>(Node))
+			{
+				// find template performs backwards lookup
+				if(const FRigVMTemplate* Template = DispatchNode->GetTemplate())
+				{
+					if(Template->GetNotation() != DispatchNode->TemplateNotation)
+					{
+						DispatchNode->TemplateNotation = Template->GetNotation();
+						if(!DispatchNode->ResolvedFunctionName.IsEmpty())
+						{
+							FString FactoryName, ArgumentsString;
+							if(DispatchNode->ResolvedFunctionName.Split(TEXT("::"), &FactoryName, &ArgumentsString))
+							{
+								DispatchNode->ResolvedFunctionName.Reset();
+								DispatchNode->ResolvedPermutation = INDEX_NONE;
+								
+								const FRigVMTemplateTypeMap ArgumentTypes = Template->GetArgumentTypesFromString(ArgumentsString);
+								if(ArgumentTypes.Num() == Template->NumArguments())
+								{
+									if(const FRigVMDispatchFactory* Factory = Template->GetDispatchFactory())
+									{
+										const FString ResolvedPermutationName = Factory->GetPermutationName(ArgumentTypes);
+										if(const FRigVMFunction* Function = FRigVMRegistry::Get().FindFunction(*ResolvedPermutationName))
+										{
+											DispatchNode->ResolvedFunctionName = Function->GetName();
+											DispatchNode->ResolvedPermutation = Template->FindPermutation(Function);
+										}
+									}
+								}
+							}
+						}
+					}
+				}
+			}
+		}
+	}
+}
+
 #endif
 
 URigVMControllerSettings::URigVMControllerSettings(const FObjectInitializer& Initializer)
diff --git a/Engine/Source/Developer/RigVMDeveloper/Public/RigVMModel/RigVMClient.h b/Engine/Source/Developer/RigVMDeveloper/Public/RigVMModel/RigVMClient.h
--- a/Engine/Source/Developer/RigVMDeveloper/Public/RigVMModel/RigVMClient.h
+++ b/Engine/Source/Developer/RigVMDeveloper/Public/RigVMModel/RigVMClient.h
@@ -127,6 +127,9 @@ public:
 	static FName GetUniqueName(UObject* InOuter, const FName& InDesiredName);
 	static void DestroyObject(UObject* InObject);
 
+	// backwards compatibility
+	void PatchModelsOnLoad();
+
 private:
 
 	enum ERigVMClientAction
diff --git a/Engine/Source/Developer/RigVMDeveloper/Public/RigVMModel/RigVMController.h b/Engine/Source/Developer/RigVMDeveloper/Public/RigVMModel/RigVMController.h
--- a/Engine/Source/Developer/RigVMDeveloper/Public/RigVMModel/RigVMController.h
+++ b/Engine/Source/Developer/RigVMDeveloper/Public/RigVMModel/RigVMController.h
@@ -991,6 +991,7 @@ private:
 	bool CanAddFunctionRefForDefinition(URigVMLibraryNode* InFunctionDefinition, bool bReportErrors);
 	void AddPinsForStruct(UStruct* InStruct, URigVMNode* InNode, URigVMPin* InParentPin, ERigVMPinDirection InPinDirection, const FString& InDefaultValue, bool bAutoExpandArrays, bool bNotify = false);
 	void AddPinsForArray(FArrayProperty* InArrayProperty, URigVMNode* InNode, URigVMPin* InParentPin, ERigVMPinDirection InPinDirection, const TArray<FString>& InDefaultValues, bool bAutoExpandArrays);
+	void AddPinsForTemplate(const FRigVMTemplate* InTemplate, const FRigVMTemplateTypeMap& InPinTypeMap, URigVMNode* InNode);
 	void ConfigurePinFromProperty(FProperty* InProperty, URigVMPin* InOutPin, ERigVMPinDirection InPinDirection = ERigVMPinDirection::Invalid);
 	void ConfigurePinFromPin(URigVMPin* InOutPin, URigVMPin* InPin, bool bCopyDisplayName = false);
 	virtual bool ShouldStructBeUnfolded(const UStruct* InStruct);
@@ -1165,6 +1166,11 @@ public:
 	// Inquire on the commonly used types for a template node. This can be used to resolve a node without user input (as a default)
 	FRigVMTemplate::FTypeMap GetCommonlyUsedTypesForTemplate(const URigVMTemplateNode* InNode) const;
 #endif
+
+protected:
+
+	// backwards compatibility code
+	void PatchDispatchNodesOnLoad();
 	
 private: 
 	UPROPERTY(transient)
@@ -1240,6 +1246,7 @@ private:
 	friend class FRigVMParserAST;
 	friend class FRigVMControllerCompileBracketScope;
 	friend class FRigVMControllerGraphGuard;
+	friend struct FRigVMClient;
 };
 
 class FRigVMControllerGraphGuard
diff --git a/Engine/Source/Runtime/RigVM/Private/RigVMCore/RigVMRegistry.cpp b/Engine/Source/Runtime/RigVM/Private/RigVMCore/RigVMRegistry.cpp
--- a/Engine/Source/Runtime/RigVM/Private/RigVMCore/RigVMRegistry.cpp
+++ b/Engine/Source/Runtime/RigVM/Private/RigVMCore/RigVMRegistry.cpp
@@ -13,6 +13,10 @@
 #include "Misc/CoreDelegates.h"
 #include "Misc/DelayedAutoRegister.h"
 
+#if WITH_EDITOR
+#include "UObject/CoreRedirects.h"
+#endif
+
 const FName FRigVMRegistry::TemplateNameMetaName = TEXT("TemplateName");
 
 
@@ -1167,9 +1171,34 @@ const FRigVMTemplate* FRigVMRegistry::FindTemplate(const FName& InNotation, bool
 		}
 	}
 
-	// if we arrive here we may have a template that used to contain an executecontext.
+	const FString OriginalNotation = InNotation.ToString();
+
+	// we may have a dispatch factory which has to be redirected
+#if WITH_EDITOR
+	if(OriginalNotation.StartsWith(FRigVMDispatchFactory::DispatchPrefix))
+	{
+		const FString OriginalDispatchFactoryName = OriginalNotation
+			.Left(OriginalNotation.Find(TEXT("(")))
+			.RightChop(FRigVMDispatchFactory::DispatchPrefix.Len());
+
+		const FCoreRedirectObjectName OldObjectName(OriginalDispatchFactoryName);
+		TArray<const FCoreRedirect*> Redirects;
+		if(FCoreRedirects::GetMatchingRedirects(ECoreRedirectFlags::Type_Struct, OldObjectName, Redirects, ECoreRedirectMatchFlags::AllowPartialMatch))
+		{
+			for(const FCoreRedirect* Redirect : Redirects)
+			{
+				const FString NewDispatchFactoryName = FRigVMDispatchFactory::DispatchPrefix + Redirect->NewName.ObjectName.ToString();
+				if(const FRigVMDispatchFactory* NewDispatchFactory = FindDispatchFactory(*NewDispatchFactoryName))
+				{
+					return NewDispatchFactory->GetTemplate();
+				}
+			}
+		}
+	}
+#endif
+
+	// if we still arrive here we may have a template that used to contain an executecontext.
 	{
-		const FString OriginalNotation = InNotation.ToString();
 		FString SanitizedNotation = OriginalNotation;
 
 		static const TArray<TPair<FString, FString>> ExecuteContextArgs = {
diff --git a/Engine/Source/Runtime/RigVM/Private/RigVMFunctions/RigVMDispatch_Core.cpp b/Engine/Source/Runtime/RigVM/Private/RigVMFunctions/RigVMDispatch_Core.cpp
new file mode 100644
--- /dev/null
+++ b/Engine/Source/Runtime/RigVM/Private/RigVMFunctions/RigVMDispatch_Core.cpp
@@ -0,0 +1,110 @@
+// Copyright Epic Games, Inc. All Rights Reserved.
+
+#include "RigVMFunctions/RigVMDispatch_Core.h"
+#include "RigVMCore/RigVMRegistry.h"
+#include "EulerTransform.h"
+
+#include UE_INLINE_GENERATED_CPP_BY_NAME(RigVMDispatch_Core)
+
+TArray<FRigVMTemplateArgument> FRigVMDispatch_CoreEquals::GetArguments() const
+{
+	const TArray<FRigVMTemplateArgument::ETypeCategory> ValueCategories = {
+		FRigVMTemplateArgument::ETypeCategory_SingleAnyValue,
+		FRigVMTemplateArgument::ETypeCategory_ArrayAnyValue
+	};
+	return {
+		FRigVMTemplateArgument(TEXT("A"), ERigVMPinDirection::Input, ValueCategories),
+		FRigVMTemplateArgument(TEXT("B"), ERigVMPinDirection::Input, ValueCategories),
+		FRigVMTemplateArgument(TEXT("Result"), ERigVMPinDirection::Output, RigVMTypeUtils::TypeIndex::Bool)
+	};
+}
+
+FRigVMTemplateTypeMap FRigVMDispatch_CoreEquals::OnNewArgumentType(const FName& InArgumentName,
+	TRigVMTypeIndex InTypeIndex) const
+{
+	FRigVMTemplateTypeMap Types;
+	Types.Add(TEXT("A"), InTypeIndex);
+	Types.Add(TEXT("B"), InTypeIndex);
+	Types.Add(TEXT("Result"), RigVMTypeUtils::TypeIndex::Bool);
+	return Types;
+}
+
+FRigVMFunctionPtr FRigVMDispatch_CoreEquals::GetDispatchFunctionImpl(const FRigVMTemplateTypeMap& InTypes) const
+{
+	const TRigVMTypeIndex TypeIndex = InTypes.FindChecked(TEXT("A"));
+	check(TypeIndex == InTypes.FindChecked(TEXT("B")));
+	check(InTypes.FindChecked(TEXT("Result")) == RigVMTypeUtils::TypeIndex::Bool);
+
+	if(TypeIndex == RigVMTypeUtils::TypeIndex::Float)
+	{
+		return &FRigVMDispatch_CoreEquals::Equals<float>;
+	}
+	if(TypeIndex == RigVMTypeUtils::TypeIndex::Double)
+	{
+		return &FRigVMDispatch_CoreEquals::Equals<double>;
+	}
+	if(TypeIndex == RigVMTypeUtils::TypeIndex::Int32)
+	{
+		return &FRigVMDispatch_CoreEquals::Equals<int32>;
+	}
+	if(TypeIndex == RigVMTypeUtils::TypeIndex::FName)
+	{
+		return &FRigVMDispatch_CoreEquals::NameEquals;
+	}
+	if(TypeIndex == RigVMTypeUtils::TypeIndex::FString)
+	{
+		return &FRigVMDispatch_CoreEquals::StringEquals;
+	}
+	if(TypeIndex == FRigVMRegistry::Get().GetTypeIndex<FVector>())
+	{
+		return &FRigVMDispatch_CoreEquals::MathTypeEquals<FVector>;
+	}
+	if(TypeIndex == FRigVMRegistry::Get().GetTypeIndex<FVector2D>())
+	{
+		return &FRigVMDispatch_CoreEquals::MathTypeEquals<FVector2D>;
+	}
+	if(TypeIndex == FRigVMRegistry::Get().GetTypeIndex<FRotator>())
+	{
+		return &FRigVMDispatch_CoreEquals::MathTypeEquals<FRotator>;
+	}
+	if(TypeIndex == FRigVMRegistry::Get().GetTypeIndex<FQuat>())
+	{
+		return &FRigVMDispatch_CoreEquals::MathTypeEquals<FQuat>;
+	}
+	if(TypeIndex == FRigVMRegistry::Get().GetTypeIndex<FTransform>())
+	{
+		return &FRigVMDispatch_CoreEquals::MathTypeEquals<FTransform>;
+	}
+	if(TypeIndex == FRigVMRegistry::Get().GetTypeIndex<FLinearColor>())
+	{
+		return &FRigVMDispatch_CoreEquals::MathTypeEquals<FLinearColor>;
+	}
+	return &FRigVMDispatch_CoreEquals::Execute;
+}
+
+bool FRigVMDispatch_CoreEquals::AdaptResult(bool bResult, const FRigVMExtendedExecuteContext& InContext)
+{
+	// if the factory is the not equals factory - let's invert the result
+	if(InContext.Factory->GetScriptStruct() == FRigVMDispatch_CoreNotEquals::StaticStruct())
+	{
+		return !bResult;
+	}
+	return bResult;
+}
+
+void FRigVMDispatch_CoreEquals::Execute(FRigVMExtendedExecuteContext& InContext, FRigVMMemoryHandleArray Handles)
+{
+	const FProperty* PropertyA = Handles[0].GetResolvedProperty(); 
+	const FProperty* PropertyB = Handles[1].GetResolvedProperty(); 
+	check(PropertyA);
+	check(PropertyB);
+	check(PropertyA->SameType(PropertyB));
+	check(Handles[2].IsBool());
+
+	const uint8* A = Handles[0].GetData();
+	const uint8* B = Handles[1].GetData();
+	bool& Result = *(bool*)Handles[2].GetData();
+	Result = PropertyA->Identical(A, B);
+	Result = AdaptResult(Result, InContext);
+}
+
diff --git a/Engine/Source/Runtime/RigVM/Private/RigVMFunctions/RigVMDispatch_Print.cpp b/Engine/Source/Runtime/RigVM/Private/RigVMFunctions/RigVMDispatch_Print.cpp
new file mode 100644
--- /dev/null
+++ b/Engine/Source/Runtime/RigVM/Private/RigVMFunctions/RigVMDispatch_Print.cpp
@@ -0,0 +1,112 @@
+// Copyright Epic Games, Inc. All Rights Reserved.
+
+#include "RigVMFunctions/RigVMDispatch_Print.h"
+#include "RigVMCore/RigVMStruct.h"
+#include "RigVMCore/RigVM.h"
+#include "Kismet/KismetSystemLibrary.h"
+
+#include UE_INLINE_GENERATED_CPP_BY_NAME(RigVMDispatch_Print)
+
+TArray<FRigVMTemplateArgument> FRigVMDispatch_Print::GetArguments() const
+{
+	const TArray<FRigVMTemplateArgument::ETypeCategory> ValueCategories = {
+		FRigVMTemplateArgument::ETypeCategory_SingleAnyValue,
+		FRigVMTemplateArgument::ETypeCategory_ArrayAnyValue
+	};
+	return {
+		FRigVMTemplateArgument(TEXT("Prefix"), ERigVMPinDirection::Input, RigVMTypeUtils::TypeIndex::FString),
+		FRigVMTemplateArgument(TEXT("Value"), ERigVMPinDirection::Input, ValueCategories),
+		FRigVMTemplateArgument(TEXT("Enabled"), ERigVMPinDirection::Input, RigVMTypeUtils::TypeIndex::Bool),
+		FRigVMTemplateArgument(TEXT("ScreenDuration"), ERigVMPinDirection::Input, RigVMTypeUtils::TypeIndex::Float),
+		FRigVMTemplateArgument(TEXT("ScreenColor"), ERigVMPinDirection::Input, FRigVMRegistry::Get().GetTypeIndex<FLinearColor>())
+	};
+}
+
+TArray<FRigVMExecuteArgument> FRigVMDispatch_Print::GetExecuteArguments_Impl() const
+{
+	return {{TEXT("ExecuteContext"), ERigVMPinDirection::IO}};
+}
+
+FRigVMTemplateTypeMap FRigVMDispatch_Print::OnNewArgumentType(const FName& InArgumentName,
+                                                              TRigVMTypeIndex InTypeIndex) const
+{
+	FRigVMTemplateTypeMap Types;
+	Types.Add(TEXT("Prefix"), RigVMTypeUtils::TypeIndex::FString);
+	Types.Add(TEXT("Value"), InTypeIndex);
+	Types.Add(TEXT("Enabled"), RigVMTypeUtils::TypeIndex::Bool);
+	Types.Add(TEXT("ScreenDuration"), RigVMTypeUtils::TypeIndex::Float);
+	Types.Add(TEXT("ScreenColor"), FRigVMRegistry::Get().GetTypeIndex<FLinearColor>());
+	return Types;
+}
+
+#if WITH_EDITOR
+
+FString FRigVMDispatch_Print::GetArgumentDefaultValue(const FName& InArgumentName, TRigVMTypeIndex InTypeIndex) const
+{
+	if(InArgumentName == TEXT("Enabled"))
+	{
+		return TEXT("True");
+	}
+	if(InArgumentName == TEXT("ScreenDuration"))
+	{
+		return TEXT("0.050000");
+	}
+	return FRigVMDispatchFactory::GetArgumentDefaultValue(InArgumentName, InTypeIndex);
+}
+
+FString FRigVMDispatch_Print::GetArgumentMetaData(const FName& InArgumentName, const FName& InMetaDataKey) const
+{
+	if(InArgumentName == TEXT("ScreenDuration") || InArgumentName == TEXT("ScreenColor"))
+	{
+		if(InMetaDataKey == FRigVMStruct::DetailsOnlyMetaName)
+		{
+			return TEXT("True");
+		}
+	}
+	return FRigVMDispatchFactory::GetArgumentMetaData(InArgumentName, InMetaDataKey);
+}
+
+#endif
+
+void FRigVMDispatch_Print::Execute(FRigVMExtendedExecuteContext& InContext, FRigVMMemoryHandleArray Handles)
+{
+#if WITH_EDITOR
+	const FProperty* ValueProperty = Handles[1].GetResolvedProperty(); 
+	check(ValueProperty);
+	check(Handles[0].IsString());
+	check(Handles[2].IsBool());
+	check(Handles[3].IsFloat());
+	check(Handles[4].IsType<FLinearColor>());
+
+	const FString& Prefix = *(const FString*)Handles[0].GetData();
+	const bool bEnabled = *(const bool*)Handles[2].GetData();
+	const float& ScreenDuration = *(const float*)Handles[3].GetData();
+	const FLinearColor& ScreenColor = *(const FLinearColor*)Handles[4].GetData();
+	const uint8* Value = Handles[1].GetData();
+	
+	if(!bEnabled)
+	{
+		return;
+	}
+
+	FString String;
+	ValueProperty->ExportText_Direct(String, Value, Value, nullptr, PPF_None, nullptr);
+
+	FString ObjectPath;
+	if(InContext.VM)
+	{
+		ObjectPath = InContext.VM->GetName();
+	}
+
+	static constexpr TCHAR LogFormat[] = TEXT("%s[%04d] %s%s");
+	UE_LOG(LogRigVM, Display, LogFormat, *ObjectPath, InContext.GetPublicData<>().GetInstructionIndex(), *Prefix, *String);
+	const UObject* WorldObject = (const UObject*)InContext.VM;
+
+	if(ScreenDuration > SMALL_NUMBER && WorldObject)
+	{
+		static constexpr TCHAR PrintStringFormat[] = TEXT("[%04d] %s%s");
+		UKismetSystemLibrary::PrintString(WorldObject, FString::Printf(PrintStringFormat, InContext.GetPublicData<>().GetInstructionIndex(), *Prefix, *String), true, false, ScreenColor, ScreenDuration);
+	}
+#endif
+}
+
diff --git a/Engine/Source/Runtime/RigVM/Private/RigVMFunctions/RigVMFunction_String.cpp b/Engine/Source/Runtime/RigVM/Private/RigVMFunctions/RigVMFunction_String.cpp
new file mode 100644
--- /dev/null
+++ b/Engine/Source/Runtime/RigVM/Private/RigVMFunctions/RigVMFunction_String.cpp
@@ -0,0 +1,428 @@
+// Copyright Epic Games, Inc. All Rights Reserved.
+
+#include "RigVMFunctions/RigVMFunction_String.h"
+#include "RigVMCore/RigVM.h"
+
+#include UE_INLINE_GENERATED_CPP_BY_NAME(RigVMFunction_String)
+
+FRigVMFunction_StringConcat_Execute()
+{
+	Result = A + B;
+}
+
+FRigVMFunction_StringTruncate_Execute()
+{
+	Remainder = Name;
+	Chopped = FString();
+
+	if(Name.IsEmpty() || Count <= 0)
+	{
+		return;
+	}
+
+	if (FromEnd)
+	{
+		Remainder = Name.LeftChop(Count);
+		Chopped = Name.Right(Count);
+	}
+	else
+	{
+		Remainder = Name.RightChop(Count);
+		Chopped = Name.Left(Count);
+	}
+}
+
+FRigVMFunction_StringReplace_Execute()
+{
+	Result = Name.Replace(*Old, *New, ESearchCase::CaseSensitive);
+}
+
+FRigVMFunction_StringEndsWith_Execute()
+{
+	Result = Name.EndsWith(Ending, ESearchCase::CaseSensitive);
+}
+
+FRigVMFunction_StringStartsWith_Execute()
+{
+	Result = Name.StartsWith(Start, ESearchCase::CaseSensitive);
+}
+
+FRigVMFunction_StringContains_Execute()
+{
+	Result = Name.Contains(Search, ESearchCase::CaseSensitive);
+}
+
+FRigVMFunction_StringLength_Execute()
+{
+	Length = Value.Len();
+}
+
+FRigVMFunction_StringTrimWhitespace_Execute()
+{
+	Result = Value;
+	Result.TrimStartAndEndInline();
+}
+
+FRigVMFunction_StringToUppercase_Execute()
+{
+	Result = Value.ToUpper();
+}
+
+FRigVMFunction_StringToLowercase_Execute()
+{
+	Result = Value.ToLower();
+}
+
+FRigVMFunction_StringReverse_Execute()
+{
+	Reverse = Value.Reverse();
+}
+
+FRigVMFunction_StringLeft_Execute()
+{
+	Result = Value.Left(Count);
+}
+
+FRigVMFunction_StringRight_Execute()
+{
+	Result = Value.Right(Count);
+}
+
+FRigVMFunction_StringMiddle_Execute()
+{
+	if(Count < 0)
+	{
+		Result = Value.Mid(Start);
+	}
+	else
+	{
+		Result = Value.Mid(Start, Count);
+	}
+}
+
+FRigVMFunction_StringFind_Execute()
+{
+	Index = INDEX_NONE;
+
+	if(!Value.IsEmpty() && !Search.IsEmpty())
+	{
+		Index = Value.Find(Search, ESearchCase::CaseSensitive);
+	}
+
+	Found = Index != INDEX_NONE;
+}
+
+FRigVMFunction_StringSplit_Execute()
+{
+	Result.Reset();
+	if(!Value.IsEmpty())
+	{
+		if(Separator.IsEmpty())
+		{
+			//UE_CONTROLRIG_RIGUNIT_REPORT_ERROR(TEXT("Separator is empty."));
+			return;
+		}
+
+		FString ValueRemaining = Value;
+		FString Left, Right;
+		while(ValueRemaining.Split(Separator, &Left, &Right, ESearchCase::CaseSensitive, ESearchDir::FromStart))
+		{
+			Result.Add(Left);
+			Left.Empty();
+			ValueRemaining = Right;
+		}
+
+		if (!Right.IsEmpty())
+		{
+			Result.Add(Right);
+		}
+	}
+}
+
+FRigVMFunction_StringJoin_Execute()
+{
+	Result.Reset();
+	if(!Values.IsEmpty())
+	{
+		Result = FString::Join(Values, *Separator);
+	}
+}
+
+FRigVMFunction_StringPadInteger_Execute()
+{
+	static constexpr TCHAR Format02Digits[] = TEXT("%02d");
+	static constexpr TCHAR Format03Digits[] = TEXT("%03d");
+	static constexpr TCHAR Format04Digits[] = TEXT("%04d");
+	static constexpr TCHAR Format05Digits[] = TEXT("%05d");
+	static constexpr TCHAR Format06Digits[] = TEXT("%06d");
+	static constexpr TCHAR Format07Digits[] = TEXT("%07d");
+	static constexpr TCHAR Format08Digits[] = TEXT("%08d");
+	static constexpr TCHAR Format09Digits[] = TEXT("%09d");
+	static constexpr TCHAR Format10Digits[] = TEXT("%10d");
+	static constexpr TCHAR Format11Digits[] = TEXT("%11d");
+	static constexpr TCHAR Format12Digits[] = TEXT("%12d");
+	static constexpr TCHAR Format13Digits[] = TEXT("%13d");
+	static constexpr TCHAR Format14Digits[] = TEXT("%14d");
+	static constexpr TCHAR Format15Digits[] = TEXT("%15d");
+	static constexpr TCHAR Format16Digits[] = TEXT("%16d");
+
+	switch(Digits)
+	{
+		case 2:
+		{
+			Result = FString::Printf(Format02Digits, Value); 
+			break;
+		}
+		case 3:
+		{
+			Result = FString::Printf(Format03Digits, Value); 
+			break;
+		}
+		case 4:
+		{
+			Result = FString::Printf(Format04Digits, Value); 
+			break;
+		}
+		case 5:
+		{
+			Result = FString::Printf(Format05Digits, Value); 
+			break;
+		}
+		case 6:
+		{
+			Result = FString::Printf(Format06Digits, Value); 
+			break;
+		}
+		case 7:
+		{
+			Result = FString::Printf(Format07Digits, Value); 
+			break;
+		}
+		case 8:
+		{
+			Result = FString::Printf(Format08Digits, Value); 
+			break;
+		}
+		case 9:
+		{
+			Result = FString::Printf(Format09Digits, Value); 
+			break;
+		}
+		case 10:
+		{
+			Result = FString::Printf(Format10Digits, Value); 
+			break;
+		}
+		case 11:
+		{
+			Result = FString::Printf(Format11Digits, Value); 
+			break;
+		}
+		case 12:
+		{
+			Result = FString::Printf(Format12Digits, Value); 
+			break;
+		}
+		case 13:
+		{
+			Result = FString::Printf(Format13Digits, Value); 
+			break;
+		}
+		case 14:
+		{
+			Result = FString::Printf(Format14Digits, Value); 
+			break;
+		}
+		case 15:
+		{
+			Result = FString::Printf(Format15Digits, Value); 
+			break;
+		}
+		case 16:
+		{
+			Result = FString::Printf(Format16Digits, Value); 
+			break;
+		}
+		default:
+		{
+			Result = FString::FormatAsNumber(Value);
+			break;
+		}
+	}
+}
+
+TArray<FRigVMTemplateArgument> FRigDispatch_ToString::GetArguments() const
+{
+	const TArray<FRigVMTemplateArgument::ETypeCategory> ValueCategories = {
+		FRigVMTemplateArgument::ETypeCategory_SingleAnyValue,
+		FRigVMTemplateArgument::ETypeCategory_ArrayAnyValue
+	};
+	return {
+		FRigVMTemplateArgument(TEXT("Value"), ERigVMPinDirection::Input, ValueCategories),
+		FRigVMTemplateArgument(TEXT("Result"), ERigVMPinDirection::Output, RigVMTypeUtils::TypeIndex::FString)
+	};
+}
+
+FRigVMTemplateTypeMap FRigDispatch_ToString::OnNewArgumentType(const FName& InArgumentName,
+	TRigVMTypeIndex InTypeIndex) const
+{
+	FRigVMTemplateTypeMap Types;
+	Types.Add(TEXT("Value"), InTypeIndex);
+	Types.Add(TEXT("Result"), RigVMTypeUtils::TypeIndex::FString);
+	return Types;
+}
+
+FRigVMFunctionPtr FRigDispatch_ToString::GetDispatchFunctionImpl(const FRigVMTemplateTypeMap& InTypes) const
+{
+	// treat name and string special
+	const TRigVMTypeIndex& ValueTypeIndex = InTypes.FindChecked(TEXT("Value"));
+	if(ValueTypeIndex == RigVMTypeUtils::TypeIndex::FName)
+	{
+		return [](FRigVMExtendedExecuteContext& InContext, FRigVMMemoryHandleArray Handles)
+		{
+			check(Handles[0].IsName());
+			check(Handles[1].IsString());
+
+			const FName& Value = *(const FName*)Handles[0].GetData();
+			FString& Result = *(FString*)Handles[1].GetData();
+			Result = Value.ToString();
+		};
+	}
+	if(ValueTypeIndex == RigVMTypeUtils::TypeIndex::FString)
+	{
+		return [](FRigVMExtendedExecuteContext& InContext, FRigVMMemoryHandleArray Handles)
+		{
+			check(Handles[0].IsString());
+			check(Handles[1].IsString());
+
+			const FString& Value = *(const FString*)Handles[0].GetData();
+			FString& Result = *(FString*)Handles[1].GetData();
+			Result = Value;
+		};
+	}
+	
+	return &FRigDispatch_ToString::Execute;
+}
+
+void FRigDispatch_ToString::Execute(FRigVMExtendedExecuteContext& InContext, FRigVMMemoryHandleArray Handles)
+{
+	const FProperty* ValueProperty = Handles[0].GetResolvedProperty(); 
+	check(ValueProperty);
+	check(Handles[1].IsString());
+
+	const uint8* Value = Handles[0].GetData();
+	FString& Result = *(FString*)Handles[1].GetData();
+	Result.Reset();
+	ValueProperty->ExportText_Direct(Result, Value, Value, nullptr, PPF_None, nullptr);
+}
+
+TArray<FRigVMTemplateArgument> FRigDispatch_FromString::GetArguments() const
+{
+	const TArray<FRigVMTemplateArgument::ETypeCategory> ValueCategories = {
+		FRigVMTemplateArgument::ETypeCategory_SingleAnyValue,
+		FRigVMTemplateArgument::ETypeCategory_ArrayAnyValue
+	};
+	return {
+		FRigVMTemplateArgument(TEXT("String"), ERigVMPinDirection::Input, RigVMTypeUtils::TypeIndex::FString),
+		FRigVMTemplateArgument(TEXT("Result"), ERigVMPinDirection::Output, ValueCategories)
+	};
+}
+
+FRigVMTemplateTypeMap FRigDispatch_FromString::OnNewArgumentType(const FName& InArgumentName,
+	TRigVMTypeIndex InTypeIndex) const
+{
+	FRigVMTemplateTypeMap Types;
+	Types.Add(TEXT("String"), RigVMTypeUtils::TypeIndex::FString);
+	Types.Add(TEXT("Result"), InTypeIndex);
+	return Types;
+}
+
+FRigVMFunctionPtr FRigDispatch_FromString::GetDispatchFunctionImpl(const FRigVMTemplateTypeMap& InTypes) const
+{
+	// treat name and string special
+	const TRigVMTypeIndex& ResultTypeIndex = InTypes.FindChecked(TEXT("Result"));
+	if(ResultTypeIndex == RigVMTypeUtils::TypeIndex::FName)
+	{
+		return [](FRigVMExtendedExecuteContext& InContext, FRigVMMemoryHandleArray Handles)
+		{
+			check(Handles[0].IsString());
+			check(Handles[1].IsName());
+
+			const FString& String = *(FString*)Handles[0].GetData();
+			FName& Result = *(FName*)Handles[1].GetData();
+			Result = *String;
+		};
+	}
+	if(ResultTypeIndex == RigVMTypeUtils::TypeIndex::FString)
+	{
+		return [](FRigVMExtendedExecuteContext& InContext, FRigVMMemoryHandleArray Handles)
+		{
+			check(Handles[0].IsString());
+			check(Handles[1].IsString());
+
+			const FString& String = *(FString*)Handles[0].GetData();
+			FString& Result = *(FString*)Handles[1].GetData();
+			Result = String;
+		};
+	}
+
+	return &FRigDispatch_FromString::Execute;
+}
+
+void FRigDispatch_FromString::Execute(FRigVMExtendedExecuteContext& InContext, FRigVMMemoryHandleArray Handles)
+{
+	const FProperty* ValueProperty = Handles[1].GetResolvedProperty(); 
+	check(ValueProperty);
+	check(Handles[0].IsString());
+
+	const FString& String = *(const FString*)Handles[0].GetData();
+	uint8* Value = Handles[1].GetData();
+
+	class FRigDispatch_FromString_ErrorPipe : public FOutputDevice
+	{
+	public:
+
+		TArray<FString> Errors;
+
+		FRigDispatch_FromString_ErrorPipe()
+			: FOutputDevice()
+		{
+		}
+
+		virtual void Serialize(const TCHAR* V, ELogVerbosity::Type Verbosity, const class FName& Category) override
+		{
+			Errors.Add(FString::Printf(TEXT("Error convert to string: %s"), V));
+		}
+	};
+
+	FRigDispatch_FromString_ErrorPipe ErrorPipe;
+	ValueProperty->ImportText_Direct(*String, Value, nullptr, PPF_None, &ErrorPipe);
+
+	if(!ErrorPipe.Errors.IsEmpty())
+	{
+		ValueProperty->InitializeValue(Value);
+		for(const FString& Error : ErrorPipe.Errors)
+		{
+			/*
+			const FRigUnitContext& Context = GetRigUnitContext(InContext);
+
+#if WITH_EDITOR
+			if(Context.Log != nullptr)
+			{
+				Context.Log->Report(EMessageSeverity::Error, InContext.GetPublicData<>().GetFunctionName(), InContext.GetPublicData<>().GetInstructionIndex(), Error);
+			}
+#endif
+			*/
+
+			FString ObjectPath;
+			if(InContext.VM)
+			{
+				ObjectPath = InContext.VM->GetName();
+			}
+			
+			static constexpr TCHAR ErrorLogFormat[] = TEXT("%s: [%04d] %s");
+			UE_LOG(LogRigVM, Error, ErrorLogFormat, *ObjectPath, InContext.GetPublicData<>().GetInstructionIndex(), *Error);
+		}
+	}
+}
+
diff --git a/Engine/Source/Runtime/RigVM/Public/RigVMFunctions/RigVMDispatch_Core.h b/Engine/Source/Runtime/RigVM/Public/RigVMFunctions/RigVMDispatch_Core.h
new file mode 100644
--- /dev/null
+++ b/Engine/Source/Runtime/RigVM/Public/RigVMFunctions/RigVMDispatch_Core.h
@@ -0,0 +1,92 @@
+// Copyright Epic Games, Inc. All Rights Reserved.
+
+#pragma once
+
+#include "RigVMCore/RigVMDispatchFactory.h"
+#include "RigVMDispatch_Core.generated.h"
+
+USTRUCT(meta=(Abstract, Category = "Core", NodeColor = "0.762745, 1,0, 0.329412"))
+struct FRigVMDispatch_CoreBase : public FRigVMDispatchFactory
+{
+	GENERATED_BODY()
+};
+
+/*
+ * Compares any two values and return true if they are identical
+ */
+USTRUCT(meta=(DisplayName = "Equals", Keywords = "Same,=="))
+struct FRigVMDispatch_CoreEquals : public FRigVMDispatch_CoreBase
+{
+	GENERATED_BODY()
+
+public:
+
+	virtual TArray<FRigVMTemplateArgument> GetArguments() const override;
+	virtual FRigVMTemplateTypeMap OnNewArgumentType(const FName& InArgumentName, TRigVMTypeIndex InTypeIndex) const override;
+
+protected:
+
+	virtual FRigVMFunctionPtr GetDispatchFunctionImpl(const FRigVMTemplateTypeMap& InTypes) const override;
+	static bool AdaptResult(bool bResult, const FRigVMExtendedExecuteContext& InContext);
+	static void Execute(FRigVMExtendedExecuteContext& InContext, FRigVMMemoryHandleArray Handles);
+
+	template<typename T>
+	static void Equals(FRigVMExtendedExecuteContext& InContext, FRigVMMemoryHandleArray Handles)
+	{
+		check(Handles[0].IsType<T>());
+		check(Handles[1].IsType<T>());
+		check(Handles[2].IsBool());
+		const T& A = *(const T*)Handles[0].GetData();
+		const T& B = *(const T*)Handles[1].GetData();
+		bool& Result = *(bool*)Handles[2].GetData();
+		Result = AdaptResult(A == B, InContext);
+	}
+
+	static void NameEquals(FRigVMExtendedExecuteContext& InContext, FRigVMMemoryHandleArray Handles)
+	{
+		check(Handles[0].IsType<FName>());
+		check(Handles[1].IsType<FName>());
+		check(Handles[2].IsBool());
+		const FName& A = *(const FName*)Handles[0].GetData();
+		const FName& B = *(const FName*)Handles[1].GetData();
+		bool& Result = *(bool*)Handles[2].GetData();
+		Result = AdaptResult(A.IsEqual(B, ENameCase::CaseSensitive), InContext);
+	}
+
+	static void StringEquals(FRigVMExtendedExecuteContext& InContext, FRigVMMemoryHandleArray Handles)
+	{
+		check(Handles[0].IsType<FString>());
+		check(Handles[1].IsType<FString>());
+		check(Handles[2].IsBool());
+		const FString& A = *(const FString*)Handles[0].GetData();
+		const FString& B = *(const FString*)Handles[1].GetData();
+		bool& Result = *(bool*)Handles[2].GetData();
+		Result = AdaptResult(A.Equals(B, ESearchCase::CaseSensitive), InContext);
+	}
+
+	template<typename T>
+	static void MathTypeEquals(FRigVMExtendedExecuteContext& InContext, FRigVMMemoryHandleArray Handles)
+	{
+		check(Handles[0].IsType<T>());
+		check(Handles[1].IsType<T>());
+		check(Handles[2].IsBool());
+		const T& A = *(const T*)Handles[0].GetData();
+		const T& B = *(const T*)Handles[1].GetData();
+		bool& Result = *(bool*)Handles[2].GetData();
+		Result = AdaptResult(A.Equals(B), InContext);
+	}
+};
+
+/*
+ * Compares any two values and return true if they are identical
+ */
+USTRUCT(meta=(DisplayName = "Not Equals", Keywords = "Different,!=,Xor"))
+struct FRigVMDispatch_CoreNotEquals : public FRigVMDispatch_CoreEquals
+{
+	GENERATED_BODY()
+
+	// we are inheriting everything from the equals dispatch,
+	// and due to the check of the factory within FRigVMDispatch_CoreEquals::Execute we can
+	// rely on that completely. we only need this class for the displayname and
+	// operation specific StaticStruct().
+};
diff --git a/Engine/Source/Runtime/RigVM/Public/RigVMFunctions/RigVMDispatch_Print.h b/Engine/Source/Runtime/RigVM/Public/RigVMFunctions/RigVMDispatch_Print.h
new file mode 100644
--- /dev/null
+++ b/Engine/Source/Runtime/RigVM/Public/RigVMFunctions/RigVMDispatch_Print.h
@@ -0,0 +1,34 @@
+// Copyright Epic Games, Inc. All Rights Reserved.
+
+#pragma once
+
+#include "RigVMCore/RigVMDispatchFactory.h"
+#include "RigVMDispatch_Print.generated.h"
+
+/*
+ * Prints any value to the log
+ */
+USTRUCT(meta=(DisplayName = "Print", NodeColor = "0.8, 0, 0.2, 1"))
+struct FRigVMDispatch_Print : public FRigVMDispatchFactory
+{
+	GENERATED_BODY()
+
+public:
+
+	virtual TArray<FRigVMTemplateArgument> GetArguments() const override;
+	virtual TArray<FRigVMExecuteArgument> GetExecuteArguments_Impl() const override;
+	virtual FRigVMTemplateTypeMap OnNewArgumentType(const FName& InArgumentName, TRigVMTypeIndex InTypeIndex) const override;
+
+#if WITH_EDITOR
+	virtual FString GetArgumentDefaultValue(const FName& InArgumentName, TRigVMTypeIndex InTypeIndex) const override;
+	virtual FString GetArgumentMetaData(const FName& InArgumentName, const FName& InMetaDataKey) const override;
+#endif
+
+protected:
+
+	virtual FRigVMFunctionPtr GetDispatchFunctionImpl(const FRigVMTemplateTypeMap& InTypes) const override
+	{
+		return &FRigVMDispatch_Print::Execute;
+	}
+	static void Execute(FRigVMExtendedExecuteContext& InContext, FRigVMMemoryHandleArray Handles);
+};
diff --git a/Engine/Source/Runtime/RigVM/Public/RigVMFunctions/RigVMFunction_String.h b/Engine/Source/Runtime/RigVM/Public/RigVMFunctions/RigVMFunction_String.h
new file mode 100644
--- /dev/null
+++ b/Engine/Source/Runtime/RigVM/Public/RigVMFunctions/RigVMFunction_String.h
@@ -0,0 +1,539 @@
+// Copyright Epic Games, Inc. All Rights Reserved.
+
+#pragma once
+
+#include "RigVMCore/RigVMStruct.h"
+#include "RigVMCore/RigVMDispatchFactory.h"
+#include "RigVMFunction_String.generated.h"
+
+USTRUCT(meta=(Abstract, Category="Core|String", NodeColor = "0.462745, 1,0, 0.329412"))
+struct RIGVM_API FRigVMFunction_StringBase : public FRigVMStruct
+{
+	GENERATED_BODY()
+};
+
+/**
+ * Concatenates two strings together to make a new string
+ */
+USTRUCT(meta = (DisplayName = "Concat", TemplateName = "Concat", Keywords = "Add,+,Combine,Merge,Append"))
+struct RIGVM_API FRigVMFunction_StringConcat : public FRigVMFunction_StringBase
+{
+	GENERATED_BODY()
+
+	FRigVMFunction_StringConcat()
+	{
+		A = B = Result = FString();
+	}
+
+	RIGVM_METHOD()
+	virtual void Execute();
+
+	UPROPERTY(meta=(Input, Aggregate))
+	FString A;
+
+	UPROPERTY(meta=(Input, Aggregate))
+	FString B;
+
+	UPROPERTY(meta=(Output, Aggregate))
+	FString Result;
+};
+
+/**
+ * Returns the left or right most characters from the string chopping the given number of characters from the start or the end
+ */
+USTRUCT(meta = (DisplayName = "Chop", TemplateName = "Chop", Keywords = "Truncate,-,Remove,Subtract,Split"))
+struct RIGVM_API FRigVMFunction_StringTruncate : public FRigVMFunction_StringBase
+{
+	GENERATED_BODY()
+
+	FRigVMFunction_StringTruncate()
+	{
+		Name = FString();
+		Count = 1;
+		FromEnd = true;
+	}
+
+	RIGVM_METHOD()
+	virtual void Execute();
+
+	UPROPERTY(meta=(Input))
+	FString Name;
+
+	// Number of characters to remove from left or right
+	UPROPERTY(meta=(Input))
+	int32 Count;
+
+	// if set to true the characters will be removed from the end
+	UPROPERTY(meta=(Input))
+	bool FromEnd;
+
+	// the part of the string without the chopped characters
+	UPROPERTY(meta=(Output))
+	FString Remainder;
+
+	// the part of the name that has been chopped off
+	UPROPERTY(meta = (Output))
+	FString Chopped;
+};
+
+/**
+ * Replace all occurrences of a substring in this string
+ */
+USTRUCT(meta = (DisplayName = "Replace", TemplateName = "Replace", Keywords = "Search,Emplace,Find"))
+struct RIGVM_API FRigVMFunction_StringReplace : public FRigVMFunction_StringBase
+{
+	GENERATED_BODY()
+
+	FRigVMFunction_StringReplace()
+	{
+		Name = Old = New = FString();
+	}
+
+	RIGVM_METHOD()
+	virtual void Execute();
+
+	UPROPERTY(meta = (Input))
+	FString Name;
+
+	UPROPERTY(meta = (Input))
+	FString Old;
+
+	UPROPERTY(meta = (Input))
+	FString New;
+
+	UPROPERTY(meta = (Output))
+	FString Result;
+};
+
+/**
+ * Tests whether this string ends with given string
+ */
+USTRUCT(meta = (DisplayName = "Ends With", TemplateName = "EndsWith", Keywords = "Right"))
+struct RIGVM_API FRigVMFunction_StringEndsWith : public FRigVMFunction_StringBase
+{
+	GENERATED_BODY()
+
+	FRigVMFunction_StringEndsWith()
+	{
+		Name = Ending = FString();
+		Result = false;
+	}
+
+	RIGVM_METHOD()
+	virtual void Execute();
+
+	UPROPERTY(meta = (Input))
+	FString Name;
+
+	UPROPERTY(meta = (Input))
+	FString Ending;
+
+	UPROPERTY(meta = (Output))
+	bool Result;
+};
+
+/**
+ * Tests whether this string starts with given string
+ */
+USTRUCT(meta = (DisplayName = "Starts With", TemplateName = "StartsWith", Keywords = "Left"))
+struct RIGVM_API FRigVMFunction_StringStartsWith : public FRigVMFunction_StringBase
+{
+	GENERATED_BODY()
+
+	FRigVMFunction_StringStartsWith()
+	{
+		Name = Start = FString();
+		Result = false;
+	}
+
+	RIGVM_METHOD()
+	virtual void Execute();
+
+	UPROPERTY(meta = (Input))
+	FString Name;
+
+	UPROPERTY(meta = (Input))
+	FString Start;
+
+	UPROPERTY(meta = (Output))
+	bool Result;
+};
+
+/**
+ * Returns true or false if a given name exists in another given name
+ */
+USTRUCT(meta = (DisplayName = "Contains", TemplateName = "Contains", Keywords = "Contains,Find,Has,Search"))
+struct RIGVM_API FRigVMFunction_StringContains : public FRigVMFunction_StringBase
+{
+	GENERATED_BODY()
+
+		FRigVMFunction_StringContains()
+	{
+		Name = Search = FString();
+		Result = false;
+	}
+
+	RIGVM_METHOD()
+	virtual void Execute();
+
+	UPROPERTY(meta = (Input))
+	FString Name;
+
+	UPROPERTY(meta = (Input))
+	FString Search;
+
+	UPROPERTY(meta = (Output))
+	bool Result;
+};
+
+/**
+ * Returns the length of a string 
+ */
+USTRUCT(meta = (DisplayName = "Length"))
+struct RIGVM_API FRigVMFunction_StringLength : public FRigVMFunction_StringBase
+{
+	GENERATED_BODY()
+
+	FRigVMFunction_StringLength()
+	{
+		Value = FString();
+		Length = 0;
+	}
+
+	RIGVM_METHOD()
+	virtual void Execute();
+
+	UPROPERTY(meta = (Input))
+	FString Value;
+
+	UPROPERTY(meta = (Output))
+	int32 Length;
+};
+
+/**
+ * Trims the whitespace from a string (start and end)
+ */
+USTRUCT(meta = (DisplayName = "Trim Whitespace", Keywords = "Space,WhiteSpace,Remove,Truncate"))
+struct RIGVM_API FRigVMFunction_StringTrimWhitespace : public FRigVMFunction_StringBase
+{
+	GENERATED_BODY()
+
+	FRigVMFunction_StringTrimWhitespace()
+	{
+		Value = Result = FString();
+	}
+
+	RIGVM_METHOD()
+	virtual void Execute();
+
+	UPROPERTY(meta = (Input))
+	FString Value;
+
+	UPROPERTY(meta = (Output))
+	FString Result;
+};
+
+/**
+ * Converts the string to upper case
+ */
+USTRUCT(meta = (DisplayName = "To Uppercase"))
+struct RIGVM_API FRigVMFunction_StringToUppercase : public FRigVMFunction_StringBase
+{
+	GENERATED_BODY()
+
+	FRigVMFunction_StringToUppercase()
+	{
+		Value = Result = FString();
+	}
+
+	RIGVM_METHOD()
+	virtual void Execute();
+
+	UPROPERTY(meta = (Input))
+	FString Value;
+
+	UPROPERTY(meta = (Output))
+	FString Result;
+};
+
+/**
+ * Converts the string to lower case
+ */
+USTRUCT(meta = (DisplayName = "To Lowercase"))
+struct RIGVM_API FRigVMFunction_StringToLowercase : public FRigVMFunction_StringBase
+{
+	GENERATED_BODY()
+
+	FRigVMFunction_StringToLowercase()
+	{
+		Value = Result = FString();
+	}
+
+	RIGVM_METHOD()
+	virtual void Execute();
+
+	UPROPERTY(meta = (Input))
+	FString Value;
+
+	UPROPERTY(meta = (Output))
+	FString Result;
+};
+
+/**
+ * Returns the reverse of the input string
+ */
+USTRUCT(meta = (DisplayName = "Reverse"))
+struct RIGVM_API FRigVMFunction_StringReverse : public FRigVMFunction_StringBase
+{
+	GENERATED_BODY()
+
+	FRigVMFunction_StringReverse()
+	{
+		Value = Reverse = FString();
+	}
+
+	RIGVM_METHOD()
+	virtual void Execute();
+
+	UPROPERTY(meta = (Input))
+	FString Value;
+
+	UPROPERTY(meta = (Output))
+	FString Reverse;
+};
+
+/**
+ * Returns the left most characters of a string
+ */
+USTRUCT(meta = (DisplayName = "Left", Keywords = "Start,Begin"))
+struct RIGVM_API FRigVMFunction_StringLeft : public FRigVMFunction_StringBase
+{
+	GENERATED_BODY()
+
+	FRigVMFunction_StringLeft()
+	{
+		Value = Result = FString();
+		Count = 1;
+	}
+
+	RIGVM_METHOD()
+	virtual void Execute();
+
+	UPROPERTY(meta = (Input))
+	FString Value;
+
+	UPROPERTY(meta = (Input))
+	int32 Count;
+
+	UPROPERTY(meta = (Output))
+	FString Result;
+};
+
+/**
+ * Returns the right most characters of a string
+ */
+USTRUCT(meta = (DisplayName = "Right", Keywords = "End"))
+struct RIGVM_API FRigVMFunction_StringRight : public FRigVMFunction_StringBase
+{
+	GENERATED_BODY()
+
+	FRigVMFunction_StringRight()
+	{
+		Value = Result = FString();
+		Count = 1;
+	}
+
+	RIGVM_METHOD()
+	virtual void Execute();
+
+	UPROPERTY(meta = (Input))
+	FString Value;
+
+	UPROPERTY(meta = (Input))
+	int32 Count;
+
+	UPROPERTY(meta = (Output))
+	FString Result;
+};
+
+/**
+ * Returns the middle section of a string
+ */
+USTRUCT(meta = (DisplayName = "Middle", Keywords = "Within,Center"))
+struct RIGVM_API FRigVMFunction_StringMiddle : public FRigVMFunction_StringBase
+{
+	GENERATED_BODY()
+
+	FRigVMFunction_StringMiddle()
+	{
+		Value = Result = FString();
+		Start = 0;
+		Count = -1;
+	}
+
+	RIGVM_METHOD()
+	virtual void Execute();
+
+	UPROPERTY(meta = (Input))
+	FString Value;
+
+	// the index of the first character
+	UPROPERTY(meta = (Input))
+	int32 Start;
+
+	// if count is set to -1 all character from Start will be returned 
+	UPROPERTY(meta = (Input))
+	int32 Count;
+
+	UPROPERTY(meta = (Output))
+	FString Result;
+};
+
+/**
+ * Finds a string within another string
+ */
+USTRUCT(meta = (DisplayName = "Find", Keywords = "IndexOf"))
+struct RIGVM_API FRigVMFunction_StringFind : public FRigVMFunction_StringBase
+{
+	GENERATED_BODY()
+
+	FRigVMFunction_StringFind()
+	{
+		Value = Search = FString();
+		Found = false;
+		Index = INDEX_NONE;
+	}
+
+	RIGVM_METHOD()
+	virtual void Execute();
+
+	UPROPERTY(meta = (Input))
+	FString Value;
+
+	UPROPERTY(meta = (Input))
+	FString Search;
+
+	UPROPERTY(meta = (Output))
+    bool Found;
+
+	UPROPERTY(meta = (Output))
+	int32 Index;
+};
+
+/**
+ * Splits a string into multiple sections given a separator
+ */
+USTRUCT(meta = (DisplayName = "Split"))
+struct RIGVM_API FRigVMFunction_StringSplit : public FRigVMFunction_StringBase
+{
+	GENERATED_BODY()
+
+	FRigVMFunction_StringSplit()
+	{
+		Value = Separator = FString();
+	}
+
+	RIGVM_METHOD()
+	virtual void Execute();
+
+	UPROPERTY(meta = (Input))
+	FString Value;
+
+	UPROPERTY(meta = (Input))
+	FString Separator;
+
+	UPROPERTY(meta = (Output))
+	TArray<FString> Result;
+};
+
+/**
+ * Joins a string into multiple sections given a separator
+ */
+USTRUCT(meta = (DisplayName = "Join", Keywords = "Combine"))
+struct RIGVM_API FRigVMFunction_StringJoin : public FRigVMFunction_StringBase
+{
+	GENERATED_BODY()
+
+	FRigVMFunction_StringJoin()
+	{
+		Result = Separator = FString();
+	}
+
+	RIGVM_METHOD()
+	virtual void Execute();
+
+	UPROPERTY(meta = (Input))
+	TArray<FString> Values;
+
+	UPROPERTY(meta = (Input))
+	FString Separator;
+
+	UPROPERTY(meta = (Output))
+	FString Result;
+};
+
+/**
+ * Converts an integer number to a string with padding
+ */
+USTRUCT(meta = (DisplayName = "Pad Integer", Keywords = "FromInt,Number,LeadingZeroes"))
+struct RIGVM_API FRigVMFunction_StringPadInteger : public FRigVMFunction_StringBase
+{
+	GENERATED_BODY()
+
+	FRigVMFunction_StringPadInteger()
+	{
+		Value = 0;
+		Digits = 4;
+		Result = FString();
+	}
+
+	RIGVM_METHOD()
+	virtual void Execute();
+
+	UPROPERTY(meta = (Input))
+	int32 Value;
+
+	UPROPERTY(meta = (Input))
+	int32 Digits;
+
+	UPROPERTY(meta = (Output))
+	FString Result;
+};
+
+/*
+ * Converts any value to string
+ */
+USTRUCT(meta=(DisplayName = "To String", NodeColor = "0.462745, 1,0, 0.329412"))
+struct FRigDispatch_ToString : public FRigVMDispatchFactory
+{
+	GENERATED_BODY()
+
+public:
+
+	virtual TArray<FRigVMTemplateArgument> GetArguments() const override;
+	virtual FRigVMTemplateTypeMap OnNewArgumentType(const FName& InArgumentName, TRigVMTypeIndex InTypeIndex) const override;
+
+protected:
+
+	virtual FRigVMFunctionPtr GetDispatchFunctionImpl(const FRigVMTemplateTypeMap& InTypes) const override;
+	static void Execute(FRigVMExtendedExecuteContext& InContext, FRigVMMemoryHandleArray Handles);
+};
+
+/*
+ * Converts a string into any value
+ */
+USTRUCT(meta=(DisplayName = "From String", NodeColor = "0.462745, 1,0, 0.329412"))
+struct FRigDispatch_FromString : public FRigVMDispatchFactory
+{
+	GENERATED_BODY()
+
+public:
+
+	virtual TArray<FRigVMTemplateArgument> GetArguments() const override;
+	virtual FRigVMTemplateTypeMap OnNewArgumentType(const FName& InArgumentName, TRigVMTypeIndex InTypeIndex) const override;
+
+protected:
+
+	virtual FRigVMFunctionPtr GetDispatchFunctionImpl(const FRigVMTemplateTypeMap& InTypes) const override;
+	static void Execute(FRigVMExtendedExecuteContext& InContext, FRigVMMemoryHandleArray Handles);
+};
\ No newline at end of file
