RigVMN: lazy compute - core side

diff --git a/Engine/Plugins/Animation/ControlRig/Source/ControlRig/Private/ControlRig.cpp b/Engine/Plugins/Animation/ControlRig/Source/ControlRig/Private/ControlRig.cpp
--- a/Engine/Plugins/Animation/ControlRig/Source/ControlRig/Private/ControlRig.cpp
+++ b/Engine/Plugins/Animation/ControlRig/Source/ControlRig/Private/ControlRig.cpp
@@ -1617,7 +1617,7 @@ bool UControlRig::ExecuteUnits(FRigUnitContext& InOutContext, const FName& InEve
 #endif
 			FRigHierarchyExecuteContextBracket HierarchyContextGuard(Hierarchy, &VM->GetContext());
 
-			bSuccess = VM->Execute(LocalMemory, AdditionalArguments, InEventName);
+			bSuccess = VM->Execute(LocalMemory, AdditionalArguments, InEventName) != ERigVMExecuteResult::Failed;
 		}
 
 #if UE_CONTROLRIG_PROFILE_EXECUTE_UNITS_NUM
diff --git a/Engine/Plugins/Animation/ControlRig/Source/ControlRig/Public/ControlRigDefines.h b/Engine/Plugins/Animation/ControlRig/Source/ControlRig/Public/ControlRigDefines.h
--- a/Engine/Plugins/Animation/ControlRig/Source/ControlRig/Public/ControlRigDefines.h
+++ b/Engine/Plugins/Animation/ControlRig/Source/ControlRig/Public/ControlRigDefines.h
@@ -23,6 +23,14 @@ struct FControlRigExecuteContext : public FRigVMExecuteContext
 		, Hierarchy(nullptr)
 	{
 	}
+
+	virtual void Copy(const FRigVMExecuteContext* InOtherContext) override
+	{
+		Super::Copy(InOtherContext);
+
+		const FControlRigExecuteContext* OtherContext = (const FControlRigExecuteContext*)InOtherContext; 
+		Hierarchy = OtherContext->Hierarchy;
+	}
 		
 	URigHierarchy* Hierarchy;
 };
diff --git a/Engine/Plugins/Experimental/DataInterfaceGraph/Source/DataInterfaceGraph/Internal/DataInterfaceExecuteContext.h b/Engine/Plugins/Experimental/DataInterfaceGraph/Source/DataInterfaceGraph/Internal/DataInterfaceExecuteContext.h
--- a/Engine/Plugins/Experimental/DataInterfaceGraph/Source/DataInterfaceGraph/Internal/DataInterfaceExecuteContext.h
+++ b/Engine/Plugins/Experimental/DataInterfaceGraph/Source/DataInterfaceGraph/Internal/DataInterfaceExecuteContext.h
@@ -37,6 +37,17 @@ struct FDataInterfaceExecuteContext : public FRigVMExecuteContext
 		check(Interface);
 		return Interface;
 	}
+	
+	virtual void Copy(const FRigVMExecuteContext* InOtherContext) override
+	{
+		Super::Copy(InOtherContext);
+
+		const FDataInterfaceExecuteContext* OtherContext = (const FDataInterfaceExecuteContext*)InOtherContext; 
+		DataInterfaceContext = OtherContext->DataInterfaceContext;
+		Interface = OtherContext->Interface; 
+		ResultPtr = OtherContext->ResultPtr; 
+	}
+
 
 private:
 	const UE::DataInterface::FContext* DataInterfaceContext = nullptr;
diff --git a/Engine/Plugins/Experimental/DataInterfaceGraph/Source/DataInterfaceGraph/Private/DataInterfaceGraph.cpp b/Engine/Plugins/Experimental/DataInterfaceGraph/Source/DataInterfaceGraph/Private/DataInterfaceGraph.cpp
--- a/Engine/Plugins/Experimental/DataInterfaceGraph/Source/DataInterfaceGraph/Private/DataInterfaceGraph.cpp
+++ b/Engine/Plugins/Experimental/DataInterfaceGraph/Source/DataInterfaceGraph/Private/DataInterfaceGraph.cpp
@@ -33,7 +33,8 @@ bool UDataInterfaceGraph::GetDataImpl(const UE::DataInterface::FContext& Context
 		RigUnitContext.State = EControlRigState::Update;
 		FDataInterfaceUnitContext DataInterfaceUnitContext(this, Context, bResult);
 		void* AdditionalArguments[2] = { &RigUnitContext, &DataInterfaceUnitContext };
-		bResult &= RigVM->Execute(TArray<URigVMMemoryStorage*>(), AdditionalArguments, FRigUnit_DataInterfaceBeginExecution::EventName);
+		bResult &=
+			(RigVM->Execute(TArray<URigVMMemoryStorage*>(), AdditionalArguments, FRigUnit_DataInterfaceBeginExecution::EventName) != ERigVMExecuteResult::Failed);
 	}
 	
 	return bResult;
diff --git a/Engine/Source/Developer/RigVMDeveloper/Private/RigVMCompiler/RigVMCodeGenerator.cpp b/Engine/Source/Developer/RigVMDeveloper/Private/RigVMCompiler/RigVMCodeGenerator.cpp
--- a/Engine/Source/Developer/RigVMDeveloper/Private/RigVMCompiler/RigVMCodeGenerator.cpp
+++ b/Engine/Source/Developer/RigVMDeveloper/Private/RigVMCompiler/RigVMCodeGenerator.cpp
@@ -88,8 +88,8 @@ static constexpr TCHAR RigVM_ArrayUnionOpFormat[] = TEXT("\tArrayOp_Union{0}({1}
 static constexpr TCHAR RigVM_ArrayDifferenceOpFormat[] = TEXT("\tArrayOp_Difference{0}({1}, {2});");
 static constexpr TCHAR RigVM_ArrayIntersectionOpFormat[] = TEXT("\tArrayOp_Intersection{0}({1}, {2});");
 static constexpr TCHAR RigVM_ArrayReverseOpFormat[] = TEXT("\tArrayOp_Reverse{0}({1});");
-static constexpr TCHAR RigVM_ReturnFalseFormat[] = TEXT("\treturn false;");
-static constexpr TCHAR RigVM_ReturnTrueFormat[] = TEXT("\treturn true;");
+static constexpr TCHAR RigVM_ReturnFailedFormat[] = TEXT("\treturn ERigVMExecuteResult::Failed;");
+static constexpr TCHAR RigVM_ReturnSucceededFormat[] = TEXT("\treturn ERigVMExecuteResult::Succeeded;");
 static constexpr TCHAR RigVM_CopyrightFormat[] = TEXT("// Copyright Epic Games, Inc. All Rights Reserved.");
 static constexpr TCHAR RigVM_AutoGeneratedFormat[] = TEXT("// THIS FILE HAS BEEN AUTO-GENERATED. PLEASE DO NOT MANUALLY EDIT THIS FILE FURTHER.");
 static constexpr TCHAR RigVM_PragmaOnceFormat[] = TEXT("#pragma once");
@@ -99,16 +99,16 @@ static constexpr TCHAR RigVM_ProtectedFormat[] = TEXT("protected:");
 static constexpr TCHAR RigVM_GetVMHashFormat[] = TEXT("\tvirtual uint32 GetVMHash() const override { return {0}; }");
 static constexpr TCHAR RigVM_GetEntryNamesFormat[] = TEXT("\tvirtual const TArray<FName>& GetEntryNames() const override\r\n\t{\r\n\t\tstatic const TArray<FName> StaticEntryNames = { {0} };\r\n\t\treturn StaticEntryNames;\r\n\t}");
 static constexpr TCHAR RigVM_DeclareUpdateExternalVariablesFormat[] = TEXT("\tvirtual void UpdateExternalVariables() override;");
-static constexpr TCHAR RigVM_DeclareInvokeEntryByNameFormat[] = TEXT("\tbool InvokeEntryByName(const FName& InEntryName{0});");
-static constexpr TCHAR RigVM_DeclareExecuteFormat[] = TEXT("\tvirtual bool Execute(TArrayView<URigVMMemoryStorage*> Memory, TArrayView<void*> AdditionalArguments, const FName& InEntryName) override;");
+static constexpr TCHAR RigVM_DeclareInvokeEntryByNameFormat[] = TEXT("\tERigVMExecuteResult InvokeEntryByName(const FName& InEntryName{0});");
+static constexpr TCHAR RigVM_DeclareExecuteFormat[] = TEXT("\tvirtual ERigVMExecuteResult Execute(TArrayView<URigVMMemoryStorage*> Memory, TArrayView<void*> AdditionalArguments, const FName& InEntryName) override;");
 static constexpr TCHAR RigVM_DefineUpdateExternalVariablesFormat[] = TEXT("void U{0}::UpdateExternalVariables()\r\n{");
-static constexpr TCHAR RigVM_DefineInvokeEntryByNameFormat[] = TEXT("bool U{0}::InvokeEntryByName(const FName& InEntryName{1})\r\n{");
-static constexpr TCHAR RigVM_DefineExecuteFormat[] = TEXT("bool U{0}::Execute(TArrayView<URigVMMemoryStorage*> Memory, TArrayView<void*> AdditionalArguments, const FName& InEntryName)\r\n{");
-static constexpr TCHAR RigVM_DeclareExecuteEntryFormat[] = TEXT("\tbool ExecuteEntry_{0}({1});");
-static constexpr TCHAR RigVM_DefineExecuteEntryFormat[] = TEXT("bool U{0}::ExecuteEntry_{1}({2})\r\n{");
-static constexpr TCHAR RigVM_DeclareExecuteGroupFormat[] = TEXT("\tbool ExecuteGroup_{0}_{1}({2});");
-static constexpr TCHAR RigVM_DefineExecuteGroupFormat[] = TEXT("bool U{0}::ExecuteGroup_{1}_{2}({3})\r\n{");
-static constexpr TCHAR RigVM_InvokeExecuteGroupFormat[] = TEXT("\tif(!ExecuteGroup_{0}_{1}({2})) return false;");
+static constexpr TCHAR RigVM_DefineInvokeEntryByNameFormat[] = TEXT("ERigVMExecuteResult U{0}::InvokeEntryByName(const FName& InEntryName{1})\r\n{");
+static constexpr TCHAR RigVM_DefineExecuteFormat[] = TEXT("ERigVMExecuteResult U{0}::Execute(TArrayView<URigVMMemoryStorage*> Memory, TArrayView<void*> AdditionalArguments, const FName& InEntryName)\r\n{");
+static constexpr TCHAR RigVM_DeclareExecuteEntryFormat[] = TEXT("\tERigVMExecuteResult ExecuteEntry_{0}({1});");
+static constexpr TCHAR RigVM_DefineExecuteEntryFormat[] = TEXT("ERigVMExecuteResult U{0}::ExecuteEntry_{1}({2})\r\n{");
+static constexpr TCHAR RigVM_DeclareExecuteGroupFormat[] = TEXT("\tERigVMExecuteResult ExecuteGroup_{0}_{1}({2});");
+static constexpr TCHAR RigVM_DefineExecuteGroupFormat[] = TEXT("ERigVMExecuteResult U{0}::ExecuteGroup_{1}_{2}({3})\r\n{");
+static constexpr TCHAR RigVM_InvokeExecuteGroupFormat[] = TEXT("\tif(ExecuteGroup_{0}_{1}({2}) != ERigVMExecuteResult::Succeeded) return ERigVMExecuteResult::Failed;");
 static constexpr TCHAR RigVM_RigVMCoreIncludeFormat[] = TEXT("RigVMCore/RigVMCore.h");
 static constexpr TCHAR RigVM_RigVMCoreLibraryFormat[] = TEXT("RigVM");
 static constexpr TCHAR RigVM_JoinFilePathFormat[] = TEXT("{0}/{1}");
@@ -116,7 +116,7 @@ static constexpr TCHAR RigVM_GetOperandSliceFormat[] = TEXT("GetOperandSlice<{0}
 static constexpr TCHAR RigVM_ExternalVariableFormat[] = TEXT("(*External_{0})");
 static constexpr TCHAR RigVM_JoinSegmentPathFormat[] = TEXT("{0}.{1}");
 static constexpr TCHAR RigVM_GetArrayElementSafeFormat[] = TEXT("GetArrayElementSafe<{0}>({1}, {2})");
-static constexpr TCHAR RigVM_InvokeEntryOpFormat[] = TEXT("\tif(!InvokeEntryByName(EntryName_{0})) return false;");
+static constexpr TCHAR RigVM_InvokeEntryOpFormat[] = TEXT("\tif(InvokeEntryByName(EntryName_{0}) != ERigVMExecuteResult::Succeeded) return ERigVMExecuteResult::Failed;");
 
 FString FRigVMCodeGenerator::DumpIncludes(bool bLog)
 {
@@ -556,7 +556,7 @@ FString FRigVMCodeGenerator::DumpInstructions(int32 InOperationGroup, bool bLog)
 
 	if(Group.Entry.IsEmpty() && OperationGroups.Num() > 0)
 	{
-		Lines.Add(RigVM_ReturnFalseFormat);
+		Lines.Add(RigVM_ReturnFailedFormat);
 		return DumpLines(Lines, bLog);
 	}
 
@@ -974,7 +974,7 @@ FString FRigVMCodeGenerator::DumpInstructions(int32 InOperationGroup, bool bLog)
 	{
 		Lines.Add(RigVM_ExecuteReachedExitFormat);
 	}
-	Lines.Add(RigVM_ReturnTrueFormat);
+	Lines.Add(RigVM_ReturnSucceededFormat);
 
 	return DumpLines(Lines, bLog);
 }
diff --git a/Engine/Source/Runtime/Core/Public/UObject/UE5MainStreamObjectVersion.h b/Engine/Source/Runtime/Core/Public/UObject/UE5MainStreamObjectVersion.h
--- a/Engine/Source/Runtime/Core/Public/UObject/UE5MainStreamObjectVersion.h
+++ b/Engine/Source/Runtime/Core/Public/UObject/UE5MainStreamObjectVersion.h
@@ -295,6 +295,9 @@ struct CORE_API FUE5MainStreamObjectVersion
 		// Added TransformScaleMethod pin to SpawnActorFromClass node
 		SpawnActorFromClassTransformScaleMethod,
 
+		// Added support for the RigVM to run branches lazily
+		RigVMLazyEvaluation,
+
 		// -----<new versions can be added above this line>-------------------------------------------------
 		VersionPlusOne,
 		LatestVersion = VersionPlusOne - 1
diff --git a/Engine/Source/Runtime/RigVM/Private/RigVMCore/RigVM.cpp b/Engine/Source/Runtime/RigVM/Private/RigVMCore/RigVM.cpp
--- a/Engine/Source/Runtime/RigVM/Private/RigVMCore/RigVM.cpp
+++ b/Engine/Source/Runtime/RigVM/Private/RigVMCore/RigVM.cpp
@@ -7,6 +7,7 @@
 #include "UObject/UE5MainStreamObjectVersion.h"
 #include "HAL/PlatformTLS.h"
 #include "Async/ParallelFor.h"
+#include "GenericPlatform/GenericPlatformSurvey.h"
 #include "UObject/UObjectHash.h"
 #include "UObject/UObjectIterator.h"
 
@@ -85,6 +86,13 @@ URigVM::URigVM()
 	, FirstEntryEventInQueue(NAME_None)
 #endif
 	, ExecutingThreadId(INDEX_NONE)
+	, CurrentExecuteResult(ERigVMExecuteResult::Failed)
+	, CurrentEntryName(NAME_None)
+	, bCurrentlyRunningRootEntry(false)
+#if WITH_EDITOR
+	, StartCycles(0)
+	, OverallCycles(0)
+#endif
 	, DeferredVMToCopy(nullptr)
 {
 }
@@ -949,7 +957,7 @@ bool URigVM::ResumeExecution()
 bool URigVM::ResumeExecution(TArrayView<URigVMMemoryStorage*> Memory, TArrayView<void*> AdditionalArguments, const FName& InEntryName)
 {
 	ResumeExecution();
-	return Execute(Memory, AdditionalArguments, InEntryName);
+	return Execute(Memory, AdditionalArguments, InEntryName) != ERigVMExecuteResult::Failed;
 }
 
 #endif
@@ -1041,6 +1049,7 @@ void URigVM::InvalidateCachedMemory()
 	FirstHandleForInstruction.Reset();
 	CachedMemoryHandles.Reset();
 	ExternalPropertyPaths.Reset();
+	LazyBranches.Reset();
 }
 
 void URigVM::CopyDeferredVMIfRequired()
@@ -1102,10 +1111,15 @@ void URigVM::CacheMemoryHandlesIfRequired(TArrayView<URigVMMemoryStorage*> InMem
 
 	FRigVMByteCode& ByteCode = GetByteCode();
 
+	// force to update the map of branch infos once
+	(void)ByteCode.GetBranchInfo({0, 0});
+	LazyBranches.Reset();
+	LazyBranches.SetNumZeroed(ByteCode.BranchInfos.Num());
+
 	auto InstructionOpEval = [&](
 		int32 InstructionIndex,
 		int32 InHandleBaseIndex,
-		TFunctionRef<void(int32 InHandleIndex, const FRigVMOperand& InArg)> InOpFunc
+		TFunctionRef<void(int32 InHandleIndex, const TRigVMBranchInfoKey& InBranchInfoKey, const FRigVMOperand& InArg)> InOpFunc
 		) -> void
 	{
 		const ERigVMOpCode OpCode = Instructions[InstructionIndex].OpCode; 
@@ -1114,9 +1128,9 @@ void URigVM::CacheMemoryHandlesIfRequired(TArrayView<URigVMMemoryStorage*> InMem
 		{
 			FRigVMOperandArray Operands = ByteCode.GetOperandsForExecuteOp(Instructions[InstructionIndex]);
 
-			for (const FRigVMOperand& Arg : Operands)
+			for (int32 ArgIndex = 0; ArgIndex < Operands.Num(); ArgIndex++)
 			{
-				InOpFunc(InHandleBaseIndex++, Arg);
+				InOpFunc(InHandleBaseIndex++, {InstructionIndex, ArgIndex}, Operands[ArgIndex]);
 			}
 		}
 		else
@@ -1132,14 +1146,14 @@ void URigVM::CacheMemoryHandlesIfRequired(TArrayView<URigVMMemoryStorage*> InMem
 			case ERigVMOpCode::ArrayReverse:
 				{
 					const FRigVMUnaryOp& Op = ByteCode.GetOpAt<FRigVMUnaryOp>(Instructions[InstructionIndex]);
-					InOpFunc(InHandleBaseIndex, Op.Arg);
+					InOpFunc(InHandleBaseIndex, {InstructionIndex, 0}, Op.Arg);
 					break;
 				}
 			case ERigVMOpCode::Copy:
 				{
 					const FRigVMCopyOp& Op = ByteCode.GetOpAt<FRigVMCopyOp>(Instructions[InstructionIndex]);
-					InOpFunc(InHandleBaseIndex + 0, Op.Source);
-					InOpFunc(InHandleBaseIndex + 1, Op.Target);
+					InOpFunc(InHandleBaseIndex + 0, {InstructionIndex, 0}, Op.Source);
+					InOpFunc(InHandleBaseIndex + 1, {InstructionIndex, 1}, Op.Target);
 					break;
 				}
 			case ERigVMOpCode::Equals:
@@ -1147,11 +1161,11 @@ void URigVM::CacheMemoryHandlesIfRequired(TArrayView<URigVMMemoryStorage*> InMem
 				{
 					const FRigVMComparisonOp& Op = ByteCode.GetOpAt<FRigVMComparisonOp>(Instructions[InstructionIndex]);
 					FRigVMOperand Arg = Op.A;
-					InOpFunc(InHandleBaseIndex + 0, Arg);
+					InOpFunc(InHandleBaseIndex + 0, {InstructionIndex, 0}, Arg);
 					Arg = Op.B;
-					InOpFunc(InHandleBaseIndex + 1, Arg);
+					InOpFunc(InHandleBaseIndex + 1, {InstructionIndex, 1}, Arg);
 					Arg = Op.Result;
-					InOpFunc(InHandleBaseIndex + 2, Arg);
+					InOpFunc(InHandleBaseIndex + 2, {InstructionIndex, 2}, Arg);
 					break;
 				}
 			case ERigVMOpCode::JumpAbsolute:
@@ -1166,14 +1180,14 @@ void URigVM::CacheMemoryHandlesIfRequired(TArrayView<URigVMMemoryStorage*> InMem
 				{
 					const FRigVMJumpIfOp& Op = ByteCode.GetOpAt<FRigVMJumpIfOp>(Instructions[InstructionIndex]);
 					const FRigVMOperand& Arg = Op.Arg;
-					InOpFunc(InHandleBaseIndex, Arg);
+					InOpFunc(InHandleBaseIndex, {InstructionIndex, 0}, Arg);
 					break;
 				}
 			case ERigVMOpCode::ChangeType:
 				{
 					const FRigVMChangeTypeOp& Op = ByteCode.GetOpAt<FRigVMChangeTypeOp>(Instructions[InstructionIndex]);
 					const FRigVMOperand& Arg = Op.Arg;
-					InOpFunc(InHandleBaseIndex, Arg);
+					InOpFunc(InHandleBaseIndex, {InstructionIndex, 0}, Arg);
 					break;
 				}
 			case ERigVMOpCode::Exit:
@@ -1189,8 +1203,8 @@ void URigVM::CacheMemoryHandlesIfRequired(TArrayView<URigVMMemoryStorage*> InMem
 			case ERigVMOpCode::ArrayUnion:
 				{
 					const FRigVMBinaryOp& Op = ByteCode.GetOpAt<FRigVMBinaryOp>(Instructions[InstructionIndex]);
-					InOpFunc(InHandleBaseIndex + 0, Op.ArgA);
-					InOpFunc(InHandleBaseIndex + 1, Op.ArgB);
+					InOpFunc(InHandleBaseIndex + 0, {InstructionIndex, 0}, Op.ArgA);
+					InOpFunc(InHandleBaseIndex + 1, {InstructionIndex, 1}, Op.ArgB);
 					break;
 				}
 			case ERigVMOpCode::ArrayAdd:
@@ -1201,29 +1215,29 @@ void URigVM::CacheMemoryHandlesIfRequired(TArrayView<URigVMMemoryStorage*> InMem
 			case ERigVMOpCode::ArrayIntersection:
 				{
 					const FRigVMTernaryOp& Op = ByteCode.GetOpAt<FRigVMTernaryOp>(Instructions[InstructionIndex]);
-					InOpFunc(InHandleBaseIndex + 0, Op.ArgA);
-					InOpFunc(InHandleBaseIndex + 1, Op.ArgB);
-					InOpFunc(InHandleBaseIndex + 2, Op.ArgC);
+					InOpFunc(InHandleBaseIndex + 0, {InstructionIndex, 0}, Op.ArgA);
+					InOpFunc(InHandleBaseIndex + 1, {InstructionIndex, 1}, Op.ArgB);
+					InOpFunc(InHandleBaseIndex + 2, {InstructionIndex, 2}, Op.ArgC);
 					break;
 				}
 			case ERigVMOpCode::ArrayFind:
 				{
 					const FRigVMQuaternaryOp& Op = ByteCode.GetOpAt<FRigVMQuaternaryOp>(Instructions[InstructionIndex]);
-					InOpFunc(InHandleBaseIndex + 0, Op.ArgA);
-					InOpFunc(InHandleBaseIndex + 1, Op.ArgB);
-					InOpFunc(InHandleBaseIndex + 2, Op.ArgC);
-					InOpFunc(InHandleBaseIndex + 3, Op.ArgD);
+					InOpFunc(InHandleBaseIndex + 0, {InstructionIndex, 0}, Op.ArgA);
+					InOpFunc(InHandleBaseIndex + 1, {InstructionIndex, 1}, Op.ArgB);
+					InOpFunc(InHandleBaseIndex + 2, {InstructionIndex, 2}, Op.ArgC);
+					InOpFunc(InHandleBaseIndex + 3, {InstructionIndex, 3}, Op.ArgD);
 					break;
 				}
 			case ERigVMOpCode::ArrayIterator:
 				{
 					const FRigVMSenaryOp& Op = ByteCode.GetOpAt<FRigVMSenaryOp>(Instructions[InstructionIndex]);
-					InOpFunc(InHandleBaseIndex + 0, Op.ArgA);
-					InOpFunc(InHandleBaseIndex + 1, Op.ArgB);
-					InOpFunc(InHandleBaseIndex + 2, Op.ArgC);
-					InOpFunc(InHandleBaseIndex + 3, Op.ArgD);
-					InOpFunc(InHandleBaseIndex + 4, Op.ArgE);
-					InOpFunc(InHandleBaseIndex + 5, Op.ArgF);
+					InOpFunc(InHandleBaseIndex + 0, {InstructionIndex, 0}, Op.ArgA);
+					InOpFunc(InHandleBaseIndex + 1, {InstructionIndex, 1}, Op.ArgB);
+					InOpFunc(InHandleBaseIndex + 2, {InstructionIndex, 2}, Op.ArgC);
+					InOpFunc(InHandleBaseIndex + 3, {InstructionIndex, 3}, Op.ArgD);
+					InOpFunc(InHandleBaseIndex + 4, {InstructionIndex, 4}, Op.ArgE);
+					InOpFunc(InHandleBaseIndex + 5, {InstructionIndex, 5}, Op.ArgF);
 					break;
 				}
 			case ERigVMOpCode::EndBlock:
@@ -1252,7 +1266,11 @@ void URigVM::CacheMemoryHandlesIfRequired(TArrayView<URigVMMemoryStorage*> InMem
 	FirstHandleForInstruction.Add(0);
 	for (int32 InstructionIndex = 0; InstructionIndex < Instructions.Num(); InstructionIndex++)
 	{
-		InstructionOpEval(InstructionIndex, INDEX_NONE, [&HandleCount](int32, const FRigVMOperand& ){ HandleCount++; });
+		InstructionOpEval(InstructionIndex, INDEX_NONE,
+			[&HandleCount](int32, const TRigVMBranchInfoKey&, const FRigVMOperand& )
+			{
+				HandleCount++;
+			});
 		FirstHandleForInstruction.Add(HandleCount);
 	}
 	
@@ -1271,7 +1289,11 @@ void URigVM::CacheMemoryHandlesIfRequired(TArrayView<URigVMMemoryStorage*> InMem
 	ParallelFor(Instructions.Num(),
 		[&](int32 InstructionIndex)
 		{
-			InstructionOpEval(InstructionIndex, FirstHandleForInstruction[InstructionIndex], [&](int32 InHandleIndex, const FRigVMOperand& InOp){ CacheSingleMemoryHandle(InHandleIndex, InOp); });
+			InstructionOpEval(InstructionIndex, FirstHandleForInstruction[InstructionIndex],
+				[&](int32 InHandleIndex, const TRigVMBranchInfoKey& InBranchInfoKey, const FRigVMOperand& InOp)
+				{
+					CacheSingleMemoryHandle(InHandleIndex, InBranchInfoKey, InOp);
+				});
 		}
 	);
 
@@ -1693,12 +1715,19 @@ bool URigVM::Initialize(TArrayView<URigVMMemoryStorage*> Memory, TArrayView<void
 	return true;
 }
 
-bool URigVM::Execute(TArrayView<URigVMMemoryStorage*> Memory, TArrayView<void*> AdditionalArguments, const FName& InEntryName)
+ERigVMExecuteResult URigVM::Execute(TArrayView<URigVMMemoryStorage*> Memory, TArrayView<void*> AdditionalArguments, const FName& InEntryName)
 {
 	// if this the first entry being executed - get ready for execution
-	const bool bIsRootEntry = EntriesBeingExecuted.IsEmpty(); 
+	const bool bIsRootEntry = EntriesBeingExecuted.IsEmpty();
+
+	TGuardValue<FName> EntryNameGuard(CurrentEntryName, InEntryName);
+	TGuardValue<bool> RootEntryGuard(bCurrentlyRunningRootEntry, bIsRootEntry);
+	TGuardValue<TArrayView<void*>> AdditionalArgumentsGuard(CurrentAdditionalArguments, AdditionalArguments);
+
 	if(bIsRootEntry)
 	{
+		CurrentExecuteResult = ERigVMExecuteResult::Succeeded;
+		
 		if (ExecutingThreadId != INDEX_NONE)
 		{
 			ensureMsgf(ExecutingThreadId == FPlatformTLS::GetCurrentThreadId(), TEXT("RigVM::Execute from multiple threads (%d and %d)"), ExecutingThreadId, (int32)FPlatformTLS::GetCurrentThreadId());
@@ -1711,7 +1740,7 @@ bool URigVM::Execute(TArrayView<URigVMMemoryStorage*> Memory, TArrayView<void*>
 
 		if (Instructions.Num() == 0)
 		{
-			return true;
+			return CurrentExecuteResult = ERigVMExecuteResult::Failed;
 		}
 
 		// changes to the layout of memory array should be reflected in GetContainerIndex()
@@ -1723,8 +1752,9 @@ bool URigVM::Execute(TArrayView<URigVMMemoryStorage*> Memory, TArrayView<void*>
 		}
 	
 		CacheMemoryHandlesIfRequired(Memory);
+		CurrentMemory = Memory;
 	}
-	
+
 	FRigVMByteCode& ByteCode = GetByteCode();
 	TArray<FRigVMFunctionPtr>& Functions = GetFunctions();
 	TArray<const FRigVMDispatchFactory*>& Factories = GetFactories();
@@ -1774,7 +1804,7 @@ bool URigVM::Execute(TArrayView<URigVMMemoryStorage*> Memory, TArrayView<void*>
 		int32 EntryIndex = ByteCode.FindEntryIndex(InEntryName);
 		if (EntryIndex == INDEX_NONE)
 		{
-			return false;
+			return CurrentExecuteResult = ERigVMExecuteResult::Failed;
 		}
 		
 		SetInstructionIndex((uint16)ByteCode.GetEntry(EntryIndex).InstructionIndex);
@@ -1824,8 +1854,8 @@ bool URigVM::Execute(TArrayView<URigVMMemoryStorage*> Memory, TArrayView<void*>
 	}
 
 #if WITH_EDITOR
-	uint64 StartCycles = 0;
-	uint64 OverallCycles = 0;
+	StartCycles = 0;
+	OverallCycles = 0;
 	if(ContextPublicData.RuntimeSettings.bEnableProfiling)
 	{
 		StartCycles = FPlatformTime::Cycles64();
@@ -1847,19 +1877,68 @@ bool URigVM::Execute(TArrayView<URigVMMemoryStorage*> Memory, TArrayView<void*>
 	
 #endif
 
+	CurrentExecuteResult = ExecuteInstructions(ContextPublicData.InstructionIndex, Instructions.Num() - 1);
+
+#if WITH_EDITOR
+	if(bIsRootEntry)
+	{
+		CurrentMemory = TArrayView<URigVMMemoryStorage*>();
+		
+		if (HaltedAtBreakpoint.IsValid() && CurrentExecuteResult != ERigVMExecuteResult::Halted)
+		{
+			DebugInfo->SetCurrentActiveBreakpoint(FRigVMBreakpoint());
+			HaltedAtBreakpoint.Reset();
+			ExecutionHalted().Broadcast(INDEX_NONE, nullptr, InEntryName);
+		}
+	}
+#endif
+
+	return CurrentExecuteResult;
+}
+
+ERigVMExecuteResult URigVM::ExecuteInstructions(int32 InFirstInstruction, int32 InLastInstruction)
+{
+	// make we are already executing this VM
+	check(!CurrentMemory.IsEmpty());
+	
+	FRigVMExecuteContext& ContextPublicData = Context.GetPublicData<>();
+	TGuardValue<uint16> InstructionIndexGuard(ContextPublicData.InstructionIndex, (uint16)InFirstInstruction);
+	
+	FRigVMByteCode& ByteCode = GetByteCode();
+	TArray<FRigVMFunctionPtr>& Functions = GetFunctions();
+	TArray<const FRigVMDispatchFactory*>& Factories = GetFactories();
+
+#if WITH_EDITOR
+	TArray<FName>& FunctionNames = GetFunctionNames();
+#endif
+	
 	while (Instructions.IsValidIndex(ContextPublicData.InstructionIndex))
 	{
 #if WITH_EDITOR
-		if (ShouldHaltAtInstruction(InEntryName, ContextPublicData.InstructionIndex))
+		if (ShouldHaltAtInstruction(CurrentEntryName, ContextPublicData.InstructionIndex))
 		{
 #if UE_RIGVM_DEBUG_EXECUTION
 			ContextPublicData.Log(EMessageSeverity::Info, DebugMemoryString);					
 #endif
 			// we'll recursively exit all invoked
 			// entries here.
-			return true;
+			return CurrentExecuteResult = ERigVMExecuteResult::Halted;
+		}
+
+		if(CurrentExecuteResult == ERigVMExecuteResult::Halted)
+		{
+			return CurrentExecuteResult;
+		}
+
+#endif
+		
+		if(ContextPublicData.InstructionIndex > (uint16)InLastInstruction)
+		{
+			return CurrentExecuteResult = ERigVMExecuteResult::Succeeded;
 		}
 
+#if WITH_EDITOR
+
 		const int32 CurrentInstructionIndex = ContextPublicData.InstructionIndex;
 		InstructionVisitedDuringLastRun[ContextPublicData.InstructionIndex]++;
 		InstructionVisitOrder.Add(ContextPublicData.InstructionIndex);
@@ -2142,25 +2221,25 @@ bool URigVM::Execute(TArrayView<URigVMMemoryStorage*> Memory, TArrayView<void*>
 			}
 			case ERigVMOpCode::Exit:
 			{
-				if(bIsRootEntry)
+				if(bCurrentlyRunningRootEntry)
 				{
 #if WITH_EDITOR
 					Context.LastExecutionMicroSeconds = OverallCycles * FPlatformTime::GetSecondsPerCycle() * 1000.0 * 1000.0;
 #endif
-					ExecutionReachedExit().Broadcast(InEntryName);
+					ExecutionReachedExit().Broadcast(CurrentEntryName);
 #if WITH_EDITOR					
 					if (HaltedAtBreakpoint.IsValid())
 					{
 						HaltedAtBreakpoint.Reset();
 						DebugInfo->SetCurrentActiveBreakpoint(FRigVMBreakpoint());
-						ExecutionHalted().Broadcast(INDEX_NONE, nullptr, InEntryName);
+						ExecutionHalted().Broadcast(INDEX_NONE, nullptr, CurrentEntryName);
 					}
 #if UE_RIGVM_DEBUG_EXECUTION
 					ContextPublicData.Log(EMessageSeverity::Info, DebugMemoryString);					
 #endif
 #endif
 				}
-				return true;
+				return ERigVMExecuteResult::Succeeded;
 			}
 			case ERigVMOpCode::BeginBlock:
 			{
@@ -2721,22 +2800,23 @@ bool URigVM::Execute(TArrayView<URigVMMemoryStorage*> Memory, TArrayView<void*>
 
 				if(!CanExecuteEntry(Op.EntryName))
 				{
-					return false;
+					return CurrentExecuteResult = ERigVMExecuteResult::Failed;
 				}
 				else
 				{
 					// this will restore the public data after invoking the entry
 					TGuardValue<FRigVMExecuteContext> PublicDataGuard(ContextPublicData, ContextPublicData);
-					if(!Execute(Memory, AdditionalArguments, Op.EntryName))
+					const ERigVMExecuteResult ExecuteResult = Execute(CurrentMemory, CurrentAdditionalArguments, Op.EntryName);
+					if(ExecuteResult != ERigVMExecuteResult::Succeeded)
 					{
-						return false;
+						return CurrentExecuteResult = ExecuteResult;
 					}
 
 #if WITH_EDITOR
 					// if we are halted at a break point we need to exit here
-					if (ShouldHaltAtInstruction(InEntryName, ContextPublicData.InstructionIndex))
+					if (ShouldHaltAtInstruction(CurrentEntryName, ContextPublicData.InstructionIndex))
 					{
-						return true;
+						return CurrentExecuteResult = ERigVMExecuteResult::Halted;
 					}
 #endif
 				}
@@ -2747,7 +2827,7 @@ bool URigVM::Execute(TArrayView<URigVMMemoryStorage*> Memory, TArrayView<void*>
 			case ERigVMOpCode::Invalid:
 			{
 				ensure(false);
-				return false;
+				return CurrentExecuteResult = ERigVMExecuteResult::Failed;
 			}
 		}
 
@@ -2814,24 +2894,19 @@ bool URigVM::Execute(TArrayView<URigVMMemoryStorage*> Memory, TArrayView<void*>
 #endif
 	}
 
-#if WITH_EDITOR
-	if(bIsRootEntry)
-	{
-		if (HaltedAtBreakpoint.IsValid())
-		{
-			DebugInfo->SetCurrentActiveBreakpoint(FRigVMBreakpoint());
-			HaltedAtBreakpoint.Reset();
-			ExecutionHalted().Broadcast(INDEX_NONE, nullptr, InEntryName);
-		}
-	}
-#endif
-
-	return true;
+	return CurrentExecuteResult = ERigVMExecuteResult::Succeeded;
 }
 
 bool URigVM::Execute(const FName& InEntryName)
 {
-	return Execute(TArray<URigVMMemoryStorage*>(), TArrayView<void*>(), InEntryName);
+	return Execute(TArray<URigVMMemoryStorage*>(), TArrayView<void*>(), InEntryName) != ERigVMExecuteResult::Failed;
+}
+
+ERigVMExecuteResult URigVM::ExecuteLazyBranch(const FRigVMBranchInfo& InBranchToRun)
+{
+	// likely have to optimize this
+	TGuardValue<FRigVMExtendedExecuteContext> ContextGuard(Context, Context);
+	return ExecuteInstructions(InBranchToRun.FirstInstruction, InBranchToRun.LastInstruction);
 }
 
 FRigVMExternalVariable URigVM::GetExternalVariableByName(const FName& InExternalVariableName)
@@ -3299,7 +3374,7 @@ void URigVM::ClearDebugMemory()
 #endif
 }
 
-void URigVM::CacheSingleMemoryHandle(int32 InHandleIndex, const FRigVMOperand& InArg, bool bForExecute)
+void URigVM::CacheSingleMemoryHandle(int32 InHandleIndex, const TRigVMBranchInfoKey& InBranchInfoKey, const FRigVMOperand& InArg, bool bForExecute)
 {
 	URigVMMemoryStorage* Memory = GetMemoryByType(InArg.GetMemoryType(), false);
 
@@ -3337,11 +3412,20 @@ void URigVM::CacheSingleMemoryHandle(int32 InHandleIndex, const FRigVMOperand& I
 	// is in - so a VM under GetTransientPackage() can be created - but not run.
 	uint8* Data = Memory->GetData<uint8>(InArg.GetRegisterIndex());
 	const FProperty* Property = Memory->GetProperties()[InArg.GetRegisterIndex()];
-	CachedMemoryHandles[InHandleIndex] = {Data, Property, PropertyPath};
+	FRigVMMemoryHandle& Handle = CachedMemoryHandles[InHandleIndex];
+	Handle = {Data, Property, PropertyPath};
+
+	// if we are lazy executing update the handle to point to a lazy branch
+	if(const FRigVMBranchInfo* BranchInfo = GetByteCode().GetBranchInfo(InBranchInfoKey))
+	{
+		FRigVMLazyBranch* LazyBranch = &LazyBranches[BranchInfo->Index];
+		LazyBranch->VM = this;
+		LazyBranch->BranchInfo = *BranchInfo;
+		Handle.LazyBranch = LazyBranch;
+	}
 
 #if UE_BUILD_DEBUG
 	// make sure the handle points to valid memory
-	FRigVMMemoryHandle& Handle = CachedMemoryHandles[InHandleIndex]; 
 	check(Handle.GetData(false) != nullptr);
 	if(PropertyPath)
 	{
diff --git a/Engine/Source/Runtime/RigVM/Private/RigVMCore/RigVMByteCode.cpp b/Engine/Source/Runtime/RigVM/Private/RigVMCore/RigVMByteCode.cpp
--- a/Engine/Source/Runtime/RigVM/Private/RigVMCore/RigVMByteCode.cpp
+++ b/Engine/Source/Runtime/RigVM/Private/RigVMCore/RigVMByteCode.cpp
@@ -452,6 +452,9 @@ void FRigVMByteCode::Save(FArchive& Ar) const
 	}
 
 	Ar << View;
+
+	TArray<FRigVMBranchInfo> TempBranchInfos = BranchInfos;
+	Ar << TempBranchInfos;
 }
 
 void FRigVMByteCode::Load(FArchive& Ar)
@@ -686,6 +689,20 @@ void FRigVMByteCode::Load(FArchive& Ar)
 			Entries.Add(Entry);
 		}
 	}
+
+	if (Ar.CustomVer(FUE5MainStreamObjectVersion::GUID) >= FUE5MainStreamObjectVersion::RigVMLazyEvaluation)
+	{
+		Ar << BranchInfos;
+
+		// make sure the lookup table is up 2 date
+		BranchInfoPerInstruction.Reset();
+		(void)GetBranchInfo({0,0});
+	}
+	else
+	{
+		BranchInfos.Reset();
+		BranchInfoPerInstruction.Reset();
+	}
 }
 
 void FRigVMByteCode::Reset()
@@ -694,6 +711,8 @@ void FRigVMByteCode::Reset()
 	bByteCodeIsAligned = false;
 	NumInstructions = 0;
 	Entries.Reset();
+	BranchInfos.Reset();
+	BranchInfoPerInstruction.Reset();
 
 #if WITH_EDITORONLY_DATA
 	SubjectPerInstruction.Reset();
@@ -1595,6 +1614,20 @@ uint64 FRigVMByteCode::AddInvokeEntryOp(const FName& InEntryName)
 	return AddOp(FRigVMInvokeEntryOp(InEntryName));
 }
 
+void FRigVMByteCode::AddBranchInfo(const FName& InBranchLabel, int32 InInstructionIndex, int32 InArgumentIndex,
+	int32 InFirstBranchInstruction, int32 InLastBranchInstruction)
+{
+	FRigVMBranchInfo BranchInfo;
+	BranchInfo.Index = BranchInfos.Num();
+	BranchInfo.Label = InBranchLabel;
+	BranchInfo.InstructionIndex = InInstructionIndex;
+	BranchInfo.ArgumentIndex = InArgumentIndex;
+	BranchInfo.FirstInstruction = InFirstBranchInstruction;
+	BranchInfo.LastInstruction = InLastBranchInstruction;
+	BranchInfos.Add(BranchInfo);
+	BranchInfoPerInstruction.Reset();
+}
+
 FRigVMOperandArray FRigVMByteCode::GetOperandsForOp(const FRigVMInstruction& InInstruction) const
 {
 	switch(InInstruction.OpCode)
@@ -2344,3 +2377,26 @@ void FRigVMByteCode::SetOperandsForInstruction(int32 InInstructionIndex, const F
 
 #endif
 
+const FRigVMBranchInfo* FRigVMByteCode::GetBranchInfo(const TRigVMBranchInfoKey& InBranchInfoKey) const
+{
+	if(BranchInfos.IsEmpty())
+	{
+		return nullptr;
+	}
+
+	if(BranchInfoPerInstruction.IsEmpty())
+	{
+		for(const FRigVMBranchInfo& BranchInfo : BranchInfos)
+		{
+			const TRigVMBranchInfoKey Key(BranchInfo.InstructionIndex, BranchInfo.ArgumentIndex);
+			BranchInfoPerInstruction.FindOrAdd(Key) = &BranchInfo;
+		}
+	}
+	
+	if(const FRigVMBranchInfo** BranchInfoPtr = BranchInfoPerInstruction.Find(InBranchInfoKey))
+	{
+		return *BranchInfoPtr;
+	}
+
+	return nullptr;
+}
diff --git a/Engine/Source/Runtime/RigVM/Private/RigVMCore/RigVMMemoryStorage.cpp b/Engine/Source/Runtime/RigVM/Private/RigVMCore/RigVMMemoryStorage.cpp
--- a/Engine/Source/Runtime/RigVM/Private/RigVMCore/RigVMMemoryStorage.cpp
+++ b/Engine/Source/Runtime/RigVM/Private/RigVMCore/RigVMMemoryStorage.cpp
@@ -5,6 +5,7 @@
 #include "RigVMTypeUtils.h"
 #include "UObject/Interface.h"
 #include "AssetRegistry/AssetData.h"
+#include "RigVMCore/RigVM.h"
 
 #include UE_INLINE_GENERATED_CPP_BY_NAME(RigVMMemoryStorage)
 
@@ -53,6 +54,80 @@ static UObject* GetGeneratorClassOuter(UPackage* InPackage)
 	return AssetObject;	
 }
 
+////////////////////////////////////////////////////////////////////////////////////////////////////////////////////////
+
+void FRigVMBranchInfo::Serialize(FArchive& Ar)
+{
+	Ar << Index;
+
+	if(Ar.IsLoading())
+	{
+		FString LabelString;
+		Ar << LabelString;
+		Label = *LabelString;
+	}
+	else
+	{
+		FString LabelString = Label.ToString();
+		Ar << LabelString;
+	}
+	
+	Ar << InstructionIndex;
+	Ar << ArgumentIndex;
+	Ar << FirstInstruction;
+	Ar << LastInstruction;
+}
+
+ERigVMExecuteResult FRigVMLazyBranch::Execute()
+{
+	check(VM);
+	return VM->ExecuteLazyBranch(BranchInfo);
+}
+
+ERigVMExecuteResult FRigVMLazyBranch::ExecuteIfRequired(int32 InSliceIndex)
+{
+	check(VM);
+
+	if(InSliceIndex == INDEX_NONE)
+	{
+		InSliceIndex = 0;
+	}
+	
+	while(!LastVMNumExecutions.IsValidIndex(InSliceIndex))
+	{
+		LastVMNumExecutions.Add(INDEX_NONE);
+	}
+
+	if(VM->GetNumExecutions() != LastVMNumExecutions[InSliceIndex])
+	{
+		const ERigVMExecuteResult Result = Execute();
+		LastVMNumExecutions[InSliceIndex] = VM->GetNumExecutions();
+		return Result;
+	}
+	
+	return ERigVMExecuteResult::Succeeded;
+}
+
+const uint8* TRigVMLazyValueBase::GetData() const
+{
+	if(MemoryHandle)
+	{
+		checkf(MemoryHandle->IsLazy(), TEXT("If you are hitting this the compiler didn't add a branch info to the bytecode for a lazy argument."));
+		MemoryHandle->ComputeLazyValueIfNecessary(SliceIndex);
+		return MemoryHandle->GetData(bFollowPropertyPath, SliceIndex);
+	}
+	return nullptr;
+}
+
+////////////////////////////////////////////////////////////////////////////////////////////////////////////////////////
+
+bool FRigVMMemoryHandle::ComputeLazyValueIfNecessary(int32 InSliceIndex)
+{
+	check(IsLazy());
+	return LazyBranch->ExecuteIfRequired(InSliceIndex) != ERigVMExecuteResult::Failed;
+}
+
+////////////////////////////////////////////////////////////////////////////////////////////////////////////////////////
 
 FRigVMPropertyDescription::FRigVMPropertyDescription(const FProperty* InProperty, const FString& InDefaultValue, const FName& InName)
 	: Name(InName)
diff --git a/Engine/Source/Runtime/RigVM/Private/RigVMCore/RigVMNativized.cpp b/Engine/Source/Runtime/RigVM/Private/RigVMCore/RigVMNativized.cpp
--- a/Engine/Source/Runtime/RigVM/Private/RigVMCore/RigVMNativized.cpp
+++ b/Engine/Source/Runtime/RigVM/Private/RigVMCore/RigVMNativized.cpp
@@ -26,11 +26,11 @@ bool URigVMNativized::Initialize(TArrayView<URigVMMemoryStorage*> Memory, TArray
 	return true;
 }
 
-bool URigVMNativized::Execute(TArrayView<URigVMMemoryStorage*> Memory, TArrayView<void*> AdditionalArguments,
+ERigVMExecuteResult URigVMNativized::Execute(TArrayView<URigVMMemoryStorage*> Memory, TArrayView<void*> AdditionalArguments,
 	const FName& InEntryName)
 {
 	// to be implemented by the generated code
-	return false;
+	return ERigVMExecuteResult::Failed;
 }
 
 const FRigVMInstructionArray& URigVMNativized::GetInstructions()
diff --git a/Engine/Source/Runtime/RigVM/Private/RigVMCore/RigVMStruct.cpp b/Engine/Source/Runtime/RigVM/Private/RigVMCore/RigVMStruct.cpp
--- a/Engine/Source/Runtime/RigVM/Private/RigVMCore/RigVMStruct.cpp
+++ b/Engine/Source/Runtime/RigVM/Private/RigVMCore/RigVMStruct.cpp
@@ -306,6 +306,7 @@ const FName FRigVMStruct::ForLoopCountPinName("Count");
 const FName FRigVMStruct::ForLoopContinuePinName("Continue");
 const FName FRigVMStruct::ForLoopCompletedPinName("Completed");
 const FName FRigVMStruct::ForLoopIndexPinName("Index");
+const FName FRigVMStruct::ComputeLazilyMetaName("Lazy");
 
 float FRigVMStruct::GetRatioFromIndex(int32 InIndex, int32 InCount)
 {
diff --git a/Engine/Source/Runtime/RigVM/Public/RigVMCore/RigVM.h b/Engine/Source/Runtime/RigVM/Public/RigVMCore/RigVM.h
--- a/Engine/Source/Runtime/RigVM/Public/RigVMCore/RigVM.h
+++ b/Engine/Source/Runtime/RigVM/Public/RigVMCore/RigVM.h
@@ -215,7 +215,7 @@ public:
 	// Executes the VM.
 	// You can optionally provide external memory to the execution
 	// and provide optional additional operands.
-	virtual bool Execute(TArrayView<URigVMMemoryStorage*> Memory, TArrayView<void*> AdditionalArguments, const FName& InEntryName = NAME_None);
+	virtual ERigVMExecuteResult Execute(TArrayView<URigVMMemoryStorage*> Memory, TArrayView<void*> AdditionalArguments, const FName& InEntryName = NAME_None);
 
 	// Executes the VM.
 	// You can optionally provide external memory to the execution
@@ -223,6 +223,14 @@ public:
 	UFUNCTION(BlueprintCallable, Category = RigVM)
 	virtual bool Execute(const FName& InEntryName = NAME_None);
 
+	// Executes a single branch on the VM. We assume that the memory is already set correctly at this point.
+	ERigVMExecuteResult ExecuteLazyBranch(const FRigVMBranchInfo& InBranchToRun);
+
+private:
+	ERigVMExecuteResult ExecuteInstructions(int32 InFirstInstruction, int32 InLastInstruction);
+
+public:
+
 	// Add a function for execute instructions to this VM.
 	// Execute instructions can then refer to the function by index.
 	UFUNCTION()
@@ -745,6 +753,7 @@ private:
 	// changes to the layout of cached memory array should be reflected in GetContainerIndex()
 	TArray<URigVMMemoryStorage*> CachedMemory;
 	TArray<FRigVMExternalVariable> ExternalVariables;
+	TArray<FRigVMLazyBranch> LazyBranches;
 
 	// this function should be kept in sync with FRigVMOperand::GetContainerIndex()
 	static int32 GetContainerIndex(ERigVMMemoryType InType)
@@ -776,7 +785,7 @@ private:
 	// debug watch register memory needs to be cleared for each execution
 	void ClearDebugMemory();
 	
-	void CacheSingleMemoryHandle(int32 InHandleIndex, const FRigVMOperand& InArg, bool bForExecute = false);
+	void CacheSingleMemoryHandle(int32 InHandleIndex, const TRigVMBranchInfoKey& InBranchInfoKey, const FRigVMOperand& InArg, bool bForExecute = false);
 
 	FORCEINLINE_DEBUGGABLE void CopyOperandForDebuggingIfNeeded(const FRigVMOperand& InArg, const FRigVMMemoryHandle& InHandle)
 	{
@@ -831,6 +840,17 @@ private:
 	
 	mutable TArray<FName> EntryNames;
 
+	ERigVMExecuteResult CurrentExecuteResult;
+	FName CurrentEntryName;
+	bool bCurrentlyRunningRootEntry;
+	TArrayView<URigVMMemoryStorage*> CurrentMemory;
+	TArrayView<void*> CurrentAdditionalArguments;
+
+#if WITH_EDITOR
+	uint64 StartCycles = 0;
+	uint64 OverallCycles = 0;
+#endif
+
 	UPROPERTY(transient)
 	TObjectPtr<URigVM> DeferredVMToCopy;
 
diff --git a/Engine/Source/Runtime/RigVM/Public/RigVMCore/RigVMByteCode.h b/Engine/Source/Runtime/RigVM/Public/RigVMCore/RigVMByteCode.h
--- a/Engine/Source/Runtime/RigVM/Public/RigVMCore/RigVMByteCode.h
+++ b/Engine/Source/Runtime/RigVM/Public/RigVMCore/RigVMByteCode.h
@@ -27,6 +27,9 @@ class FArchive;
 class UObject;
 struct FRigVMByteCode;
 
+typedef TTuple<int32,int32> TRigVMBranchInfoKey;
+
+
 // The code for a single operation within the RigVM
 UENUM()
 enum class ERigVMOpCode : uint8
@@ -992,6 +995,9 @@ public:
 	// adds an invoke entry operator
 	uint64 AddInvokeEntryOp(const FName& InEntryName);
 
+	// adds information about a branch for an instruction's argument
+	void AddBranchInfo(const FName& InBranchLabel, int32 InInstructionIndex, int32 InArgumentIndex, int32 InFirstBranchInstruction, int32 InLastBranchInstruction);
+
 	// returns an instruction array for iterating over all operators
 	FORCEINLINE FRigVMInstructionArray GetInstructions() const
 	{
@@ -1207,6 +1213,13 @@ private:
 	UPROPERTY()
 	TArray<FRigVMByteCodeEntry> Entries;
 
+	// a list of all lazily evaluation branches
+	UPROPERTY()
+	TArray<FRigVMBranchInfo> BranchInfos;
+
+	const FRigVMBranchInfo* GetBranchInfo(const TRigVMBranchInfoKey& InBranchInfoKey) const;
+	mutable TMap<TRigVMBranchInfoKey, const FRigVMBranchInfo*> BranchInfoPerInstruction;
+
 	// if this is set to true the stored bytecode is aligned / padded
 	bool bByteCodeIsAligned;
 
diff --git a/Engine/Source/Runtime/RigVM/Public/RigVMCore/RigVMDispatchFactory.h b/Engine/Source/Runtime/RigVM/Public/RigVMCore/RigVMDispatchFactory.h
--- a/Engine/Source/Runtime/RigVM/Public/RigVMCore/RigVMDispatchFactory.h
+++ b/Engine/Source/Runtime/RigVM/Public/RigVMCore/RigVMDispatchFactory.h
@@ -74,6 +74,12 @@ public:
 	// Returns meta data on the property of the permutations 
 	FORCEINLINE virtual FString GetArgumentMetaData(const FName& InArgumentName, const FName& InMetaDataKey) const { return FString(); }
 
+	// Returns true if the factory provides metadata for a given argument
+	bool HasArgumentMetaData(const FName& InArgumentName, const FName& InMetaDataKey) const
+	{
+		return !GetArgumentMetaData(InArgumentName, InMetaDataKey).IsEmpty();
+	}
+
 	// Returns the category this factory is under
 	virtual FString GetCategory() const;
 
diff --git a/Engine/Source/Runtime/RigVM/Public/RigVMCore/RigVMExecuteContext.h b/Engine/Source/Runtime/RigVM/Public/RigVMCore/RigVMExecuteContext.h
--- a/Engine/Source/Runtime/RigVM/Public/RigVMCore/RigVMExecuteContext.h
+++ b/Engine/Source/Runtime/RigVM/Public/RigVMCore/RigVMExecuteContext.h
@@ -21,6 +21,7 @@
 #include "UObject/UnrealNames.h"
 #include "UObject/UnrealType.h"
 #include "RigVMCore/RigVMNameCache.h"
+#include "UObject/StructOnScope.h"
 
 #include "RigVMExecuteContext.generated.h"
 
@@ -276,6 +277,15 @@ struct RIGVM_API FRigVMExecuteContext
 		}
 	}
 
+	virtual void Copy(const FRigVMExecuteContext* InOtherContext)
+	{
+		EventName = InOtherContext->EventName;
+		FunctionName = InOtherContext->FunctionName;
+		InstructionIndex = InOtherContext->InstructionIndex;
+		RuntimeSettings = InOtherContext->RuntimeSettings;
+		NameCache = InOtherContext->NameCache;
+	}
+
 protected:
 
 	FORCEINLINE virtual void Reset()
@@ -340,6 +350,15 @@ struct RIGVM_API FRigVMExtendedExecuteContext
 		}
 	}
 
+	FORCEINLINE_DEBUGGABLE FRigVMExtendedExecuteContext(const FRigVMExtendedExecuteContext& InOther)
+		: PublicDataScope()
+		, VM(nullptr)
+		, LastExecutionMicroSeconds()
+		, Factory(nullptr)
+	{
+		*this = InOther;
+	}
+
 	virtual ~FRigVMExtendedExecuteContext() {}
 
 	FORCEINLINE void Reset()
@@ -352,13 +371,23 @@ struct RIGVM_API FRigVMExtendedExecuteContext
 		Factory = nullptr;
 	}
 
-	FORCEINLINE FRigVMExtendedExecuteContext& operator =(const FRigVMExtendedExecuteContext& Other)
+	FORCEINLINE_DEBUGGABLE FRigVMExtendedExecuteContext& operator =(const FRigVMExtendedExecuteContext& Other)
 	{
-		const UScriptStruct* PublicDataStruct = Cast<UScriptStruct>(Other.PublicDataScope.GetStruct());
-		check(PublicDataStruct);
-		PublicDataScope = FStructOnScope(PublicDataStruct);
-		PublicDataStruct->CopyScriptStruct(PublicDataScope.GetStructMemory(), Other.PublicDataScope.GetStructMemory());
-		SetDefaultNameCache();
+		const UScriptStruct* OtherPublicDataStruct = Cast<UScriptStruct>(Other.PublicDataScope.GetStruct());
+		check(OtherPublicDataStruct);
+		if(PublicDataScope.GetStruct() != OtherPublicDataStruct)
+		{
+			PublicDataScope = FStructOnScope(OtherPublicDataStruct);
+		}
+
+		FRigVMExecuteContext* ThisPublicContext = (FRigVMExecuteContext*)PublicDataScope.GetStructMemory();
+		const FRigVMExecuteContext* OtherPublicContext = (const FRigVMExecuteContext*)Other.PublicDataScope.GetStructMemory();
+		ThisPublicContext->Copy(OtherPublicContext);
+
+		if(OtherPublicContext->GetNameCache() == &Other.NameCache)
+		{
+			SetDefaultNameCache();
+		}
 
 		VM = Other.VM;
 		OpaqueArguments = Other.OpaqueArguments;
diff --git a/Engine/Source/Runtime/RigVM/Public/RigVMCore/RigVMMemoryStorage.h b/Engine/Source/Runtime/RigVM/Public/RigVMCore/RigVMMemoryStorage.h
--- a/Engine/Source/Runtime/RigVM/Public/RigVMCore/RigVMMemoryStorage.h
+++ b/Engine/Source/Runtime/RigVM/Public/RigVMCore/RigVMMemoryStorage.h
@@ -6,6 +6,7 @@
 #include "Containers/UnrealString.h"
 #include "CoreMinimal.h"
 #include "CoreTypes.h"
+#include "RigVMExecuteContext.h"
 #include "EdGraph/EdGraphNode.h"
 #include "Misc/AssertionMacros.h"
 #include "RigVMMemoryCommon.h"
@@ -31,6 +32,158 @@
 #include "RigVMMemoryStorage.generated.h"
 
 class FArchive;
+class URigVM;
+struct FRigVMMemoryHandle;
+
+enum class ERigVMExecuteResult : uint8
+{
+	Failed
+	,Succeeded
+#if WITH_EDITOR
+	,Halted
+#endif
+};
+
+//////////////////////////////////////////////////////////////////////////////
+/// Lazy execution
+//////////////////////////////////////////////////////////////////////////////
+
+// A description of a branch in the VM's bytecode
+USTRUCT()
+struct RIGVM_API FRigVMBranchInfo
+{
+	GENERATED_USTRUCT_BODY()
+
+	FRigVMBranchInfo()
+	: Index(INDEX_NONE)
+	, Label(NAME_None)
+	, InstructionIndex(INDEX_NONE)
+	, ArgumentIndex(INDEX_NONE)
+	, FirstInstruction(INDEX_NONE)
+	, LastInstruction(INDEX_NONE)
+	{}
+
+	bool IsValid() const
+	{
+		return InstructionIndex != INDEX_NONE;
+	}
+
+	void Serialize(FArchive& Ar);
+	FORCEINLINE friend FArchive& operator<<(FArchive& Ar, FRigVMBranchInfo& P)
+	{
+		P.Serialize(Ar);
+		return Ar;
+	}
+
+	UPROPERTY()
+	int32 Index;
+
+	UPROPERTY()
+	FName Label;
+
+	UPROPERTY()
+	int32 InstructionIndex;
+
+	UPROPERTY()
+	int32 ArgumentIndex;
+
+	UPROPERTY()
+	uint16 FirstInstruction;
+
+	UPROPERTY()
+	uint16 LastInstruction;
+};
+
+
+/**
+ * A branch which can be lazily executed
+ */
+struct RIGVM_API FRigVMLazyBranch
+{
+public:
+	
+	FRigVMLazyBranch()
+		: VM(nullptr)
+		, LastVMNumExecutions()
+		, BranchInfo()
+	{}
+
+	ERigVMExecuteResult Execute();
+	ERigVMExecuteResult ExecuteIfRequired(int32 InSliceIndex = INDEX_NONE);
+
+private:
+
+	URigVM* VM;
+	TArray<int32> LastVMNumExecutions;
+	FRigVMBranchInfo BranchInfo;
+
+	friend class URigVM;
+};
+
+/**
+ * A template container for a lazily computed value
+ */
+struct RIGVM_API TRigVMLazyValueBase
+{
+public:
+
+	TRigVMLazyValueBase()
+		: bFollowPropertyPath(false)
+		, SliceIndex(INDEX_NONE)
+		, MemoryHandle(nullptr)
+	{
+	}
+
+	const uint8* GetData() const;
+
+protected:
+	
+	bool bFollowPropertyPath;
+	int32 SliceIndex;
+	FRigVMMemoryHandle* MemoryHandle;
+	
+	friend class URigVM;
+	friend struct FRigVMMemoryHandle;
+};
+
+/**
+ * A template container for a lazily computed value
+ */
+template<typename ComputedType>
+struct TRigVMLazyValue : public TRigVMLazyValueBase
+{
+public:
+
+	TRigVMLazyValue(ComputedType InConstValue)
+		: TRigVMLazyValueBase()
+		, ConstValue(InConstValue)
+	{}
+
+	FORCEINLINE const ComputedType& Get() const
+	{
+		if(MemoryHandle)
+		{
+			return *(const ComputedType*)GetData();
+		}
+		return ConstValue;
+	}
+
+protected:
+
+	TRigVMLazyValue(const TRigVMLazyValue<ComputedType>& InOther) = delete;
+
+	TRigVMLazyValue(const TRigVMLazyValueBase& InOther)
+		: TRigVMLazyValueBase(InOther)
+	{}
+
+	ComputedType ConstValue;
+
+	friend struct FRigVMMemoryHandle;
+};
+
+//////////////////////////////////////////////////////////////////////////////
+/// Memory handle
+//////////////////////////////////////////////////////////////////////////////
 
 /**
  * The FRigVMMemoryHandle is used to access the memory used within a URigMemoryStorage.
@@ -53,6 +206,7 @@ public:
 		: Ptr(nullptr)
 		, Property(nullptr) 
 		, PropertyPath(nullptr)
+		, LazyBranch(nullptr)
 	{}
 
 	// Constructor from complete data
@@ -60,6 +214,7 @@ public:
 		: Ptr(InPtr)
 		, Property(InProperty)
 		, PropertyPath(InPropertyPath)
+		, LazyBranch(nullptr)
 	{
 	}
 
@@ -75,7 +230,7 @@ public:
 	}
 
 	/**
-	 * Returns the cached pointer stored within the handle.
+	 * Computes the data if necessary and returns the cached pointer stored within the handle.
 	 * @param bFollowPropertyPath If set to true the memory handle will traverse the memory using the property path
 	 * @param InSliceIndex If this is != INDEX_NONE the memory handle will return the slice of the memory requested
 	 * @return The traversed memory cached by this handle (const)
@@ -85,6 +240,33 @@ public:
 		return GetData_Internal(bFollowPropertyPath, InSliceIndex);
 	}
 
+	/**
+	 * Computes the data if necessary and returns the cached pointer stored within the handle.
+	 * @param bFollowPropertyPath If set to true the memory handle will traverse the memory using the property path
+	 * @param InSliceIndex If this is != INDEX_NONE the memory handle will return the slice of the memory requested
+	 * @return The traversed memory cached by this handle (const)
+	 */
+	template<typename ComputedType>
+	FORCEINLINE TRigVMLazyValue<ComputedType> GetDataLazily(bool bFollowPropertyPath = false, int32 InSliceIndex = INDEX_NONE) const
+	{
+		return GetDataLazily_Internal(bFollowPropertyPath, InSliceIndex);
+	}
+
+	/**
+	 * Returns true if this memory handle depends on a lazily executed branch
+	 * @return True if this memory handle depends on a lazily executed branch
+	 */
+	FORCEINLINE bool IsLazy() const
+	{
+		return LazyBranch != nullptr;
+	}
+	
+	/**
+	 * Computes the data if necessary and returns true if the value is valid
+	 * @return True if the value of the handle is valid after the compute
+	 */
+	bool ComputeLazyValueIfNecessary(int32 InSliceIndex = INDEX_NONE);
+
 	// Returns the head property of this handle
 	FORCEINLINE const FProperty* GetProperty() const
 	{
@@ -596,6 +778,18 @@ private:
 		return Ptr;
 	}
 
+	FORCEINLINE_DEBUGGABLE TRigVMLazyValueBase GetDataLazily_Internal(bool bFollowPropertyPath, int32 InSliceIndex) const
+	{
+		// this should only get called on memory handles which can execute lazily
+		check(LazyBranch != nullptr);
+		
+		TRigVMLazyValueBase LazyValue;
+		LazyValue.MemoryHandle = (FRigVMMemoryHandle*)this;
+		LazyValue.bFollowPropertyPath = bFollowPropertyPath;
+		LazyValue.SliceIndex = InSliceIndex;
+		return LazyValue;
+	}
+
 	FORCEINLINE static const FProperty* GetArrayElementProperty(const FProperty* InProperty)
 	{
 		if(const FArrayProperty* ArrayProperty = CastField<FArrayProperty>(InProperty))
@@ -678,7 +872,11 @@ private:
 	// The [optional] property path used by this handle
 	const FRigVMPropertyPath* PropertyPath;
 
+	// The [optional] lazy branch used by this handle
+	FRigVMLazyBranch* LazyBranch;
+
 	friend class URigVM;
+	friend struct TRigVMLazyValueBase;
 };
 
 //////////////////////////////////////////////////////////////////////////////
diff --git a/Engine/Source/Runtime/RigVM/Public/RigVMCore/RigVMNativized.h b/Engine/Source/Runtime/RigVM/Public/RigVMCore/RigVMNativized.h
--- a/Engine/Source/Runtime/RigVM/Public/RigVMCore/RigVMNativized.h
+++ b/Engine/Source/Runtime/RigVM/Public/RigVMCore/RigVMNativized.h
@@ -51,7 +51,7 @@ public:
 	virtual void Empty() override { return; }
 	virtual void CopyFrom(URigVM* InVM, bool bDeferCopy = false, bool bReferenceLiteralMemory = false, bool bReferenceByteCode = false, bool bCopyExternalVariables = false, bool bCopyDynamicRegisters = false) override { return; }
 	virtual bool Initialize(TArrayView<URigVMMemoryStorage*> Memory, TArrayView<void*> AdditionalArguments, bool bInitializeMemory = true) override;
-	virtual bool Execute(TArrayView<URigVMMemoryStorage*> Memory, TArrayView<void*> AdditionalArguments, const FName& InEntryName = NAME_None) override;
+	virtual ERigVMExecuteResult Execute(TArrayView<URigVMMemoryStorage*> Memory, TArrayView<void*> AdditionalArguments, const FName& InEntryName = NAME_None) override;
 	virtual int32 AddRigVMFunction(UScriptStruct* InRigVMStruct, const FName& InMethodName) override { return INDEX_NONE; }
 	virtual FString GetRigVMFunctionName(int32 InFunctionIndex) const override { return FString(); }
 	virtual URigVMMemoryStorage* GetMemoryByType(ERigVMMemoryType InMemoryType, bool bCreateIfNeeded = true) override { return nullptr; }
diff --git a/Engine/Source/Runtime/RigVM/Public/RigVMCore/RigVMStruct.h b/Engine/Source/Runtime/RigVM/Public/RigVMCore/RigVMStruct.h
--- a/Engine/Source/Runtime/RigVM/Public/RigVMCore/RigVMStruct.h
+++ b/Engine/Source/Runtime/RigVM/Public/RigVMCore/RigVMStruct.h
@@ -249,6 +249,7 @@ public:
 	static const FName ForLoopContinuePinName;
 	static const FName ForLoopCompletedPinName;
 	static const FName ForLoopIndexPinName;
+	static const FName ComputeLazilyMetaName;
 
 protected:
 
