Control Rig: Initialize PropertyIndex with a value.

#rnx

diff --git a/Engine/Source/Runtime/RigVM/Public/RigVMCore/RigVMGraphFunctionDefinition.h b/Engine/Source/Runtime/RigVM/Public/RigVMCore/RigVMGraphFunctionDefinition.h
--- a/Engine/Source/Runtime/RigVM/Public/RigVMCore/RigVMGraphFunctionDefinition.h
+++ b/Engine/Source/Runtime/RigVM/Public/RigVMCore/RigVMGraphFunctionDefinition.h
@@ -55,7 +55,7 @@ struct FRigVMFunctionCompilationPropertyPath
 	GENERATED_BODY()
 
 	UPROPERTY()
-	int32 PropertyIndex;
+	int32 PropertyIndex = INDEX_NONE;
 
 	UPROPERTY()
 	FString HeadCPPType;
