Control Rig: Remove an #include that was causing a circular include chain.

#rnx

diff --git a/Engine/Source/Runtime/RigVM/Public/RigVMCore/RigVMMemoryStorage.h b/Engine/Source/Runtime/RigVM/Public/RigVMCore/RigVMMemoryStorage.h
--- a/Engine/Source/Runtime/RigVM/Public/RigVMCore/RigVMMemoryStorage.h
+++ b/Engine/Source/Runtime/RigVM/Public/RigVMCore/RigVMMemoryStorage.h
@@ -6,7 +6,6 @@
 #include "Containers/UnrealString.h"
 #include "CoreMinimal.h"
 #include "CoreTypes.h"
-#include "RigVMExecuteContext.h"
 #include "EdGraph/EdGraphNode.h"
 #include "Misc/AssertionMacros.h"
 #include "RigVMMemoryCommon.h"
