RigVM: Follow up fixes for Control Flow functionality

diff --git a/Engine/Source/Developer/RigVMDeveloper/Private/RigVMCompiler/RigVMAST.cpp b/Engine/Source/Developer/RigVMDeveloper/Private/RigVMCompiler/RigVMAST.cpp
--- a/Engine/Source/Developer/RigVMDeveloper/Private/RigVMCompiler/RigVMAST.cpp
+++ b/Engine/Source/Developer/RigVMDeveloper/Private/RigVMCompiler/RigVMAST.cpp
@@ -2888,7 +2888,9 @@ FString FRigVMParserAST::DumpDot() const
 					if(SubGraphIndex == INDEX_NONE)
 					{
 						const int32 ASTGraphIndex = OutGraph.FindSubGraph(TEXT("AST"));
-						SubGraphIndex = OutGraph.AddSubGraph(InExpr->GetName(), InExpr->GetName(), ASTGraphIndex);
+						FString SanitizedNameString = InExpr->GetName().ToString();
+						SanitizedNameString.RemoveSpacesInline();
+						SubGraphIndex = OutGraph.AddSubGraph(*SanitizedNameString, InExpr->GetName(), ASTGraphIndex);
 					}
 					break;
 				}
diff --git a/Engine/Source/Developer/RigVMDeveloper/Private/RigVMCompiler/RigVMCompiler.cpp b/Engine/Source/Developer/RigVMDeveloper/Private/RigVMCompiler/RigVMCompiler.cpp
--- a/Engine/Source/Developer/RigVMDeveloper/Private/RigVMCompiler/RigVMCompiler.cpp
+++ b/Engine/Source/Developer/RigVMDeveloper/Private/RigVMCompiler/RigVMCompiler.cpp
@@ -1174,19 +1174,17 @@ void URigVMCompiler::TraverseBlock(const FRigVMBlockExprAST* InExpr, FRigVMCompi
 								// to step over them as well here.
 								else if(const URigVMUnitNode* UnitNode = Cast<URigVMUnitNode>(CallExternNode))
 								{
-									int32 ArgumentIndex = 0;
-									for(const URigVMPin* NodePin : UnitNode->GetPins())
+									if(const FRigVMFunction* Function = UnitNode->GetResolvedFunction())
 									{
-										if(NodePin->IsExecuteContext())
+										for(int32 ArgumentIndex = 0; ArgumentIndex != Function->Arguments.Num(); ArgumentIndex++)
 										{
-											continue;
+											const FRigVMFunctionArgument& Argument = Function->Arguments[ArgumentIndex];
+											if(Argument.Name == Pin->GetFName())
+											{
+												BranchInfo.ArgumentIndex = ArgumentIndex;
+												break;
+											}
 										}
-										if(NodePin->GetFName() == Pin->GetFName())
-										{
-											BranchInfo.ArgumentIndex = ArgumentIndex;
-											break;
-										}
-										ArgumentIndex++;
 									}
 								}	
 							}
diff --git a/Engine/Source/Runtime/RigVM/Private/RigVMCore/RigVM.cpp b/Engine/Source/Runtime/RigVM/Private/RigVMCore/RigVM.cpp
--- a/Engine/Source/Runtime/RigVM/Private/RigVMCore/RigVM.cpp
+++ b/Engine/Source/Runtime/RigVM/Private/RigVMCore/RigVM.cpp
@@ -737,7 +737,7 @@ int32 URigVM::AddRigVMFunction(const FString& InFunctionName)
 
 	GetFunctionNames().Add(FunctionName);
 	GetFactories().Add(Function->Factory);
-	return GetFunctions().Add(Function->FunctionPtr);
+	return GetFunctions().Add(Function);
 }
 
 FString URigVM::GetRigVMFunctionName(int32 InFunctionIndex) const
@@ -1025,7 +1025,7 @@ void URigVM::ResolveFunctionsIfRequired()
 		{
 			if(const FRigVMFunction* Function = FRigVMRegistry::Get().FindFunction(*GetFunctionNames()[FunctionIndex].ToString()))
 			{
-				GetFunctions()[FunctionIndex] = Function->FunctionPtr;
+				GetFunctions()[FunctionIndex] = Function;
 				GetFactories()[FunctionIndex] = Function->Factory;
 			}
 			else
@@ -1116,6 +1116,7 @@ void URigVM::CacheMemoryHandlesIfRequired(TArrayView<URigVMMemoryStorage*> InMem
 	RefreshExternalPropertyPaths();
 
 	FRigVMByteCode& ByteCode = GetByteCode();
+	TArray<const FRigVMFunction*>& Functions = GetFunctions();
 
 	// force to update the map of branch infos once
 	(void)ByteCode.GetBranchInfo({0, 0});
@@ -1132,11 +1133,13 @@ void URigVM::CacheMemoryHandlesIfRequired(TArrayView<URigVMMemoryStorage*> InMem
 
 		if (OpCode >= ERigVMOpCode::Execute_0_Operands && OpCode <= ERigVMOpCode::Execute_64_Operands)
 		{
+			const FRigVMExecuteOp& Op = ByteCode.GetOpAt<FRigVMExecuteOp>(Instructions[InstructionIndex]);
 			FRigVMOperandArray Operands = ByteCode.GetOperandsForExecuteOp(Instructions[InstructionIndex]);
+			const FRigVMFunction* Function = Functions[Op.FunctionIndex];
 
 			for (int32 ArgIndex = 0; ArgIndex < Operands.Num(); ArgIndex++)
 			{
-				InOpFunc(InHandleBaseIndex++, {InstructionIndex, ArgIndex}, Operands[ArgIndex]);
+				InOpFunc(InHandleBaseIndex++, {InstructionIndex, ArgIndex, Function->Arguments[ArgIndex].Name}, Operands[ArgIndex]);
 			}
 		}
 		else
@@ -1152,14 +1155,14 @@ void URigVM::CacheMemoryHandlesIfRequired(TArrayView<URigVMMemoryStorage*> InMem
 			case ERigVMOpCode::ArrayReverse:
 				{
 					const FRigVMUnaryOp& Op = ByteCode.GetOpAt<FRigVMUnaryOp>(Instructions[InstructionIndex]);
-					InOpFunc(InHandleBaseIndex, {InstructionIndex, 0}, Op.Arg);
+					InOpFunc(InHandleBaseIndex, {}, Op.Arg);
 					break;
 				}
 			case ERigVMOpCode::Copy:
 				{
 					const FRigVMCopyOp& Op = ByteCode.GetOpAt<FRigVMCopyOp>(Instructions[InstructionIndex]);
-					InOpFunc(InHandleBaseIndex + 0, {InstructionIndex, 0}, Op.Source);
-					InOpFunc(InHandleBaseIndex + 1, {InstructionIndex, 1}, Op.Target);
+					InOpFunc(InHandleBaseIndex + 0, {}, Op.Source);
+					InOpFunc(InHandleBaseIndex + 1, {}, Op.Target);
 					break;
 				}
 			case ERigVMOpCode::Equals:
@@ -1167,11 +1170,11 @@ void URigVM::CacheMemoryHandlesIfRequired(TArrayView<URigVMMemoryStorage*> InMem
 				{
 					const FRigVMComparisonOp& Op = ByteCode.GetOpAt<FRigVMComparisonOp>(Instructions[InstructionIndex]);
 					FRigVMOperand Arg = Op.A;
-					InOpFunc(InHandleBaseIndex + 0, {InstructionIndex, 0}, Arg);
+					InOpFunc(InHandleBaseIndex + 0, {}, Arg);
 					Arg = Op.B;
-					InOpFunc(InHandleBaseIndex + 1, {InstructionIndex, 1}, Arg);
+					InOpFunc(InHandleBaseIndex + 1, {}, Arg);
 					Arg = Op.Result;
-					InOpFunc(InHandleBaseIndex + 2, {InstructionIndex, 2}, Arg);
+					InOpFunc(InHandleBaseIndex + 2, {}, Arg);
 					break;
 				}
 			case ERigVMOpCode::JumpAbsolute:
@@ -1186,14 +1189,14 @@ void URigVM::CacheMemoryHandlesIfRequired(TArrayView<URigVMMemoryStorage*> InMem
 				{
 					const FRigVMJumpIfOp& Op = ByteCode.GetOpAt<FRigVMJumpIfOp>(Instructions[InstructionIndex]);
 					const FRigVMOperand& Arg = Op.Arg;
-					InOpFunc(InHandleBaseIndex, {InstructionIndex, 0}, Arg);
+					InOpFunc(InHandleBaseIndex, {}, Arg);
 					break;
 				}
 			case ERigVMOpCode::ChangeType:
 				{
 					const FRigVMChangeTypeOp& Op = ByteCode.GetOpAt<FRigVMChangeTypeOp>(Instructions[InstructionIndex]);
 					const FRigVMOperand& Arg = Op.Arg;
-					InOpFunc(InHandleBaseIndex, {InstructionIndex, 0}, Arg);
+					InOpFunc(InHandleBaseIndex, {}, Arg);
 					break;
 				}
 			case ERigVMOpCode::Exit:
@@ -1209,8 +1212,8 @@ void URigVM::CacheMemoryHandlesIfRequired(TArrayView<URigVMMemoryStorage*> InMem
 			case ERigVMOpCode::ArrayUnion:
 				{
 					const FRigVMBinaryOp& Op = ByteCode.GetOpAt<FRigVMBinaryOp>(Instructions[InstructionIndex]);
-					InOpFunc(InHandleBaseIndex + 0, {InstructionIndex, 0}, Op.ArgA);
-					InOpFunc(InHandleBaseIndex + 1, {InstructionIndex, 1}, Op.ArgB);
+					InOpFunc(InHandleBaseIndex + 0, {}, Op.ArgA);
+					InOpFunc(InHandleBaseIndex + 1, {}, Op.ArgB);
 					break;
 				}
 			case ERigVMOpCode::ArrayAdd:
@@ -1221,29 +1224,29 @@ void URigVM::CacheMemoryHandlesIfRequired(TArrayView<URigVMMemoryStorage*> InMem
 			case ERigVMOpCode::ArrayIntersection:
 				{
 					const FRigVMTernaryOp& Op = ByteCode.GetOpAt<FRigVMTernaryOp>(Instructions[InstructionIndex]);
-					InOpFunc(InHandleBaseIndex + 0, {InstructionIndex, 0}, Op.ArgA);
-					InOpFunc(InHandleBaseIndex + 1, {InstructionIndex, 1}, Op.ArgB);
-					InOpFunc(InHandleBaseIndex + 2, {InstructionIndex, 2}, Op.ArgC);
+					InOpFunc(InHandleBaseIndex + 0, {}, Op.ArgA);
+					InOpFunc(InHandleBaseIndex + 1, {}, Op.ArgB);
+					InOpFunc(InHandleBaseIndex + 2, {}, Op.ArgC);
 					break;
 				}
 			case ERigVMOpCode::ArrayFind:
 				{
 					const FRigVMQuaternaryOp& Op = ByteCode.GetOpAt<FRigVMQuaternaryOp>(Instructions[InstructionIndex]);
-					InOpFunc(InHandleBaseIndex + 0, {InstructionIndex, 0}, Op.ArgA);
-					InOpFunc(InHandleBaseIndex + 1, {InstructionIndex, 1}, Op.ArgB);
-					InOpFunc(InHandleBaseIndex + 2, {InstructionIndex, 2}, Op.ArgC);
-					InOpFunc(InHandleBaseIndex + 3, {InstructionIndex, 3}, Op.ArgD);
+					InOpFunc(InHandleBaseIndex + 0, {}, Op.ArgA);
+					InOpFunc(InHandleBaseIndex + 1, {}, Op.ArgB);
+					InOpFunc(InHandleBaseIndex + 2, {}, Op.ArgC);
+					InOpFunc(InHandleBaseIndex + 3, {}, Op.ArgD);
 					break;
 				}
 			case ERigVMOpCode::ArrayIterator:
 				{
 					const FRigVMSenaryOp& Op = ByteCode.GetOpAt<FRigVMSenaryOp>(Instructions[InstructionIndex]);
-					InOpFunc(InHandleBaseIndex + 0, {InstructionIndex, 0}, Op.ArgA);
-					InOpFunc(InHandleBaseIndex + 1, {InstructionIndex, 1}, Op.ArgB);
-					InOpFunc(InHandleBaseIndex + 2, {InstructionIndex, 2}, Op.ArgC);
-					InOpFunc(InHandleBaseIndex + 3, {InstructionIndex, 3}, Op.ArgD);
-					InOpFunc(InHandleBaseIndex + 4, {InstructionIndex, 4}, Op.ArgE);
-					InOpFunc(InHandleBaseIndex + 5, {InstructionIndex, 5}, Op.ArgF);
+					InOpFunc(InHandleBaseIndex + 0, {}, Op.ArgA);
+					InOpFunc(InHandleBaseIndex + 1, {}, Op.ArgB);
+					InOpFunc(InHandleBaseIndex + 2, {}, Op.ArgC);
+					InOpFunc(InHandleBaseIndex + 3, {}, Op.ArgD);
+					InOpFunc(InHandleBaseIndex + 4, {}, Op.ArgE);
+					InOpFunc(InHandleBaseIndex + 5, {}, Op.ArgF);
 					break;
 				}
 			case ERigVMOpCode::EndBlock:
@@ -1555,7 +1558,7 @@ bool URigVM::Initialize(TArrayView<URigVMMemoryStorage*> Memory, TArrayView<void
 	}
 	
 	FRigVMByteCode& ByteCode = GetByteCode();
-	TArray<FRigVMFunctionPtr>& Functions = GetFunctions();
+	TArray<const FRigVMFunction*>& Functions = GetFunctions();
 	TArray<const FRigVMDispatchFactory*>& Factories = GetFactories();
 
 #if WITH_EDITOR
@@ -1661,7 +1664,7 @@ bool URigVM::Initialize(TArrayView<URigVMMemoryStorage*> Memory, TArrayView<void
 				Context.BeginSlice(MaxSliceCount);
 				for (int32 SliceIndex = 0; SliceIndex < MaxSliceCount; SliceIndex++)
 				{
-					(*Functions[Op.FunctionIndex])(Context, OpHandles);
+					(*Functions[Op.FunctionIndex]->FunctionPtr)(Context, OpHandles);
 					Context.IncrementSlice();
 				}
 				Context.EndSlice();
@@ -1770,7 +1773,7 @@ ERigVMExecuteResult URigVM::Execute(TArrayView<URigVMMemoryStorage*> Memory, TAr
 	}
 
 	FRigVMByteCode& ByteCode = GetByteCode();
-	TArray<FRigVMFunctionPtr>& Functions = GetFunctions();
+	TArray<const FRigVMFunction*>& Functions = GetFunctions();
 	TArray<const FRigVMDispatchFactory*>& Factories = GetFactories();
 
 	FRigVMExecuteContext& ContextPublicData = Context.GetPublicData<>();
@@ -1919,7 +1922,7 @@ ERigVMExecuteResult URigVM::ExecuteInstructions(int32 InFirstInstruction, int32
 	TGuardValue<uint16> InstructionIndexGuard(ContextPublicData.InstructionIndex, (uint16)InFirstInstruction);
 	
 	FRigVMByteCode& ByteCode = GetByteCode();
-	TArray<FRigVMFunctionPtr>& Functions = GetFunctions();
+	TArray<const FRigVMFunction*>& Functions = GetFunctions();
 	TArray<const FRigVMDispatchFactory*>& Factories = GetFactories();
 
 #if WITH_EDITOR
@@ -2068,13 +2071,7 @@ ERigVMExecuteResult URigVM::ExecuteInstructions(int32 InFirstInstruction, int32
 				ContextPublicData.FunctionName = FunctionNames[Op.FunctionIndex];
 #endif
 				Context.Factory = Factories[Op.FunctionIndex];
-				(*Functions[Op.FunctionIndex])(Context, Handles);
-
-				// todo: if this is a control flow node - we may need to jump
-				// this should likely be done with a control flow jump op - which can
-				// loop up the branch based on additional branch infos for it.
-				// we should also make sure to ignore the first branch - since that's the default
-				// we'll just automatically continue to...
+				(*Functions[Op.FunctionIndex]->FunctionPtr)(Context, Handles);
 
 #if WITH_EDITOR
 
@@ -3483,12 +3480,16 @@ void URigVM::CacheSingleMemoryHandle(int32 InHandleIndex, const FRigVMBranchInfo
 	Handle = {Data, Property, PropertyPath};
 
 	// if we are lazy executing update the handle to point to a lazy branch
-	if(const FRigVMBranchInfo* BranchInfo = GetByteCode().GetBranchInfo(InBranchInfoKey))
+	if(InBranchInfoKey.IsValid())
 	{
-		FRigVMLazyBranch* LazyBranch = &LazyBranches[BranchInfo->Index];
-		LazyBranch->VM = this;
-		LazyBranch->BranchInfo = *BranchInfo;
-		Handle.LazyBranch = LazyBranch;
+		UE_LOG(LogRigVM, Warning, TEXT("Looking up branch info key instruction %d, argument %d"), InBranchInfoKey.InstructionIndex, InBranchInfoKey.ArgumentIndex);
+		if(const FRigVMBranchInfo* BranchInfo = GetByteCode().GetBranchInfo(InBranchInfoKey))
+		{
+			FRigVMLazyBranch* LazyBranch = &LazyBranches[BranchInfo->Index];
+			LazyBranch->VM = this;
+			LazyBranch->BranchInfo = *BranchInfo;
+			Handle.LazyBranch = LazyBranch;
+		}
 	}
 
 #if UE_BUILD_DEBUG
diff --git a/Engine/Source/Runtime/RigVM/Private/RigVMCore/RigVMFunction.cpp b/Engine/Source/Runtime/RigVM/Private/RigVMCore/RigVMFunction.cpp
--- a/Engine/Source/Runtime/RigVM/Private/RigVMCore/RigVMFunction.cpp
+++ b/Engine/Source/Runtime/RigVM/Private/RigVMCore/RigVMFunction.cpp
@@ -32,6 +32,13 @@ FString FRigVMFunction::GetModuleName() const
 			return Package->GetName();
 		}
 	}
+	if (Factory)
+	{
+		if (UPackage* Package = Factory->GetScriptStruct()->GetPackage())
+		{
+			return Package->GetName();
+		}
+	}
 #endif
 	return FString();
 }
diff --git a/Engine/Source/Runtime/RigVM/Public/RigVMCore/RigVM.h b/Engine/Source/Runtime/RigVM/Public/RigVMCore/RigVM.h
--- a/Engine/Source/Runtime/RigVM/Public/RigVMCore/RigVM.h
+++ b/Engine/Source/Runtime/RigVM/Public/RigVMCore/RigVM.h
@@ -732,10 +732,10 @@ private:
 	FORCEINLINE TArray<FName>& GetFunctionNames() { return *FunctionNamesPtr; }
 	FORCEINLINE const TArray<FName>& GetFunctionNames() const { return *FunctionNamesPtr; }
 
-	TArray<FRigVMFunctionPtr> FunctionsStorage;
-	TArray<FRigVMFunctionPtr>* FunctionsPtr;
-	FORCEINLINE TArray<FRigVMFunctionPtr>& GetFunctions() { return *FunctionsPtr; }
-	FORCEINLINE const TArray<FRigVMFunctionPtr>& GetFunctions() const { return *FunctionsPtr; }
+	TArray<const FRigVMFunction*> FunctionsStorage;
+	TArray<const FRigVMFunction*>* FunctionsPtr;
+	FORCEINLINE TArray<const FRigVMFunction*>& GetFunctions() { return *FunctionsPtr; }
+	FORCEINLINE const TArray<const FRigVMFunction*>& GetFunctions() const { return *FunctionsPtr; }
 
 	TArray<const FRigVMDispatchFactory*> FactoriesStorage;
 	TArray<const FRigVMDispatchFactory*>* FactoriesPtr;
diff --git a/Engine/Source/Runtime/RigVM/Public/RigVMCore/RigVMByteCode.h b/Engine/Source/Runtime/RigVM/Public/RigVMCore/RigVMByteCode.h
--- a/Engine/Source/Runtime/RigVM/Public/RigVMCore/RigVMByteCode.h
+++ b/Engine/Source/Runtime/RigVM/Public/RigVMCore/RigVMByteCode.h
@@ -53,6 +53,11 @@ struct FRigVMBranchInfoKey
 		, Label(InLabel)
 	{}
 
+	bool IsValid() const
+	{
+		return InstructionIndex != INDEX_NONE && ArgumentIndex != INDEX_NONE && !Label.IsNone();
+	}
+
 	friend uint32 GetTypeHash(const FRigVMBranchInfoKey& InKey)
 	{
 		return HashCombine(
