RigVM: Inline branches as part of function inlining

diff --git a/Engine/Source/Developer/RigVMDeveloper/Private/RigVMCompiler/RigVMCompiler.cpp b/Engine/Source/Developer/RigVMDeveloper/Private/RigVMCompiler/RigVMCompiler.cpp
--- a/Engine/Source/Developer/RigVMDeveloper/Private/RigVMCompiler/RigVMCompiler.cpp
+++ b/Engine/Source/Developer/RigVMDeveloper/Private/RigVMCompiler/RigVMCompiler.cpp
@@ -1642,7 +1642,7 @@ int32 URigVMCompiler::TraverseInlineFunction(const FRigVMInlineFunctionExprAST*
 	
 	int32 InstructionIndexStart = INDEX_NONE;
 	int32 InstructionIndexEnd = INDEX_NONE;
-
+	int32 BranchIndexStart = INDEX_NONE;
 	
 	FString FunctionHash = FunctionReferenceNode->GetReferencedFunctionHeader().GetHash();
 	if (!CompiledFunctions.Contains(FunctionHash))
@@ -1769,6 +1769,18 @@ int32 URigVMCompiler::TraverseInlineFunction(const FRigVMInlineFunctionExprAST*
 		InstructionIndexEnd = ByteCode.GetNumInstructions() - 1;
 		FRigVMCallstack FuncRefCallstack = InExpr->GetProxy().GetCallstack();
 
+		for(const FRigVMBranchInfo& BranchInfo : FunctionByteCode.BranchInfos)
+		{
+			const int32 BranchInfoIndex = ByteCode.AddBranchInfo(BranchInfo);
+			if(BranchIndexStart == INDEX_NONE)
+			{
+				BranchIndexStart = BranchInfoIndex;
+			}
+			ByteCode.BranchInfos[BranchInfoIndex].InstructionIndex += InstructionIndexStart;
+			ByteCode.BranchInfos[BranchInfoIndex].FirstInstruction += (uint16)InstructionIndexStart;
+			ByteCode.BranchInfos[BranchInfoIndex].LastInstruction += (uint16)InstructionIndexStart;
+		}
+
 		// For each instruction, substitute the operand for the one used in the current bytecode
 		const FRigVMInstructionArray FunctionInstructions = FunctionByteCode.GetInstructions();
 		FRigVMInstructionArray Instructions = ByteCode.GetInstructions();
@@ -1854,6 +1866,12 @@ int32 URigVMCompiler::TraverseInlineFunction(const FRigVMInlineFunctionExprAST*
 				Op.FunctionIndex = FunctionIndex;
 			}
 
+			if (Instruction.OpCode == ERigVMOpCode::JumpToBranch)
+			{
+				FRigVMJumpToBranchOp& Op = ByteCode.GetOpAt<FRigVMJumpToBranchOp>(Instruction);
+				Op.FirstBranchInfoIndex = Op.FirstBranchInfoIndex + BranchIndexStart;
+			}
+
 			if (Settings.SetupNodeInstructionIndex)
 			{
 				if (const TArray<UObject*>* Callstack = FunctionByteCode.GetCallstackForInstruction(i-InstructionIndexStart))
diff --git a/Engine/Source/Runtime/RigVM/Private/RigVMCore/RigVM.cpp b/Engine/Source/Runtime/RigVM/Private/RigVMCore/RigVM.cpp
--- a/Engine/Source/Runtime/RigVM/Private/RigVMCore/RigVM.cpp
+++ b/Engine/Source/Runtime/RigVM/Private/RigVMCore/RigVM.cpp
@@ -1261,7 +1261,7 @@ void URigVM::CacheMemoryHandlesIfRequired(TArrayView<URigVMMemoryStorage*> InMem
 				{
 					const FRigVMJumpToBranchOp& Op = ByteCode.GetOpAt<FRigVMJumpToBranchOp>(Instructions[InstructionIndex]);
 					const FRigVMOperand& Arg = Op.Arg;
-					InOpFunc(InHandleBaseIndex, {InstructionIndex, 0}, Arg);
+					InOpFunc(InHandleBaseIndex, {}, Arg);
 					break;
 				}
 			case ERigVMOpCode::Invalid:
@@ -2861,6 +2861,12 @@ ERigVMExecuteResult URigVM::ExecuteInstructions(int32 InFirstInstruction, int32
 				// find the one with the right label - then jump to the branch.
 				bool bBranchFound = false;
 				const TArray<FRigVMBranchInfo>& Branches = ByteCode.BranchInfos;
+				if(Branches.IsEmpty())
+				{
+					UE_LOG(LogRigVM, Error, TEXT("No branches in ByteCode - but JumpToBranch instruction %d found. Likely a corrupt VM. Exiting."), ContextPublicData.InstructionIndex);
+					return CurrentExecuteResult = ERigVMExecuteResult::Failed;
+				}
+					
 				for(int32 BranchIndex = Op.FirstBranchInfoIndex // start at the first branch known to this jump op
 					; BranchIndex < Branches.Num(); BranchIndex++)
 				{
@@ -2878,7 +2884,8 @@ ERigVMExecuteResult URigVM::ExecuteInstructions(int32 InFirstInstruction, int32
 				}
 				if (!bBranchFound)
 				{
-					ContextPublicData.InstructionIndex++;
+					UE_LOG(LogRigVM, Error, TEXT("Branch '%s' was not found for instruction %d. Exiting."), *BranchLabel.ToString(), ContextPublicData.InstructionIndex);
+					return CurrentExecuteResult = ERigVMExecuteResult::Failed;
 				}
 				break;
 			}
@@ -3482,7 +3489,6 @@ void URigVM::CacheSingleMemoryHandle(int32 InHandleIndex, const FRigVMBranchInfo
 	// if we are lazy executing update the handle to point to a lazy branch
 	if(InBranchInfoKey.IsValid())
 	{
-		UE_LOG(LogRigVM, Warning, TEXT("Looking up branch info key instruction %d, argument %d"), InBranchInfoKey.InstructionIndex, InBranchInfoKey.ArgumentIndex);
 		if(const FRigVMBranchInfo* BranchInfo = GetByteCode().GetBranchInfo(InBranchInfoKey))
 		{
 			FRigVMLazyBranch* LazyBranch = &LazyBranches[BranchInfo->Index];
diff --git a/Engine/Source/Runtime/RigVM/Private/RigVMCore/RigVMByteCode.cpp b/Engine/Source/Runtime/RigVM/Private/RigVMCore/RigVMByteCode.cpp
--- a/Engine/Source/Runtime/RigVM/Private/RigVMCore/RigVMByteCode.cpp
+++ b/Engine/Source/Runtime/RigVM/Private/RigVMCore/RigVMByteCode.cpp
@@ -783,6 +783,11 @@ uint32 FRigVMByteCode::GetByteCodeHash() const
 		Hash = HashCombine(Hash, GetOperatorHash(Instruction));
 	}
 
+	for(const FRigVMBranchInfo& BranchInfo : BranchInfos)
+	{
+		Hash = HashCombine(Hash, GetTypeHash(BranchInfo));
+	}
+
 	return Hash;
 }
 
@@ -1822,6 +1827,11 @@ FRigVMOperandArray FRigVMByteCode::GetOperandsForOp(const FRigVMInstruction& InI
 			const FRigVMSenaryOp& Op = GetOpAt<FRigVMSenaryOp>(InInstruction.ByteCodeIndex);
 			return FRigVMOperandArray(&Op.ArgA, 6);
 		}
+		case ERigVMOpCode::JumpToBranch:
+		{
+			const FRigVMJumpToBranchOp& Op = GetOpAt<FRigVMJumpToBranchOp>(InInstruction.ByteCodeIndex);
+			return FRigVMOperandArray(&Op.Arg, 1);
+		}
 		case ERigVMOpCode::JumpAbsolute:
 		case ERigVMOpCode::JumpForward:
 		case ERigVMOpCode::JumpBackward:
@@ -1910,6 +1920,11 @@ uint64 FRigVMByteCode::GetFirstOperandByteIndex(const FRigVMInstruction& InInstr
 			const FRigVMSenaryOp& Op = GetOpAt<FRigVMSenaryOp>(InInstruction.ByteCodeIndex);
 			return InInstruction.ByteCodeIndex + ((uint64)&Op.ArgA - (uint64)&Op);
 		}
+		case ERigVMOpCode::JumpToBranch:
+		{
+			const FRigVMJumpToBranchOp& Op = GetOpAt<FRigVMJumpToBranchOp>(InInstruction.ByteCodeIndex);
+			return InInstruction.ByteCodeIndex + ((uint64)&Op.Arg - (uint64)&Op);
+		}
 		case ERigVMOpCode::JumpAbsolute:
 		case ERigVMOpCode::JumpForward:
 		case ERigVMOpCode::JumpBackward:
diff --git a/Engine/Source/Runtime/RigVM/Public/RigVMCore/RigVMMemoryStorage.h b/Engine/Source/Runtime/RigVM/Public/RigVMCore/RigVMMemoryStorage.h
--- a/Engine/Source/Runtime/RigVM/Public/RigVMCore/RigVMMemoryStorage.h
+++ b/Engine/Source/Runtime/RigVM/Public/RigVMCore/RigVMMemoryStorage.h
@@ -74,6 +74,17 @@ struct RIGVM_API FRigVMBranchInfo
 		return Ar;
 	}
 
+	friend uint32 GetTypeHash(const FRigVMBranchInfo& InBranchInfo)
+	{
+		uint32 Hash = GetTypeHash(GetTypeHash(InBranchInfo.Index));
+		Hash = HashCombine(Hash, GetTypeHash(InBranchInfo.Label));
+		Hash = HashCombine(Hash, GetTypeHash(InBranchInfo.InstructionIndex));
+		Hash = HashCombine(Hash, GetTypeHash(InBranchInfo.ArgumentIndex));
+		Hash = HashCombine(Hash, GetTypeHash(InBranchInfo.FirstInstruction));
+		Hash = HashCombine(Hash, GetTypeHash(InBranchInfo.LastInstruction));
+		return Hash;
+	}
+
 	UPROPERTY()
 	int32 Index;
 
