fixed CIS issues

diff --git a/Engine/Source/Runtime/RigVM/Public/RigVMCore/RigVMGraphFunctionDefinition.h b/Engine/Source/Runtime/RigVM/Public/RigVMCore/RigVMGraphFunctionDefinition.h
--- a/Engine/Source/Runtime/RigVM/Public/RigVMCore/RigVMGraphFunctionDefinition.h
+++ b/Engine/Source/Runtime/RigVM/Public/RigVMCore/RigVMGraphFunctionDefinition.h
@@ -348,7 +348,7 @@ struct RIGVM_API FRigVMGraphFunctionHeader
 	FString NodeTitle;
 	
 	UPROPERTY()
-	FLinearColor NodeColor;
+	FLinearColor NodeColor = FLinearColor::White;
 
 	UPROPERTY()
 	FText Tooltip;
