ControlRig, RigVM: Remove use of FORCEINLINE

diff --git a/Engine/Source/Runtime/RigVM/Public/RigVMCore/RigVMMemoryDeprecated.h b/Engine/Source/Runtime/RigVM/Public/RigVMCore/RigVMMemoryDeprecated.h
--- a/Engine/Source/Runtime/RigVM/Public/RigVMCore/RigVMMemoryDeprecated.h
+++ b/Engine/Source/Runtime/RigVM/Public/RigVMCore/RigVMMemoryDeprecated.h
@@ -141,50 +141,50 @@ struct RIGVM_API FRigVMRegister
 	void Serialize(FArchive& Ar);
 	void Save(FArchive& Ar);
 	void Load(FArchive& Ar);
-	FORCEINLINE_DEBUGGABLE friend FArchive& operator<<(FArchive& Ar, FRigVMRegister& P)
+	friend FArchive& operator<<(FArchive& Ar, FRigVMRegister& P)
 	{
 		P.Serialize(Ar);
 		return Ar;
 	}
 
 	// returns true if this register is using a dynamic array for storage
-	FORCEINLINE_DEBUGGABLE bool IsDynamic() const { return bIsDynamic; }
+	bool IsDynamic() const { return bIsDynamic; }
 
 	// returns true if this register is using a dynamic array for storage
-	FORCEINLINE_DEBUGGABLE bool IsNestedDynamic() const { return bIsDynamic && bIsArray; }
+	bool IsNestedDynamic() const { return bIsDynamic && bIsArray; }
 
 	// returns the current address of the register within the data byte array.
 	// this can change over time - as the register is moving through slices.
 	// use GetFirstAllocatedByte to get the fixed first byte.
-	FORCEINLINE_DEBUGGABLE uint64 GetWorkByteIndex(int32 InSliceIndex = 0) const
+	uint64 GetWorkByteIndex(int32 InSliceIndex = 0) const
 	{
 		ensure(InSliceIndex >= 0);
 		return ByteIndex + ((uint64)InSliceIndex * GetNumBytesPerSlice());
 	}
 
 	// returns the first allocated byte in the data byte array
-	FORCEINLINE_DEBUGGABLE uint64 GetFirstAllocatedByte() const
+	uint64 GetFirstAllocatedByte() const
 	{ 
 		return ByteIndex - (uint64)AlignmentBytes;
 	}
 
 	// Returns the leading alignment bytes
-	FORCEINLINE_DEBUGGABLE uint8 GetAlignmentBytes() const { return AlignmentBytes; }
+	uint8 GetAlignmentBytes() const { return AlignmentBytes; }
 
 	// Returns true if the register stores more than one element
-	FORCEINLINE_DEBUGGABLE bool IsArray() const { return bIsArray || (ElementCount > 1); }
+	bool IsArray() const { return bIsArray || (ElementCount > 1); }
 
 	// Returns the number of allocated bytes (including alignment + trailing bytes)
-	FORCEINLINE_DEBUGGABLE uint16 GetAllocatedBytes() const { return ElementCount * ElementSize * SliceCount + (uint16)AlignmentBytes + TrailingBytes; }
+	uint16 GetAllocatedBytes() const { return ElementCount * ElementSize * SliceCount + (uint16)AlignmentBytes + TrailingBytes; }
 
 	// Returns the number of bytes for a complete slice
-	FORCEINLINE_DEBUGGABLE uint16 GetNumBytesPerSlice() const { return ElementCount * ElementSize; }
+	uint16 GetNumBytesPerSlice() const { return ElementCount * ElementSize; }
 
 	// Returns the number of bytes for all slices
-	FORCEINLINE_DEBUGGABLE uint16 GetNumBytesAllSlices() const { return ElementCount * ElementSize * SliceCount; }
+	uint16 GetNumBytesAllSlices() const { return ElementCount * ElementSize * SliceCount; }
 
 	// Returns the total number of elements (elementcount * slicecount) in the register
-	FORCEINLINE_DEBUGGABLE uint32 GetTotalElementCount() const { return (uint32)ElementCount * (uint32)SliceCount; }
+	uint32 GetTotalElementCount() const { return (uint32)ElementCount * (uint32)SliceCount; }
 
 };
 
@@ -215,7 +215,7 @@ public:
 	void Serialize(FArchive& Ar);
 	void Save(FArchive& Ar);
 	void Load(FArchive& Ar);
-	FORCEINLINE_DEBUGGABLE friend FArchive& operator<<(FArchive& Ar, FRigVMRegisterOffset& P)
+	friend FArchive& operator<<(FArchive& Ar, FRigVMRegisterOffset& P)
 	{
 		P.Serialize(Ar);
 		return Ar;
@@ -237,11 +237,11 @@ public:
 
 	bool operator == (const FRigVMRegisterOffset& InOther) const;
 
-	FORCEINLINE_DEBUGGABLE bool IsValid() const { return Type != ERigVMRegisterType::Invalid; }
-	FORCEINLINE_DEBUGGABLE ERigVMRegisterType GetType() const { return Type; }
-	FORCEINLINE_DEBUGGABLE FName GetCPPType() const { return CPPType; }
-	FORCEINLINE_DEBUGGABLE FString GetCachedSegmentPath() const { return CachedSegmentPath; }
-	FORCEINLINE_DEBUGGABLE int32 GetArrayIndex() const { return ArrayIndex; }
+	bool IsValid() const { return Type != ERigVMRegisterType::Invalid; }
+	ERigVMRegisterType GetType() const { return Type; }
+	FName GetCPPType() const { return CPPType; }
+	FString GetCachedSegmentPath() const { return CachedSegmentPath; }
+	int32 GetArrayIndex() const { return ArrayIndex; }
 	uint16 GetElementSize() const;
 	void SetElementSize(uint16 InElementSize) { ElementSize = InElementSize; };
 	UScriptStruct* GetScriptStruct() const;
@@ -328,16 +328,16 @@ public:
 	FRigVMMemoryContainer& operator= (const FRigVMMemoryContainer &InOther);
 
 	// returns the memory type of this container
-	FORCEINLINE_DEBUGGABLE ERigVMMemoryType GetMemoryType() const { return MemoryType;  }
+	ERigVMMemoryType GetMemoryType() const { return MemoryType;  }
 
 	// sets the memory type. should only be used when the container is empty
-	FORCEINLINE_DEBUGGABLE void SetMemoryType(ERigVMMemoryType InMemoryType) { MemoryType = InMemoryType; }
+	void SetMemoryType(ERigVMMemoryType InMemoryType) { MemoryType = InMemoryType; }
 
 	// returns true if this container supports name based lookup
-	FORCEINLINE_DEBUGGABLE bool SupportsNames() const { return bUseNameMap;  }
+	bool SupportsNames() const { return bUseNameMap;  }
 
 	// returns the number of registers in this container
-	FORCEINLINE_DEBUGGABLE int32 Num() const { return Registers.Num(); }
+	int32 Num() const { return Registers.Num(); }
 
 	// resets the container but maintains storage.
 	void Reset();
@@ -346,50 +346,50 @@ public:
 	void Empty();
 
 	// const accessor for a register based on index
-	FORCEINLINE_DEBUGGABLE const FRigVMRegister& operator[](int32 InIndex) const { return GetRegister(InIndex); }
+	const FRigVMRegister& operator[](int32 InIndex) const { return GetRegister(InIndex); }
 
 	// accessor for a register based on index
-	FORCEINLINE_DEBUGGABLE FRigVMRegister& operator[](int32 InIndex) { return GetRegister(InIndex); }
+	FRigVMRegister& operator[](int32 InIndex) { return GetRegister(InIndex); }
 
 	// const accessor for a register based on an argument
-	FORCEINLINE_DEBUGGABLE const FRigVMRegister& operator[](const FRigVMOperand& InArg) const { return GetRegister(InArg); }
+	const FRigVMRegister& operator[](const FRigVMOperand& InArg) const { return GetRegister(InArg); }
 
 	// accessor for a register based on an argument
-	FORCEINLINE_DEBUGGABLE FRigVMRegister& operator[](const FRigVMOperand& InArg) { return GetRegister(InArg); }
+	FRigVMRegister& operator[](const FRigVMOperand& InArg) { return GetRegister(InArg); }
 
 	// const accessor for a register based on a a name. note: only works if SupportsNames() == true
-	FORCEINLINE_DEBUGGABLE const FRigVMRegister& operator[](const FName& InName) const { return GetRegister(InName); }
+	const FRigVMRegister& operator[](const FName& InName) const { return GetRegister(InName); }
 
 	// accessor for a register based on a a name. note: only works if SupportsNames() == true
-	FORCEINLINE_DEBUGGABLE FRigVMRegister& operator[](const FName& InName) { return GetRegister(InName); }
+	FRigVMRegister& operator[](const FName& InName) { return GetRegister(InName); }
 
-	FORCEINLINE_DEBUGGABLE TArray<FRigVMRegister>::RangedForIteratorType      begin() { return Registers.begin(); }
-	FORCEINLINE_DEBUGGABLE TArray<FRigVMRegister>::RangedForConstIteratorType begin() const { return Registers.begin(); }
-	FORCEINLINE_DEBUGGABLE TArray<FRigVMRegister>::RangedForIteratorType      end() { return Registers.end(); }
-	FORCEINLINE_DEBUGGABLE TArray<FRigVMRegister>::RangedForConstIteratorType end() const { return Registers.end(); }
+	TArray<FRigVMRegister>::RangedForIteratorType      begin() { return Registers.begin(); }
+	TArray<FRigVMRegister>::RangedForConstIteratorType begin() const { return Registers.begin(); }
+	TArray<FRigVMRegister>::RangedForIteratorType      end() { return Registers.end(); }
+	TArray<FRigVMRegister>::RangedForConstIteratorType end() const { return Registers.end(); }
 
 	// const accessor for a register based on index
-	FORCEINLINE_DEBUGGABLE const FRigVMRegister& GetRegister(int32 InIndex) const { return Registers[InIndex]; }
+	const FRigVMRegister& GetRegister(int32 InIndex) const { return Registers[InIndex]; }
 
 	// accessor for a register based on index
-	FORCEINLINE_DEBUGGABLE FRigVMRegister& GetRegister(int32 InIndex) { return Registers[InIndex]; }
+	FRigVMRegister& GetRegister(int32 InIndex) { return Registers[InIndex]; }
 
 	// const accessor for a register based on an argument
-	FORCEINLINE_DEBUGGABLE const FRigVMRegister& GetRegister(const FRigVMOperand& InArg) const { return Registers[InArg.GetRegisterIndex()]; }
+	const FRigVMRegister& GetRegister(const FRigVMOperand& InArg) const { return Registers[InArg.GetRegisterIndex()]; }
 
 	// accessor for a register based on an argument
-	FORCEINLINE_DEBUGGABLE FRigVMRegister& GetRegister(const FRigVMOperand& InArg) { return Registers[InArg.GetRegisterIndex()]; }
+	FRigVMRegister& GetRegister(const FRigVMOperand& InArg) { return Registers[InArg.GetRegisterIndex()]; }
 
 	// const accessor for a register based on a a name. note: only works if SupportsNames() == true
-	FORCEINLINE_DEBUGGABLE const FRigVMRegister& GetRegister(const FName& InName) const { return Registers[GetIndex(InName)]; }
+	const FRigVMRegister& GetRegister(const FName& InName) const { return Registers[GetIndex(InName)]; }
 
 	// accessor for a register based on a a name. note: only works if SupportsNames() == true
-	FORCEINLINE_DEBUGGABLE FRigVMRegister& GetRegister(const FName& InName) { return Registers[GetIndex(InName)]; }
+	FRigVMRegister& GetRegister(const FName& InName) { return Registers[GetIndex(InName)]; }
 	
 	void Serialize(FArchive& Ar);
 	void Save(FArchive& Ar);
 	void Load(FArchive& Ar);
-	FORCEINLINE_DEBUGGABLE friend FArchive& operator<<(FArchive& Ar, FRigVMMemoryContainer& P)
+	friend FArchive& operator<<(FArchive& Ar, FRigVMMemoryContainer& P)
 	{
 		P.Serialize(Ar);
 		return Ar;
@@ -397,7 +397,7 @@ public:
 
 	// Returns an argument for a given register.
 	// This is typically used to store a light weight address for use within a VM.
-	FORCEINLINE_DEBUGGABLE FRigVMOperand GetOperand(int32 InRegisterIndex, int32 InRegisterOffset)
+	FRigVMOperand GetOperand(int32 InRegisterIndex, int32 InRegisterOffset)
 	{
 		ensure(Registers.IsValidIndex(InRegisterIndex));
 		return FRigVMOperand(MemoryType, InRegisterIndex, InRegisterOffset);
@@ -405,7 +405,7 @@ public:
 
 	// Returns an argument for a given register.
 	// This is typically used to store a light weight address for use within a VM.
-	FORCEINLINE_DEBUGGABLE FRigVMOperand GetOperand(int32 InRegisterIndex, const FString& InSegmentPath = FString(), int32 InArrayElement = INDEX_NONE)
+	FRigVMOperand GetOperand(int32 InRegisterIndex, const FString& InSegmentPath = FString(), int32 InArrayElement = INDEX_NONE)
 	{
 		ensure(Registers.IsValidIndex(InRegisterIndex));
 		// Register offset must hold on to the ScriptStruct such that it can recalculate the struct size after cook
@@ -424,7 +424,7 @@ public:
 
 private:
 
-	FORCEINLINE_DEBUGGABLE uint8* GetDataPtr(const FRigVMRegister& Register, int32 InRegisterOffset = INDEX_NONE, int32 InSliceIndex = 0, bool bArrayContent = false) const
+	uint8* GetDataPtr(const FRigVMRegister& Register, int32 InRegisterOffset = INDEX_NONE, int32 InSliceIndex = 0, bool bArrayContent = false) const
 	{
 		if (Register.ElementCount == 0 && !Register.IsNestedDynamic())
 		{
@@ -476,7 +476,7 @@ private:
 public:
 	
 	// Returns the script struct used for a given register (can be nullptr for non-struct-registers).
-	FORCEINLINE_DEBUGGABLE UScriptStruct* GetScriptStruct(const FRigVMRegister& Register) const
+	UScriptStruct* GetScriptStruct(const FRigVMRegister& Register) const
 	{
 		if (Register.ScriptStructIndex != INDEX_NONE)
 		{
@@ -487,7 +487,7 @@ public:
 	}
 
 	// Returns the script struct used for a given register index (can be nullptr for non-struct-registers).
-	FORCEINLINE_DEBUGGABLE UScriptStruct* GetScriptStruct(int32 InRegisterIndex, int32 InRegisterOffset = INDEX_NONE) const
+	UScriptStruct* GetScriptStruct(int32 InRegisterIndex, int32 InRegisterOffset = INDEX_NONE) const
 	{
 		if (InRegisterOffset == INDEX_NONE)
 		{
@@ -502,7 +502,7 @@ public:
 	
 	// Returns the index of a register based on the register name.
 	// Note: This only works if SupportsNames() == true
-	FORCEINLINE_DEBUGGABLE int32 GetIndex(const FName& InName) const
+	int32 GetIndex(const FName& InName) const
 	{
 		if (!bUseNameMap)
 		{
@@ -532,7 +532,7 @@ public:
 	}
 
 	// Returns true if a given name is available for a new register.
-	FORCEINLINE_DEBUGGABLE bool IsNameAvailable(const FName& InPotentialNewName) const
+	bool IsNameAvailable(const FName& InPotentialNewName) const
 	{
 		if (!bUseNameMap)
 		{
diff --git a/Engine/Source/Runtime/RigVM/Public/RigVMCore/RigVMMemoryStorage.h b/Engine/Source/Runtime/RigVM/Public/RigVMCore/RigVMMemoryStorage.h
--- a/Engine/Source/Runtime/RigVM/Public/RigVMCore/RigVMMemoryStorage.h
+++ b/Engine/Source/Runtime/RigVM/Public/RigVMCore/RigVMMemoryStorage.h
@@ -78,7 +78,7 @@ struct RIGVM_API FRigVMBranchInfo
 	}
 	
 	void Serialize(FArchive& Ar);
-	FORCEINLINE friend FArchive& operator<<(FArchive& Ar, FRigVMBranchInfo& P)
+	friend FArchive& operator<<(FArchive& Ar, FRigVMBranchInfo& P)
 	{
 		P.Serialize(Ar);
 		return Ar;
@@ -183,7 +183,7 @@ public:
 		, ConstValue(InConstValue)
 	{}
 
-	FORCEINLINE const ComputedType& Get() const
+	const ComputedType& Get() const
 	{
 		if(MemoryHandle)
 		{
@@ -226,7 +226,7 @@ struct FRigVMMemoryHandle
 public:
 
 	// Default constructor
-	FORCEINLINE_DEBUGGABLE FRigVMMemoryHandle()
+	FRigVMMemoryHandle()
 		: Ptr(nullptr)
 		, Property(nullptr) 
 		, PropertyPath(nullptr)
@@ -234,7 +234,7 @@ public:
 	{}
 
 	// Constructor from complete data
-	FORCEINLINE_DEBUGGABLE FRigVMMemoryHandle(uint8* InPtr, const FProperty* InProperty,  const FRigVMPropertyPath* InPropertyPath, FRigVMLazyBranch* InLazyBranch = nullptr)
+	FRigVMMemoryHandle(uint8* InPtr, const FProperty* InProperty,  const FRigVMPropertyPath* InPropertyPath, FRigVMLazyBranch* InLazyBranch = nullptr)
 		: Ptr(InPtr)
 		, Property(InProperty)
 		, PropertyPath(InPropertyPath)
@@ -248,7 +248,7 @@ public:
 	 * @param InSliceIndex If this is != INDEX_NONE the memory handle will return the slice of the memory requested
 	 * @return The traversed memory cached by this handle.
 	 */
-	FORCEINLINE uint8* GetData(bool bFollowPropertyPath = false, int32 InSliceIndex = INDEX_NONE)
+	uint8* GetData(bool bFollowPropertyPath = false, int32 InSliceIndex = INDEX_NONE)
 	{
 		return GetData_Internal(bFollowPropertyPath, InSliceIndex);
 	}
@@ -259,7 +259,7 @@ public:
 	 * @param InSliceIndex If this is != INDEX_NONE the memory handle will return the slice of the memory requested
 	 * @return The traversed memory cached by this handle (const)
 	 */
-	FORCEINLINE const uint8* GetData(bool bFollowPropertyPath = false, int32 InSliceIndex = INDEX_NONE) const
+	const uint8* GetData(bool bFollowPropertyPath = false, int32 InSliceIndex = INDEX_NONE) const
 	{
 		return GetData_Internal(bFollowPropertyPath, InSliceIndex);
 	}
@@ -271,7 +271,7 @@ public:
 	 * @return The traversed memory cached by this handle (const)
 	 */
 	template<typename ComputedType>
-	FORCEINLINE TRigVMLazyValue<ComputedType> GetDataLazily(bool bFollowPropertyPath = false, int32 InSliceIndex = INDEX_NONE) const
+	TRigVMLazyValue<ComputedType> GetDataLazily(bool bFollowPropertyPath = false, int32 InSliceIndex = INDEX_NONE) const
 	{
 		return GetDataLazily_Internal(bFollowPropertyPath, InSliceIndex);
 	}
@@ -280,7 +280,7 @@ public:
 	 * Returns true if this memory handle depends on a lazily executed branch
 	 * @return True if this memory handle depends on a lazily executed branch
 	 */
-	FORCEINLINE bool IsLazy() const
+	bool IsLazy() const
 	{
 		return LazyBranch != nullptr;
 	}
@@ -292,19 +292,19 @@ public:
 	bool ComputeLazyValueIfNecessary(int32 InSliceIndex = INDEX_NONE);
 
 	// Returns the head property of this handle
-	FORCEINLINE const FProperty* GetProperty() const
+	const FProperty* GetProperty() const
 	{
 		return Property;
 	}
 
 	// Returns the [optional] property path used within this handle
-	FORCEINLINE const FRigVMPropertyPath* GetPropertyPath() const
+	const FRigVMPropertyPath* GetPropertyPath() const
 	{
 		return PropertyPath;
 	}
 
 	// Returns the [optional] property path used within this handle (ref)
-	FORCEINLINE const FRigVMPropertyPath& GetPropertyPathRef() const
+	const FRigVMPropertyPath& GetPropertyPathRef() const
 	{
 		if(PropertyPath)
 		{
@@ -314,7 +314,7 @@ public:
 	}
 
 	// Returns the resolved property the data is pointing to
-	FORCEINLINE const FProperty* GetResolvedProperty(bool bIsHiddenArgument = false) const
+	const FProperty* GetResolvedProperty(bool bIsHiddenArgument = false) const
 	{
 		if(bIsHiddenArgument)
 		{
@@ -330,7 +330,7 @@ public:
 
 	// Returns true if this memory handle maps to a given type of property
 	template<typename PropertyType>
-	FORCEINLINE bool IsPropertyType(bool bIsHiddenArgument = false) const
+	bool IsPropertyType(bool bIsHiddenArgument = false) const
 	{
 		const FProperty* ResolvedProperty = GetResolvedProperty(bIsHiddenArgument);
 		return ResolvedProperty->IsA<PropertyType>();
@@ -338,7 +338,7 @@ public:
 
 	// Returns true if this memory handle maps to a given array type of property
 	template<typename PropertyType>
-	FORCEINLINE bool IsPropertyArrayType(bool bIsHiddenArgument = false) const
+	bool IsPropertyArrayType(bool bIsHiddenArgument = false) const
 	{
 		const FProperty* ResolvedProperty = GetResolvedProperty(bIsHiddenArgument);
 		if(const FProperty* ElementProperty = GetArrayElementProperty(ResolvedProperty))
@@ -353,7 +353,7 @@ public:
 		typename T,
 		typename TEnableIf<TIsTArray<T>::Value>::Type* = nullptr
 	>
-	FORCEINLINE bool IsType(bool bIsHiddenArgument = false) const
+	bool IsType(bool bIsHiddenArgument = false) const
 	{
 		typedef typename T::ElementType ElementType;
 		return IsTypeArray<ElementType>(bIsHiddenArgument);
@@ -364,7 +364,7 @@ public:
 		typename T,
 		typename TEnableIf<TRigVMIsBool<T>::Value, T>::Type* = nullptr
 	>
-	FORCEINLINE bool IsType(bool bIsHiddenArgument = false) const
+	bool IsType(bool bIsHiddenArgument = false) const
 	{
 		return IsBool(bIsHiddenArgument);
 	}
@@ -374,7 +374,7 @@ public:
 		typename T,
 		typename TEnableIf<TRigVMIsBool<T>::Value, T>::Type* = nullptr
 	>
-	FORCEINLINE bool IsTypeArray(bool bIsHiddenArgument = false) const
+	bool IsTypeArray(bool bIsHiddenArgument = false) const
 	{
 		return IsBoolArray(bIsHiddenArgument);
 	}
@@ -384,7 +384,7 @@ public:
 		typename T,
 		typename TEnableIf<TRigVMIsFloat<T>::Value, T>::Type* = nullptr
 	>
-	FORCEINLINE bool IsType(bool bIsHiddenArgument = false) const
+	bool IsType(bool bIsHiddenArgument = false) const
 	{
 		return IsFloat(bIsHiddenArgument);
 	}
@@ -394,7 +394,7 @@ public:
 		typename T,
 		typename TEnableIf<TRigVMIsFloat<T>::Value, T>::Type* = nullptr
 	>
-	FORCEINLINE bool IsTypeArray(bool bIsHiddenArgument = false) const
+	bool IsTypeArray(bool bIsHiddenArgument = false) const
 	{
 		return IsFloatArray(bIsHiddenArgument);
 	}
@@ -404,7 +404,7 @@ public:
 		typename T,
 		typename TEnableIf<TRigVMIsDouble<T>::Value, T>::Type* = nullptr
 	>
-	FORCEINLINE bool IsType(bool bIsHiddenArgument = false) const
+	bool IsType(bool bIsHiddenArgument = false) const
 	{
 		return IsDouble(bIsHiddenArgument);
 	}
@@ -414,7 +414,7 @@ public:
 		typename T,
 		typename TEnableIf<TRigVMIsDouble<T>::Value, T>::Type* = nullptr
 	>
-	FORCEINLINE bool IsTypeArray(bool bIsHiddenArgument = false) const
+	bool IsTypeArray(bool bIsHiddenArgument = false) const
 	{
 		return IsDoubleArray(bIsHiddenArgument);
 	}
@@ -424,7 +424,7 @@ public:
 		typename T,
 		typename TEnableIf<TRigVMIsInt32<T>::Value, T>::Type* = nullptr
 	>
-	FORCEINLINE bool IsType(bool bIsHiddenArgument = false) const
+	bool IsType(bool bIsHiddenArgument = false) const
 	{
 		return IsInt32(bIsHiddenArgument);
 	}
@@ -434,7 +434,7 @@ public:
 		typename T,
 		typename TEnableIf<TRigVMIsInt32<T>::Value, T>::Type* = nullptr
 	>
-	FORCEINLINE bool IsTypeArray(bool bIsHiddenArgument = false) const
+	bool IsTypeArray(bool bIsHiddenArgument = false) const
 	{
 		return IsInt32Array(bIsHiddenArgument);
 	}
@@ -443,7 +443,7 @@ public:
 		typename T,
 		typename TEnableIf<TRigVMIsName<T>::Value, T>::Type* = nullptr
 	>
-	FORCEINLINE bool IsType(bool bIsHiddenArgument = false) const
+	bool IsType(bool bIsHiddenArgument = false) const
 	{
 		return IsName(bIsHiddenArgument);
 	}
@@ -453,7 +453,7 @@ public:
 		typename T,
 		typename TEnableIf<TRigVMIsName<T>::Value, T>::Type* = nullptr
 	>
-	FORCEINLINE bool IsTypeArray(bool bIsHiddenArgument = false) const
+	bool IsTypeArray(bool bIsHiddenArgument = false) const
 	{
 		return IsNameArray(bIsHiddenArgument);
 	}
@@ -463,7 +463,7 @@ public:
 		typename T,
 		typename TEnableIf<TRigVMIsString<T>::Value, T>::Type* = nullptr
 	>
-	FORCEINLINE bool IsType(bool bIsHiddenArgument = false) const
+	bool IsType(bool bIsHiddenArgument = false) const
 	{
 		return IsString(bIsHiddenArgument);
 	}
@@ -473,7 +473,7 @@ public:
 		typename T,
 		typename TEnableIf<TRigVMIsString<T>::Value, T>::Type* = nullptr
 	>
-	FORCEINLINE bool IsTypeArray(bool bIsHiddenArgument = false) const
+	bool IsTypeArray(bool bIsHiddenArgument = false) const
 	{
 		return IsStringArray(bIsHiddenArgument);
 	}
@@ -483,7 +483,7 @@ public:
 		typename T,
 		typename TEnableIf<TRigVMIsBaseStructure<T>::Value, T>::Type* = nullptr
 	>
-	FORCEINLINE bool IsType(bool bIsHiddenArgument = false) const
+	bool IsType(bool bIsHiddenArgument = false) const
 	{
 		return IsStruct<T>(bIsHiddenArgument);
 	}
@@ -493,7 +493,7 @@ public:
 		typename T,
 		typename TEnableIf<TRigVMIsBaseStructure<T>::Value, T>::Type* = nullptr
 	>
-	FORCEINLINE bool IsTypeArray(bool bIsHiddenArgument = false) const
+	bool IsTypeArray(bool bIsHiddenArgument = false) const
 	{
 		return IsStructArray<T>(bIsHiddenArgument);
 	}
@@ -503,7 +503,7 @@ public:
 		typename T,
 		typename TEnableIf<TModels<CRigVMUStruct, T>::Value>::Type * = nullptr
 	>
-	FORCEINLINE bool IsType(bool bIsHiddenArgument = false) const
+	bool IsType(bool bIsHiddenArgument = false) const
 	{
 		return IsStruct<T>(bIsHiddenArgument);
 	}
@@ -513,7 +513,7 @@ public:
 		typename T,
 		typename TEnableIf<TModels<CRigVMUStruct, T>::Value>::Type * = nullptr
 	>
-	FORCEINLINE bool IsTypeArray(bool bIsHiddenArgument = false) const
+	bool IsTypeArray(bool bIsHiddenArgument = false) const
 	{
 		return IsStructArray<T>(bIsHiddenArgument);
 	}
@@ -523,7 +523,7 @@ public:
 		typename T,
 		typename TEnableIf<TIsEnum<T>::Value>::Type* = nullptr
 	>
-	FORCEINLINE bool IsType(bool bIsHiddenArgument = false) const
+	bool IsType(bool bIsHiddenArgument = false) const
 	{
 		return IsEnum<T>(bIsHiddenArgument);
 	}
@@ -533,7 +533,7 @@ public:
 		typename T,
 		typename TEnableIf<TIsEnum<T>::Value>::Type* = nullptr
 	>
-	FORCEINLINE bool IsTypeArray(bool bIsHiddenArgument = false) const
+	bool IsTypeArray(bool bIsHiddenArgument = false) const
 	{
 		return IsEnumArray<T>(bIsHiddenArgument);
 	}
@@ -543,7 +543,7 @@ public:
 		typename T,
 		typename TEnableIf<TModels<CRigVMUClass, T>::Value>::Type * = nullptr
 	>
-	FORCEINLINE bool IsType(bool bIsHiddenArgument = false) const
+	bool IsType(bool bIsHiddenArgument = false) const
 	{
 		return IsObject<T>(bIsHiddenArgument);
 	}
@@ -553,92 +553,92 @@ public:
 		typename T,
 		typename TEnableIf<TModels<CRigVMUClass, T>::Value>::Type * = nullptr
 	>
-	FORCEINLINE bool IsTypeArray(bool bIsHiddenArgument = false) const
+	bool IsTypeArray(bool bIsHiddenArgument = false) const
 	{
 		return IsObjectArray<T>(bIsHiddenArgument);
 	}
 	
 	// Returns true if this memory handle maps to a bool property
-	FORCEINLINE bool IsBool(bool bIsHiddenArgument = false) const
+	bool IsBool(bool bIsHiddenArgument = false) const
 	{
 		return IsPropertyType<FBoolProperty>(bIsHiddenArgument);
 	}
 
 	// Returns true if this memory handle maps to a bool array property
-	FORCEINLINE bool IsBoolArray(bool bIsHiddenArgument = false) const
+	bool IsBoolArray(bool bIsHiddenArgument = false) const
 	{
 		return IsPropertyArrayType<FBoolProperty>(bIsHiddenArgument);
 	}
 	
 	// Returns true if this memory handle maps to a float property
-	FORCEINLINE bool IsFloat(bool bIsHiddenArgument = false) const
+	bool IsFloat(bool bIsHiddenArgument = false) const
 	{
 		return IsPropertyType<FFloatProperty>(bIsHiddenArgument);
 	}
 
 	// Returns true if this memory handle maps to a float array property
-	FORCEINLINE bool IsFloatArray(bool bIsHiddenArgument = false) const
+	bool IsFloatArray(bool bIsHiddenArgument = false) const
 	{
 		return IsPropertyArrayType<FFloatProperty>(bIsHiddenArgument);
 	}
 
 	// Returns true if this memory handle maps to a double property
-	FORCEINLINE bool IsDouble(bool bIsHiddenArgument = false) const
+	bool IsDouble(bool bIsHiddenArgument = false) const
 	{
 		return IsPropertyType<FDoubleProperty>(bIsHiddenArgument);
 	}
 
 	// Returns true if this memory handle maps to a double array property
-	FORCEINLINE bool IsDoubleArray(bool bIsHiddenArgument = false) const
+	bool IsDoubleArray(bool bIsHiddenArgument = false) const
 	{
 		return IsPropertyArrayType<FDoubleProperty>(bIsHiddenArgument);
 	}
 
 	// Returns true if this memory handle maps to a int32 property
-	FORCEINLINE bool IsInt32(bool bIsHiddenArgument = false) const
+	bool IsInt32(bool bIsHiddenArgument = false) const
 	{
 		return IsPropertyType<FIntProperty>(bIsHiddenArgument);
 	}
 
 	// Returns true if this memory handle maps to a int32 array property
-	FORCEINLINE bool IsInt32Array(bool bIsHiddenArgument = false) const
+	bool IsInt32Array(bool bIsHiddenArgument = false) const
 	{
 		return IsPropertyArrayType<FIntProperty>(bIsHiddenArgument);
 	}
 
 	// Returns true if this memory handle maps to an FName property
-    FORCEINLINE bool IsName(bool bIsHiddenArgument = false) const
+    bool IsName(bool bIsHiddenArgument = false) const
     {
     	return IsPropertyType<FNameProperty>(bIsHiddenArgument);
     }
 
 	// Returns true if this memory handle maps to an FName array property
-    FORCEINLINE bool IsNameArray(bool bIsHiddenArgument = false) const
+    bool IsNameArray(bool bIsHiddenArgument = false) const
     {
     	return IsPropertyArrayType<FNameProperty>(bIsHiddenArgument);
     }
 
 	// Returns true if this memory handle maps to an FString property
-	FORCEINLINE bool IsString(bool bIsHiddenArgument = false) const
+	bool IsString(bool bIsHiddenArgument = false) const
 	{
 		return IsPropertyType<FStrProperty>(bIsHiddenArgument);
 	}
 
 	// Returns true if this memory handle maps to an FString array property
-	FORCEINLINE bool IsStringArray(bool bIsHiddenArgument = false) const
+	bool IsStringArray(bool bIsHiddenArgument = false) const
 	{
 		return IsPropertyArrayType<FStrProperty>(bIsHiddenArgument);
 	}
 
 	// Returns true if this memory handle maps to an enum property
-	FORCEINLINE bool IsEnum(const UEnum* InEnum, bool bIsHiddenArgument = false) const
+	bool IsEnum(const UEnum* InEnum, bool bIsHiddenArgument = false) const
 	{
 		const FProperty* ResolvedProperty = GetResolvedProperty();
 		return IsEnum(ResolvedProperty, InEnum, bIsHiddenArgument);
 	}
 
 	// Returns true if this memory handle maps to an enum array property
-	FORCEINLINE bool IsEnumArray(const UEnum* InEnum, bool bIsHiddenArgument = false) const
+	bool IsEnumArray(const UEnum* InEnum, bool bIsHiddenArgument = false) const
 	{
 		const FProperty* ResolvedProperty = GetResolvedProperty();
 		if(bIsHiddenArgument)
@@ -654,27 +654,27 @@ public:
 
 	// Returns true if this memory handle maps to an enum property
 	template<typename T>
-	FORCEINLINE bool IsEnum(bool bIsHiddenArgument = false) const
+	bool IsEnum(bool bIsHiddenArgument = false) const
 	{
 		return IsEnum(StaticEnum<T>(), bIsHiddenArgument);
 	}
 
 	// Returns true if this memory handle maps to an enum array property
 	template<typename T>
-	FORCEINLINE bool IsEnumArray(bool bIsHiddenArgument = false) const
+	bool IsEnumArray(bool bIsHiddenArgument = false) const
 	{
 		return IsEnumArray(StaticEnum<T>(), bIsHiddenArgument);
 	}
 
 	// Returns true if this memory handle maps to a struct property
-	FORCEINLINE bool IsStruct(const UScriptStruct* InStruct, bool bIsHiddenArgument = false) const
+	bool IsStruct(const UScriptStruct* InStruct, bool bIsHiddenArgument = false) const
 	{
 		const FProperty* ResolvedProperty = GetResolvedProperty();
 		return IsStruct(ResolvedProperty, InStruct, bIsHiddenArgument);
 	}
 
 	// Returns true if this memory handle maps to a struct array property
-	FORCEINLINE bool IsStructArray(const UScriptStruct* InStruct, bool bIsHiddenArgument = false) const
+	bool IsStructArray(const UScriptStruct* InStruct, bool bIsHiddenArgument = false) const
 	{
 		const FProperty* ResolvedProperty = GetResolvedProperty();
 		if(bIsHiddenArgument)
@@ -693,7 +693,7 @@ public:
 		typename T,
 		typename TEnableIf<TRigVMIsBaseStructure<T>::Value, T>::Type* = nullptr
 	>
-	FORCEINLINE bool IsStruct(bool bIsHiddenArgument = false) const
+	bool IsStruct(bool bIsHiddenArgument = false) const
 	{
 		return IsStruct(TBaseStructure<T>::Get(), bIsHiddenArgument);
 	}
@@ -703,7 +703,7 @@ public:
 		typename T,
 		typename TEnableIf<TModels<CRigVMUStruct, T>::Value>::Type * = nullptr
 	>
-	FORCEINLINE bool IsStruct(bool bIsHiddenArgument = false) const
+	bool IsStruct(bool bIsHiddenArgument = false) const
 	{
 		return IsStruct(T::StaticStruct(), bIsHiddenArgument);
 	}
@@ -713,7 +713,7 @@ public:
 		typename T,
 		typename TEnableIf<TRigVMIsBaseStructure<T>::Value, T>::Type* = nullptr
 	>
-	FORCEINLINE bool IsStructArray(bool bIsHiddenArgument = false) const
+	bool IsStructArray(bool bIsHiddenArgument = false) const
 	{
 		return IsStructArray(TBaseStructure<T>::Get(), bIsHiddenArgument);
 	}
@@ -723,20 +723,20 @@ public:
 		typename T,
 		typename TEnableIf<TModels<CRigVMUStruct, T>::Value>::Type * = nullptr
 	>
-	FORCEINLINE bool IsStructArray(bool bIsHiddenArgument = false) const
+	bool IsStructArray(bool bIsHiddenArgument = false) const
 	{
 		return IsStructArray(T::StaticStruct(), bIsHiddenArgument);
 	}
 
 	// Returns true if this memory handle maps to an object property
-	FORCEINLINE bool IsObject(const UClass* InClass, bool bIsHiddenArgument = false) const
+	bool IsObject(const UClass* InClass, bool bIsHiddenArgument = false) const
 	{
 		const FProperty* ResolvedProperty = GetResolvedProperty();
 		return IsObject(ResolvedProperty, InClass, bIsHiddenArgument);
 	}
 
 	// Returns true if this memory handle maps to an objct array property
-	FORCEINLINE bool IsObjectArray(const UClass* InClass, bool bIsHiddenArgument = false) const
+	bool IsObjectArray(const UClass* InClass, bool bIsHiddenArgument = false) const
 	{
 		const FProperty* ResolvedProperty = GetResolvedProperty();
 		if(bIsHiddenArgument)
@@ -755,7 +755,7 @@ public:
 		typename T,
 		typename TEnableIf<TModels<CRigVMUClass, T>::Value>::Type * = nullptr
 	>
-	FORCEINLINE bool IsObject(bool bIsHiddenArgument = false) const
+	bool IsObject(bool bIsHiddenArgument = false) const
 	{
 		return IsObject(T::StaticClass(), bIsHiddenArgument);
 	}
@@ -765,14 +765,14 @@ public:
 		typename T,
 		typename TEnableIf<TModels<CRigVMUClass, T>::Value>::Type * = nullptr
 	>
-	FORCEINLINE bool IsObjectArray(bool bIsHiddenArgument = false) const
+	bool IsObjectArray(bool bIsHiddenArgument = false) const
 	{
 		return IsObjectArray(T::StaticClass(), bIsHiddenArgument);
 	}
 
 private:
 
-	FORCEINLINE_DEBUGGABLE uint8* GetData_Internal(bool bFollowPropertyPath, int32 InSliceIndex) const
+	uint8* GetData_Internal(bool bFollowPropertyPath, int32 InSliceIndex) const
 	{
 		if(InSliceIndex != INDEX_NONE)
 		{
@@ -802,7 +802,7 @@ private:
 		return Ptr;
 	}
 
-	FORCEINLINE_DEBUGGABLE TRigVMLazyValueBase GetDataLazily_Internal(bool bFollowPropertyPath, int32 InSliceIndex) const
+	TRigVMLazyValueBase GetDataLazily_Internal(bool bFollowPropertyPath, int32 InSliceIndex) const
 	{
 		// note: this works also for memory handles which don't provide a lazy branch
 		TRigVMLazyValueBase LazyValue;
@@ -812,7 +812,7 @@ private:
 		return LazyValue;
 	}
 
-	FORCEINLINE static const FProperty* GetArrayElementProperty(const FProperty* InProperty)
+	static const FProperty* GetArrayElementProperty(const FProperty* InProperty)
 	{
 		if(const FArrayProperty* ArrayProperty = CastField<FArrayProperty>(InProperty))
 		{
@@ -821,13 +821,13 @@ private:
 		return nullptr;
 	}
 
-	FORCEINLINE static const FProperty* GetArrayElementPropertyChecked(const FProperty* InProperty)
+	static const FProperty* GetArrayElementPropertyChecked(const FProperty* InProperty)
 	{
 		const FArrayProperty* ArrayProperty = CastFieldChecked<FArrayProperty>(InProperty);
 		return ArrayProperty->Inner;
 	}
 
-	FORCEINLINE static bool IsEnum(const FProperty* InProperty, const UEnum* InEnum, bool bUseArrayElement)
+	static bool IsEnum(const FProperty* InProperty, const UEnum* InEnum, bool bUseArrayElement)
 	{
 		const FProperty* Property = InProperty;
 		if(bUseArrayElement)
@@ -851,7 +851,7 @@ private:
 		return false;
 	}
 
-	FORCEINLINE static bool IsStruct(const FProperty* InProperty, const UScriptStruct* InScriptStruct, bool bUseArrayElement)
+	static bool IsStruct(const FProperty* InProperty, const UScriptStruct* InScriptStruct, bool bUseArrayElement)
 	{
 		const FProperty* Property = InProperty;
 		if(bUseArrayElement)
@@ -868,7 +868,7 @@ private:
 		return false;
 	}
 
-	FORCEINLINE static bool IsObject(const FProperty* InProperty, const UClass* InClass, bool bUseArrayElement)
+	static bool IsObject(const FProperty* InProperty, const UClass* InClass, bool bUseArrayElement)
 	{
 		const FProperty* Property = InProperty;
 		if(bUseArrayElement)
@@ -1097,7 +1097,7 @@ public:
 	//////////////////////////////////////////////////////////////////////////////
 
 	// Returns the memory type of this memory
-	FORCEINLINE ERigVMMemoryType GetMemoryType() const
+	ERigVMMemoryType GetMemoryType() const
 	{
 		if(URigVMMemoryStorageGeneratorClass* Class = Cast<URigVMMemoryStorageGeneratorClass>(GetClass()))
 		{
@@ -1108,7 +1108,7 @@ public:
 	}
 
 	// Returns a hash of unique to the configuration of the memory
-	FORCEINLINE uint32 GetMemoryHash() const
+	uint32 GetMemoryHash() const
 	{
 		if(URigVMMemoryStorageGeneratorClass* Class = Cast<URigVMMemoryStorageGeneratorClass>(GetClass()))
 		{
@@ -1119,13 +1119,13 @@ public:
 	}
 
 	// Returns the number of properties stored in this instance
-	FORCEINLINE int32 Num() const
+	int32 Num() const
 	{
 		return GetProperties().Num();
 	}
 
 	// Returns true if a provided property index is valid
-	FORCEINLINE bool IsValidIndex(int32 InIndex) const
+	bool IsValidIndex(int32 InIndex) const
 	{
 		return GetProperties().IsValidIndex(InIndex);
 	}
@@ -1143,7 +1143,7 @@ public:
 	int32 GetPropertyIndexByName(const FName& InName) const;
 
 	// Returns a property given its index
-	FORCEINLINE const FProperty* GetProperty(int32 InPropertyIndex) const
+	const FProperty* GetProperty(int32 InPropertyIndex) const
 	{
 		return GetProperties()[InPropertyIndex];
 	}
@@ -1173,20 +1173,20 @@ public:
 	}
 
 	// Returns true if the property at a given index is a TArray
-	FORCEINLINE bool IsArray(int32 InPropertyIndex) const
+	bool IsArray(int32 InPropertyIndex) const
 	{
 		return GetProperty(InPropertyIndex)->IsA<FArrayProperty>();
 	}
 
 	// Returns true if the property at a given index is a TMap
-	FORCEINLINE bool IsMap(int32 InPropertyIndex) const
+	bool IsMap(int32 InPropertyIndex) const
 	{
 		return GetProperty(InPropertyIndex)->IsA<FMapProperty>();
 	}
 
 	// Returns the memory for a property given its index
 	template<typename T>
-	FORCEINLINE T* GetData(int32 InPropertyIndex)
+	T* GetData(int32 InPropertyIndex)
 	{
 		const TArray<const FProperty*>& Properties = GetProperties();
 		check(Properties.IsValidIndex(InPropertyIndex));
@@ -1195,7 +1195,7 @@ public:
 
 	// Returns the memory for a property given its name (or nullptr)
 	template<typename T>
-	FORCEINLINE T* GetDataByName(const FName& InName)
+	T* GetDataByName(const FName& InName)
 	{
 		const int32 PropertyIndex = GetPropertyIndexByName(InName);
 		if(PropertyIndex == INDEX_NONE)
@@ -1207,7 +1207,7 @@ public:
 
 	// Returns the memory for a property given its index and a matching property path
 	template<typename T>
-	FORCEINLINE T* GetData(int32 InPropertyIndex, const FRigVMPropertyPath& InPropertyPath)
+	T* GetData(int32 InPropertyIndex, const FRigVMPropertyPath& InPropertyPath)
 	{
 		const FProperty* Property = GetProperty(InPropertyIndex);
 		return InPropertyPath.GetData<T>(GetData<uint8>(InPropertyIndex), Property);
@@ -1215,7 +1215,7 @@ public:
 
 	// Returns the memory for a property given its name and a matching property path (or nullptr)
 	template<typename T>
-	FORCEINLINE T* GetDataByName(const FName& InName, const FRigVMPropertyPath& InPropertyPath)
+	T* GetDataByName(const FName& InName, const FRigVMPropertyPath& InPropertyPath)
 	{
 		const int32 PropertyIndex = GetPropertyIndexByName(InName);
 		if(PropertyIndex == INDEX_NONE)
@@ -1227,7 +1227,7 @@ public:
 
 	// Returns the memory for a property (and optionally a property path) given an operand
 	template<typename T>
-	FORCEINLINE T* GetData(const FRigVMOperand& InOperand)
+	T* GetData(const FRigVMOperand& InOperand)
 	{
 		const int32 PropertyIndex = InOperand.GetRegisterIndex();
 		const int32 PropertyPathIndex = InOperand.GetRegisterOffset();
@@ -1245,35 +1245,35 @@ public:
 
 	// Returns the ref of an element stored at a given property index
 	template<typename T>
-	FORCEINLINE T& GetRef(int32 InPropertyIndex)
+	T& GetRef(int32 InPropertyIndex)
 	{
 		return *GetData<T>(InPropertyIndex);
 	}
 
 	// Returns the ref of an element stored at a given property name (throws if name is invalid)
 	template<typename T>
-	FORCEINLINE T& GetRefByName(const FName& InName)
+	T& GetRefByName(const FName& InName)
 	{
 		return *GetDataByName<T>(InName);
 	}
 
 	// Returns the ref of an element stored at a given property index and a property path
 	template<typename T>
-	FORCEINLINE T& GetRef(int32 InPropertyIndex, const FRigVMPropertyPath& InPropertyPath)
+	T& GetRef(int32 InPropertyIndex, const FRigVMPropertyPath& InPropertyPath)
 	{
 		return *GetData<T>(InPropertyIndex, InPropertyPath);
 	}
 
 	// Returns the ref of an element stored at a given property name and a property path (throws if name is invalid)
 	template<typename T>
-	FORCEINLINE T& GetRefByName(const FName& InName, const FRigVMPropertyPath& InPropertyPath)
+	T& GetRefByName(const FName& InName, const FRigVMPropertyPath& InPropertyPath)
 	{
 		return *GetDataByName<T>(InName, InPropertyPath);
 	}
 
 	// Returns the ref of an element stored for a given operand
 	template<typename T>
-	FORCEINLINE T& GetRef(const FRigVMOperand& InOperand)
+	T& GetRef(const FRigVMOperand& InOperand)
 	{
 		return *GetData<T>(InOperand);
 	}
@@ -1282,7 +1282,7 @@ public:
 	FString GetDataAsString(int32 InPropertyIndex, int32 PortFlags = PPF_None);
 
 	// Returns the exported text for given property name 
-	FORCEINLINE FString GetDataAsStringByName(const FName& InName, int32 PortFlags = PPF_None)
+	FString GetDataAsStringByName(const FName& InName, int32 PortFlags = PPF_None)
 	{
 		const int32 PropertyIndex = GetPropertyIndexByName(InName);
 		return GetDataAsString(PropertyIndex, PortFlags);
@@ -1295,7 +1295,7 @@ public:
 	FString GetDataAsStringSafe(int32 InPropertyIndex, int32 PortFlags = PPF_None);
 
 	// Returns the exported text for given property name 
-	FORCEINLINE FString GetDataAsStringByNameSafe(const FName& InName, int32 PortFlags = PPF_None)
+	FString GetDataAsStringByNameSafe(const FName& InName, int32 PortFlags = PPF_None)
 	{
 		const int32 PropertyIndex = GetPropertyIndexByName(InName);
 		return GetDataAsStringSafe(PropertyIndex, PortFlags);
@@ -1308,7 +1308,7 @@ public:
 	bool SetDataFromString(int32 InPropertyIndex, const FString& InValue);
 
 	// Sets the content of a property by name given an exported string. Returns true if succeeded
-	FORCEINLINE bool SetDataFromStringByName(const FName& InName, const FString& InValue)
+	bool SetDataFromStringByName(const FName& InName, const FString& InValue)
 	{
 		const int32 PropertyIndex = GetPropertyIndexByName(InName);
 		return SetDataFromString(PropertyIndex, InValue);
@@ -1318,7 +1318,7 @@ public:
 	FRigVMMemoryHandle GetHandle(int32 InPropertyIndex, const FRigVMPropertyPath* InPropertyPath = nullptr);
 
 	// Returns the handle for a given property by name (and optionally property path)
-	FORCEINLINE FRigVMMemoryHandle GetHandleByName(const FName& InName, const FRigVMPropertyPath* InPropertyPath = nullptr)
+	FRigVMMemoryHandle GetHandleByName(const FName& InName, const FRigVMPropertyPath* InPropertyPath = nullptr)
 	{
 		const int32 PropertyIndex = GetPropertyIndexByName(InName);
 		return GetHandle(PropertyIndex, InPropertyPath);
diff --git a/Engine/Source/Runtime/RigVM/Public/RigVMCore/RigVMNameCache.h b/Engine/Source/Runtime/RigVM/Public/RigVMCore/RigVMNameCache.h
--- a/Engine/Source/Runtime/RigVM/Public/RigVMCore/RigVMNameCache.h
+++ b/Engine/Source/Runtime/RigVM/Public/RigVMCore/RigVMNameCache.h
@@ -26,14 +26,14 @@ struct RIGVM_API FRigVMNameOp
 {
 public:
 
-	FORCEINLINE FRigVMNameOp()
+	FRigVMNameOp()
 	: A(INDEX_NONE)
 	, B(INDEX_NONE)
 	, C(INDEX_NONE)
 	, Type(ERigVMNameOp::None)
 	{}
 
-	FORCEINLINE static uint32 GetTypeHash(const FName& InName)
+	static uint32 GetTypeHash(const FName& InName)
 	{
 		return HashCombine(InName.GetComparisonIndex().ToUnstableInt(), uint32(InName.GetNumber()));
 	}
@@ -48,22 +48,22 @@ public:
     static FRigVMNameOp StartsWith(const FName& InA, const FName& InB, const ESearchCase::Type InSearchCase);
     static FRigVMNameOp Contains(const FName& InA, const FName& InB, const ESearchCase::Type InSearchCase);
 
-	friend FORCEINLINE uint32 GetTypeHash(const FRigVMNameOp& Op)
+	friend uint32 GetTypeHash(const FRigVMNameOp& Op)
 	{
 		return HashCombine(Op.A, HashCombine(Op.B, HashCombine(Op.C, (uint32)Op.Type)));
 	}
 
-	FORCEINLINE bool operator ==(const FRigVMNameOp& Other) const
+	bool operator ==(const FRigVMNameOp& Other) const
 	{
 		return Type == Other.Type && A == Other.A && B == Other.B && C == Other.C;
 	}
 
-	FORCEINLINE bool operator !=(const FRigVMNameOp& Other) const
+	bool operator !=(const FRigVMNameOp& Other) const
 	{
 		return Type != Other.Type || A != Other.A || B != Other.B || C != Other.C;
 	}
 
-	FORCEINLINE bool operator <(const FRigVMNameOp& Other) const
+	bool operator <(const FRigVMNameOp& Other) const
 	{
 		if (Type < Other.Type)
 		{
@@ -80,7 +80,7 @@ public:
 		return C < Other.C;
 	}
 
-	FORCEINLINE bool operator >(const FRigVMNameOp& Other) const
+	bool operator >(const FRigVMNameOp& Other) const
 	{
 		if (Type > Other.Type)
 		{
@@ -110,7 +110,7 @@ class RIGVM_API FRigVMNameCache
 {
 public:
 
-	FORCEINLINE void Reset()
+	void Reset()
 	{
 		NameCache.Reset();
 		BoolCache.Reset();
@@ -126,8 +126,8 @@ public:
 	bool StartsWith(const FName& InA, const FName& InB, const ESearchCase::Type InSearchCase);
 	bool Contains(const FName& InA, const FName& InB, const ESearchCase::Type InSearchCase);
 
-	FORCEINLINE const TMap<FRigVMNameOp, FName>& GetNameCache() const { return NameCache; }
-	FORCEINLINE const TMap<FRigVMNameOp, bool>&  GetBoolCache() const { return BoolCache; }
+	const TMap<FRigVMNameOp, FName>& GetNameCache() const { return NameCache; }
+	const TMap<FRigVMNameOp, bool>&  GetBoolCache() const { return BoolCache; }
 	TArray<FRigVMNameOp> GetNameOps() const;
 	TArray<FName> GetNameValues() const;
 	TArray<FRigVMNameOp> GetBoolOps() const;
@@ -138,7 +138,7 @@ private:
 #if WITH_EDITOR
 	bool CheckCacheSize() const;
 #else
-	FORCEINLINE bool CheckCacheSize() const { return true; } 
+	bool CheckCacheSize() const { return true; } 
 #endif
 
 	TMap<FRigVMNameOp, FName> NameCache;
diff --git a/Engine/Source/Runtime/RigVM/Public/RigVMCore/RigVMNativized.h b/Engine/Source/Runtime/RigVM/Public/RigVMCore/RigVMNativized.h
--- a/Engine/Source/Runtime/RigVM/Public/RigVMCore/RigVMNativized.h
+++ b/Engine/Source/Runtime/RigVM/Public/RigVMCore/RigVMNativized.h
@@ -76,7 +76,7 @@ public:
 protected:
 
 	template<typename ExecuteContextType = FRigVMExecuteContext>
-	FORCEINLINE ExecuteContextType& UpdateContext(TArrayView<void*> AdditionalArguments, int32 InNumberInstructions, const FName& InEntryName)
+	ExecuteContextType& UpdateContext(TArrayView<void*> AdditionalArguments, int32 InNumberInstructions, const FName& InEntryName)
 	{
 		UpdateExternalVariables();
 	
@@ -108,29 +108,29 @@ protected:
 		}
 	};
 
-	FORCEINLINE void BeginSlice(int32 InCount, int32 InRelativeIndex = 0)
+	void BeginSlice(int32 InCount, int32 InRelativeIndex = 0)
 	{
 		Context.BeginSlice(InCount, InRelativeIndex);
 	}
 
-	FORCEINLINE void EndSlice()
+	void EndSlice()
 	{
 		Context.EndSlice();
 	}
 
-	FORCEINLINE bool IsValidArraySize(int32 InArraySize) const
+	bool IsValidArraySize(int32 InArraySize) const
 	{
 		return Context.IsValidArraySize(InArraySize);
 	}
 
 	template<typename T>
-	FORCEINLINE bool IsValidArrayIndex(int32& InOutIndex, const TArray<T>& InArray) const
+	bool IsValidArrayIndex(int32& InOutIndex, const TArray<T>& InArray) const
 	{
 		return Context.IsValidArrayIndex(InOutIndex, InArray.Num());
 	}
 
 	template<typename T>
-	FORCEINLINE static const T& GetArrayElementSafe(const TArray<T>& InArray, const int32 InIndex)
+	static const T& GetArrayElementSafe(const TArray<T>& InArray, const int32 InIndex)
 	{
 		if(InArray.IsValidIndex(InIndex))
 		{
@@ -141,7 +141,7 @@ protected:
 	}
 
 	template<typename T>
-	FORCEINLINE static T& GetArrayElementSafe(TArray<T>& InArray, const int32 InIndex)
+	static T& GetArrayElementSafe(TArray<T>& InArray, const int32 InIndex)
 	{
 		if(InArray.IsValidIndex(InIndex))
 		{
@@ -152,7 +152,7 @@ protected:
 	}
 
 	template<typename A, typename B>
-	FORCEINLINE static void CopyUnrelatedArrays(TArray<A>& Target, const TArray<B>& Source)
+	static void CopyUnrelatedArrays(TArray<A>& Target, const TArray<B>& Source)
 	{
 		const int32 Num = Source.Num();
 		Target.SetNum(Num);
@@ -162,7 +162,7 @@ protected:
 		}
 	}
 
-	FORCEINLINE void BroadcastExecutionReachedExit()
+	void BroadcastExecutionReachedExit()
 	{
 		if(EntriesBeingExecuted.Num() == 1)
 		{
@@ -172,7 +172,7 @@ protected:
 	}
 
 	template<typename T>
-	FORCEINLINE T& GetExternalVariableRef(const FName& InExternalVariableName, const FName& InExpectedType) const
+	T& GetExternalVariableRef(const FName& InExternalVariableName, const FName& InExpectedType) const
 	{
 		for(const FRigVMExternalVariable& ExternalVariable : GetExternalVariables())
 		{
@@ -189,7 +189,7 @@ protected:
 	}
 
 	template<typename T>
-	FORCEINLINE T& GetOperandSlice(TArray<T>& InOutArray)
+	T& GetOperandSlice(TArray<T>& InOutArray)
 	{
 		const int32 SliceIndex = Context.GetSlice().GetIndex();
 		if(InOutArray.Num() <= SliceIndex)
@@ -200,7 +200,7 @@ protected:
 	}
 
 	template<typename T>
-	FORCEINLINE static bool ArrayOp_Iterator(const TArray<T>& InArray, T& OutElement, int32 InIndex, int32& OutCount, float& OutRatio)
+	static bool ArrayOp_Iterator(const TArray<T>& InArray, T& OutElement, int32 InIndex, int32& OutCount, float& OutRatio)
 	{
 		OutCount = InArray.Num();
 		const bool bContinue = InIndex >=0 && InIndex < OutCount;
@@ -218,7 +218,7 @@ protected:
 		return bContinue;
 	}
 
-	FORCEINLINE static bool ArrayOp_Iterator(const TArray<double>& InArray, float& OutElement, int32 InIndex, int32& OutCount, float& OutRatio)
+	static bool ArrayOp_Iterator(const TArray<double>& InArray, float& OutElement, int32 InIndex, int32& OutCount, float& OutRatio)
 	{
 		double Element = (double)OutElement;
 		const bool bResult = ArrayOp_Iterator<double>(InArray, Element, InIndex, OutCount, OutRatio);
@@ -226,7 +226,7 @@ protected:
 		return bResult;
 	}
 
-	FORCEINLINE static bool ArrayOp_Iterator(const TArray<float>& InArray, double& OutElement, int32 InIndex, int32& OutCount, float& OutRatio)
+	static bool ArrayOp_Iterator(const TArray<float>& InArray, double& OutElement, int32 InIndex, int32& OutCount, float& OutRatio)
 	{
 		float Element = (float)OutElement;
 		const bool bResult = ArrayOp_Iterator<float>(InArray, Element, InIndex, OutCount, OutRatio);
@@ -235,24 +235,24 @@ protected:
 	}
 
 	template<typename T>
-	FORCEINLINE static bool ArrayOp_Find(const TArray<T>& InArray, const T& InElement, int32& OutIndex)
+	static bool ArrayOp_Find(const TArray<T>& InArray, const T& InElement, int32& OutIndex)
 	{
 		OutIndex = InArray.Find(InElement);
 		return OutIndex != INDEX_NONE;
 	}
 
-	FORCEINLINE static bool ArrayOp_Find(const TArray<float>& InArray, const double& InElement, int32& OutIndex)
+	static bool ArrayOp_Find(const TArray<float>& InArray, const double& InElement, int32& OutIndex)
 	{
 		return ArrayOp_Find<float>(InArray, (float)InElement, OutIndex);
 	}
 
-	FORCEINLINE static bool ArrayOp_Find(const TArray<double>& InArray, const float& InElement, int32& OutIndex)
+	static bool ArrayOp_Find(const TArray<double>& InArray, const float& InElement, int32& OutIndex)
 	{
 		return ArrayOp_Find<double>(InArray, (double)InElement, OutIndex);
 	}
 
 	template<typename T>
-	FORCEINLINE static void ArrayOp_Union(TArray<T>& InOutA, const TArray<T>& InB)
+	static void ArrayOp_Union(TArray<T>& InOutA, const TArray<T>& InB)
 	{
 		for(const T& Element : InB)
 		{
@@ -261,7 +261,7 @@ protected:
 	}
 
 	template<typename T>
-	FORCEINLINE static void ArrayOp_Difference(TArray<T>& InOutA, const TArray<T>& InB)
+	static void ArrayOp_Difference(TArray<T>& InOutA, const TArray<T>& InB)
 	{
 		TArray<T> Temp;
 		Swap(InOutA, Temp);
@@ -283,7 +283,7 @@ protected:
 	}
 
 	template<typename T>
-	FORCEINLINE static void ArrayOp_Intersection(TArray<T>& InOutA, const TArray<T>& InB)
+	static void ArrayOp_Intersection(TArray<T>& InOutA, const TArray<T>& InB)
 	{
 		TArray<T> Temp;
 		Swap(InOutA, Temp);
@@ -298,7 +298,7 @@ protected:
 	}
 
 	template<typename T>
-	FORCEINLINE static void ArrayOp_Reverse(TArray<T>& InOutArray)
+	static void ArrayOp_Reverse(TArray<T>& InOutArray)
 	{
 		Algo::Reverse(InOutArray);
 	}
@@ -309,7 +309,7 @@ public:
 		typename T,
 		typename TEnableIf<TRigVMIsBaseStructure<T>::Value, T>::Type* = nullptr
 	>
-	FORCEINLINE static T GetStructConstant(const FString& InDefaultValue)
+	static T GetStructConstant(const FString& InDefaultValue)
 	{
 		T Value;
 		if(!InDefaultValue.IsEmpty())
@@ -324,7 +324,7 @@ public:
 		typename T,
 		typename TEnableIf<TModels<CRigVMUStruct, T>::Value>::Type* = nullptr
 	>
-	FORCEINLINE static T GetStructConstant(const FString& InDefaultValue)
+	static T GetStructConstant(const FString& InDefaultValue)
 	{
 		T Value;
 		if(!InDefaultValue.IsEmpty())
@@ -339,7 +339,7 @@ public:
 		typename T,
 		typename TEnableIf<TRigVMIsBaseStructure<T>::Value, T>::Type* = nullptr
 	>
-	FORCEINLINE static TArray<T> GetStructArrayConstant(const FString& InDefaultValue)
+	static TArray<T> GetStructArrayConstant(const FString& InDefaultValue)
 	{
 		TArray<T> Value;
 		if(!InDefaultValue.IsEmpty())
@@ -354,7 +354,7 @@ public:
 		typename T,
 		typename TEnableIf<TModels<CRigVMUStruct, T>::Value>::Type* = nullptr
 	>
-	FORCEINLINE static TArray<T> GetStructArrayConstant(const FString& InDefaultValue)
+	static TArray<T> GetStructArrayConstant(const FString& InDefaultValue)
 	{
 		TArray<T> Value;
 		if(!InDefaultValue.IsEmpty())
@@ -370,36 +370,36 @@ protected:
 	class FTransformSetter
 	{
 	public:
-		FORCEINLINE FTransformSetter(FTransform& InOutTransform)
+		FTransformSetter(FTransform& InOutTransform)
 			: Transform(InOutTransform)
 		{
 		}
 
-		FORCEINLINE void SetTranslationX(const FTransform::FReal InValue) const { SetTranslationComponent(0, InValue); }
-		FORCEINLINE void SetTranslationY(const FTransform::FReal InValue) const { SetTranslationComponent(1, InValue); }
-		FORCEINLINE void SetTranslationZ(const FTransform::FReal InValue) const { SetTranslationComponent(2, InValue); }
-		FORCEINLINE void SetTranslationComponent(const uint8 InComponent, const FTransform::FReal InValue) const
+		void SetTranslationX(const FTransform::FReal InValue) const { SetTranslationComponent(0, InValue); }
+		void SetTranslationY(const FTransform::FReal InValue) const { SetTranslationComponent(1, InValue); }
+		void SetTranslationZ(const FTransform::FReal InValue) const { SetTranslationComponent(2, InValue); }
+		void SetTranslationComponent(const uint8 InComponent, const FTransform::FReal InValue) const
 		{
 			FVector Translation = Transform.GetTranslation();
 			Translation.Component(InComponent) = InValue;
 			Transform.SetTranslation(Translation);
 		}
 		
-		FORCEINLINE void SetScaleX(const FTransform::FReal InValue) const { SetScaleComponent(0, InValue); }
-		FORCEINLINE void SetScaleY(const FTransform::FReal InValue) const { SetScaleComponent(1, InValue); }
-		FORCEINLINE void SetScaleZ(const FTransform::FReal InValue) const { SetScaleComponent(2, InValue); }
-		FORCEINLINE void SetScaleComponent(const uint8 InComponent, const FTransform::FReal InValue) const
+		void SetScaleX(const FTransform::FReal InValue) const { SetScaleComponent(0, InValue); }
+		void SetScaleY(const FTransform::FReal InValue) const { SetScaleComponent(1, InValue); }
+		void SetScaleZ(const FTransform::FReal InValue) const { SetScaleComponent(2, InValue); }
+		void SetScaleComponent(const uint8 InComponent, const FTransform::FReal InValue) const
 		{
 			FVector Scale = Transform.GetScale3D();
 			Scale.Component(InComponent) = InValue;
 			Transform.SetScale3D(Scale);
 		}
 		
-		FORCEINLINE void SetRotationX(const FTransform::FReal InValue) const { SetRotationComponent(0, InValue); }
-		FORCEINLINE void SetRotationY(const FTransform::FReal InValue) const { SetRotationComponent(1, InValue); }
-		FORCEINLINE void SetRotationZ(const FTransform::FReal InValue) const { SetRotationComponent(2, InValue); }
-		FORCEINLINE void SetRotationW(const FTransform::FReal InValue) const { SetRotationComponent(3, InValue); }
-		FORCEINLINE void SetRotationComponent(const uint8 InComponent, const FTransform::FReal InValue) const
+		void SetRotationX(const FTransform::FReal InValue) const { SetRotationComponent(0, InValue); }
+		void SetRotationY(const FTransform::FReal InValue) const { SetRotationComponent(1, InValue); }
+		void SetRotationZ(const FTransform::FReal InValue) const { SetRotationComponent(2, InValue); }
+		void SetRotationW(const FTransform::FReal InValue) const { SetRotationComponent(3, InValue); }
+		void SetRotationComponent(const uint8 InComponent, const FTransform::FReal InValue) const
 		{
 			FQuat Rotation = Transform.GetRotation();
 			switch(InComponent)
diff --git a/Engine/Source/Runtime/RigVM/Public/RigVMCore/RigVMPropertyPath.h b/Engine/Source/Runtime/RigVM/Public/RigVMCore/RigVMPropertyPath.h
--- a/Engine/Source/Runtime/RigVM/Public/RigVMCore/RigVMPropertyPath.h
+++ b/Engine/Source/Runtime/RigVM/Public/RigVMCore/RigVMPropertyPath.h
@@ -52,7 +52,7 @@ struct RIGVM_API FRigVMPropertyPathDescription
 	bool IsValid() const { return PropertyIndex != INDEX_NONE && !HeadCPPType.IsEmpty() && !SegmentPath.IsEmpty(); }
 
 	// Archive operator for serialization
-	FORCEINLINE friend FArchive& operator<<(FArchive& Ar, FRigVMPropertyPathDescription& Path)
+	friend FArchive& operator<<(FArchive& Ar, FRigVMPropertyPathDescription& Path)
 	{
 		Ar << Path.PropertyIndex;
 		Ar << Path.HeadCPPType;
@@ -117,36 +117,36 @@ public:
 	FRigVMPropertyPath(const FRigVMPropertyPath& InOther);
 
 	// Returns the property path as a string (the sanitized SegmentPath)
-	FORCEINLINE const FString& ToString() const { return Path; }
+	const FString& ToString() const { return Path; }
 
 	// Returns the number of segments
-	FORCEINLINE int32 Num() const { return Segments.Num(); }
+	int32 Num() const { return Segments.Num(); }
 
 	// Returns true if the property path is valid (has any segments)
-	FORCEINLINE bool IsValid() const { return Num() > 0; }
+	bool IsValid() const { return Num() > 0; }
 
 	// Returns true if the property is empty (no segments)
-	FORCEINLINE bool IsEmpty() const { return Num() == 0; }
+	bool IsEmpty() const { return Num() == 0; }
 
 	// Index based access operator for a segment
-	FORCEINLINE const FRigVMPropertyPathSegment& operator[](int32 InIndex) const { return Segments[InIndex]; }
+	const FRigVMPropertyPathSegment& operator[](int32 InIndex) const { return Segments[InIndex]; }
 
-	FORCEINLINE bool operator ==(const FRigVMPropertyPath& Other) const
+	bool operator ==(const FRigVMPropertyPath& Other) const
 	{
 		return GetTypeHash(this) == GetTypeHash(Other);
 	}
 
-	FORCEINLINE bool operator !=(const FRigVMPropertyPath& Other) const
+	bool operator !=(const FRigVMPropertyPath& Other) const
 	{
 		return GetTypeHash(this) != GetTypeHash(Other);
 	}
 
-	FORCEINLINE bool operator >(const FRigVMPropertyPath& Other) const
+	bool operator >(const FRigVMPropertyPath& Other) const
 	{
 		return GetTypeHash(this) > GetTypeHash(Other);
 	}
 
-	FORCEINLINE bool operator <(const FRigVMPropertyPath& Other) const
+	bool operator <(const FRigVMPropertyPath& Other) const
 	{
 		return GetTypeHash(this) < GetTypeHash(Other);
 	}
@@ -157,13 +157,13 @@ public:
 	 * @param InProperty The head property to traverse off of
 	 */
 	template<typename T>
-	FORCEINLINE T* GetData(uint8* InPtr, const FProperty* InProperty) const
+	T* GetData(uint8* InPtr, const FProperty* InProperty) const
 	{
 		return (T*)GetData_Internal(InPtr, InProperty);
 	}
 
 	// Returns a unique hash for the property path
-	friend FORCEINLINE uint32 GetTypeHash(const FRigVMPropertyPath& InPropertyPath)
+	friend uint32 GetTypeHash(const FRigVMPropertyPath& InPropertyPath)
 	{
 		if(InPropertyPath.IsEmpty())
 		{
diff --git a/Engine/Source/Runtime/RigVM/Public/RigVMCore/RigVMRegistry.h b/Engine/Source/Runtime/RigVM/Public/RigVMCore/RigVMRegistry.h
--- a/Engine/Source/Runtime/RigVM/Public/RigVMCore/RigVMRegistry.h
+++ b/Engine/Source/Runtime/RigVM/Public/RigVMCore/RigVMRegistry.h
@@ -81,7 +81,7 @@ public:
 
 	// Adds a type if it doesn't exist yet and returns its index.
 	// This function is not thead-safe
-	FORCEINLINE TRigVMTypeIndex FindOrAddType(const FRigVMTemplateArgumentType& InType)
+	TRigVMTypeIndex FindOrAddType(const FRigVMTemplateArgumentType& InType)
 	{
 		return FindOrAddType_Internal(InType, false);
 	}
@@ -94,7 +94,7 @@ public:
 	TRigVMTypeIndex GetTypeIndex(const FRigVMTemplateArgumentType& InType) const;
 
 	// Returns the type index given a cpp type and a type object
-	FORCEINLINE TRigVMTypeIndex GetTypeIndex(const FName& InCPPType, UObject* InCPPTypeObject) const
+	TRigVMTypeIndex GetTypeIndex(const FName& InCPPType, UObject* InCPPTypeObject) const
 	{
 		return GetTypeIndex(FRigVMTemplateArgumentType(InCPPType, InCPPTypeObject));
 	}
@@ -104,7 +104,7 @@ public:
 		typename T,
 		typename TEnableIf<TIsEnum<T>::Value>::Type* = nullptr
 	>
-	FORCEINLINE TRigVMTypeIndex GetTypeIndex(bool bAsArray = false) const
+	TRigVMTypeIndex GetTypeIndex(bool bAsArray = false) const
 	{
 		FRigVMTemplateArgumentType Type(StaticEnum<T>());
 		if(bAsArray)
@@ -119,7 +119,7 @@ public:
 		typename T,
 		typename TEnableIf<TRigVMIsBaseStructure<T>::Value, T>::Type* = nullptr
 	>
-	FORCEINLINE TRigVMTypeIndex GetTypeIndex(bool bAsArray = false) const
+	TRigVMTypeIndex GetTypeIndex(bool bAsArray = false) const
 	{
 		FRigVMTemplateArgumentType Type(TBaseStructure<T>::Get());
 		if(bAsArray)
@@ -134,7 +134,7 @@ public:
 		typename T,
 		typename TEnableIf<TModels<CRigVMUStruct, T>::Value>::Type * = nullptr
 	>
-	FORCEINLINE TRigVMTypeIndex GetTypeIndex(bool bAsArray = false) const
+	TRigVMTypeIndex GetTypeIndex(bool bAsArray = false) const
 	{
 		FRigVMTemplateArgumentType Type(T::StaticStruct());
 		if(bAsArray)
@@ -149,7 +149,7 @@ public:
 		typename T,
 		typename TEnableIf<TModels<CRigVMUClass, T>::Value>::Type * = nullptr
 	>
-	FORCEINLINE TRigVMTypeIndex GetTypeIndex(bool bAsArray = false) const
+	TRigVMTypeIndex GetTypeIndex(bool bAsArray = false) const
 	{
 		FRigVMTemplateArgumentType Type(T::StaticClass());
 		if(bAsArray)
@@ -163,7 +163,7 @@ public:
 	const FRigVMTemplateArgumentType& GetType(TRigVMTypeIndex InTypeIndex) const;
 
 	// Returns the number of types
-	FORCEINLINE int32 NumTypes() const { return Types.Num(); }
+	int32 NumTypes() const { return Types.Num(); }
 
 	// Returns the type given only its cpp type
 	const FRigVMTemplateArgumentType& FindTypeFromCPPType(const FString& InCPPType) const;
@@ -266,12 +266,12 @@ private:
 
 	TRigVMTypeIndex FindOrAddType_Internal(const FRigVMTemplateArgumentType& InType, bool bForce);
 
-	FORCEINLINE static EObjectFlags DisallowedFlags()
+	static EObjectFlags DisallowedFlags()
 	{
 		return RF_BeginDestroyed | RF_FinishDestroyed;
 	}
 
-	FORCEINLINE static EObjectFlags NeededFlags()
+	static EObjectFlags NeededFlags()
 	{
 		return RF_Public;
 	}
diff --git a/Engine/Source/Runtime/RigVM/Public/RigVMCore/RigVMStruct.h b/Engine/Source/Runtime/RigVM/Public/RigVMCore/RigVMStruct.h
--- a/Engine/Source/Runtime/RigVM/Public/RigVMCore/RigVMStruct.h
+++ b/Engine/Source/Runtime/RigVM/Public/RigVMCore/RigVMStruct.h
@@ -102,7 +102,7 @@ public:
 		typename T,
 		typename TEnableIf<TRigVMIsBaseStructure<T>::Value>::Type * = nullptr
 	>
-	FORCEINLINE FName FindFirstVariableOfType() const
+	FName FindFirstVariableOfType() const
 	{
 		return FindFirstVariableOfType(TBaseStructure<T>::Get());
 	}
@@ -112,7 +112,7 @@ public:
 		typename T,
 		typename TEnableIf<TModels<CRigVMUStruct, T>::Value>::Type * = nullptr
 	>
-	FORCEINLINE FName FindFirstVariableOfType() const
+	FName FindFirstVariableOfType() const
 	{
 		return FindFirstVariableOfType(T::StaticStruct());
 	}
@@ -122,7 +122,7 @@ public:
 		typename T,
 		typename TEnableIf<TModels<CRigVMUClass, T>::Value>::Type * = nullptr
 	>
-	FORCEINLINE FName FindFirstVariableOfType() const
+	FName FindFirstVariableOfType() const
 	{
 		return FindFirstVariableOfType(T::StaticClass());
 	}
@@ -132,7 +132,7 @@ public:
 		typename T,
 		typename TEnableIf<TIsEnum<T>::Value>::Type * = nullptr
 	>
-		FORCEINLINE FName FindFirstVariableOfType() const
+		FName FindFirstVariableOfType() const
 	{
 		return FindFirstVariableOfType(StaticEnum<T>());
 	}
@@ -173,12 +173,12 @@ public:
 	// control flow related
 	bool IsForLoop() const;
 	bool IsControlFlowNode() const; 
-	FORCEINLINE virtual int32 GetNumSlices() const { return 1; }
+	virtual int32 GetNumSlices() const { return 1; }
 	const TArray<FName>& GetControlFlowBlocks() const;
 	virtual const bool IsControlFlowBlockSliced(const FName& InBlockName) const { return false; }
 
 	// node creation
-	FORCEINLINE virtual void OnUnitNodeCreated(FRigVMUnitNodeCreatedContext& InContext) const {}
+	virtual void OnUnitNodeCreated(FRigVMUnitNodeCreatedContext& InContext) const {}
 
 	// user workflow
 	TArray<FRigVMUserWorkflow> GetWorkflows(ERigVMUserWorkflowType InType, const UObject* InSubject) const; 
@@ -199,7 +199,7 @@ public:
 		typename T,
 		typename TEnableIf<TRigVMIsBaseStructure<T>::Value>::Type * = nullptr
 	>
-	FORCEINLINE static FString ExportToFullyQualifiedText(const T& InStructValue)
+	static FString ExportToFullyQualifiedText(const T& InStructValue)
 	{
 		return ExportToFullyQualifiedText(TBaseStructure<T>::Get(), (const uint8*)&InStructValue);
 	}
@@ -208,7 +208,7 @@ public:
 		typename T,
 		typename TEnableIf<TModels<CRigVMUStruct, T>::Value>::Type * = nullptr
 	>
-	FORCEINLINE static FString ExportToFullyQualifiedText(const T& InStructValue)
+	static FString ExportToFullyQualifiedText(const T& InStructValue)
 	{
 		return ExportToFullyQualifiedText(T::StaticStruct(), (const uint8*)&InStructValue);
 	}
@@ -263,7 +263,7 @@ protected:
 	static float GetRatioFromIndex(int32 InIndex, int32 InCount);
 	TMap<FName, FString> GetDefaultValues(UScriptStruct* InScriptStruct) const;
 	bool ApplyUpgradeInfo(const FRigVMStructUpgradeInfo& InUpgradeInfo);
-	FORCEINLINE virtual TArray<FRigVMUserWorkflow> GetSupportedWorkflows(const UObject* InSubject) const { return TArray<FRigVMUserWorkflow>(); }
+	virtual TArray<FRigVMUserWorkflow> GetSupportedWorkflows(const UObject* InSubject) const { return TArray<FRigVMUserWorkflow>(); }
 	virtual const TArray<FName>& GetControlFlowBlocks_Impl() const;
 
 #if WITH_EDITOR
diff --git a/Engine/Source/Runtime/RigVM/Public/RigVMCore/RigVMTemplate.h b/Engine/Source/Runtime/RigVM/Public/RigVMCore/RigVMTemplate.h
--- a/Engine/Source/Runtime/RigVM/Public/RigVMCore/RigVMTemplate.h
+++ b/Engine/Source/Runtime/RigVM/Public/RigVMCore/RigVMTemplate.h
@@ -100,7 +100,7 @@ struct RIGVM_API FRigVMTemplateArgumentType
 		return CPPType != InOther.CPPType;
 	}
 
-	friend FORCEINLINE uint32 GetTypeHash(const FRigVMTemplateArgumentType& InType)
+	friend uint32 GetTypeHash(const FRigVMTemplateArgumentType& InType)
 	{
 		return GetTypeHash(InType.CPPType);
 	}
@@ -390,7 +390,7 @@ public:
 	FRigVMTemplate_NewArgumentTypeDelegate& OnNewArgumentType() { return Delegates.NewArgumentTypeDelegate; }
 
 	// Returns the factory this template was created by
-	FORCEINLINE const FRigVMDispatchFactory* GetDispatchFactory() const
+	const FRigVMDispatchFactory* GetDispatchFactory() const
 	{
 		if(Delegates.GetDispatchFactoryDelegate.IsBound())
 		{
@@ -400,7 +400,7 @@ public:
 	}
 
 	// Returns true if this template is backed by a dispatch factory
-	FORCEINLINE bool UsesDispatch() const
+	bool UsesDispatch() const
 	{
 		return Delegates.RequestDispatchFunctionDelegate.IsBound() &&
 			Delegates.GetDispatchFactoryDelegate.IsBound();
diff --git a/Engine/Source/Runtime/RigVM/Public/RigVMCore/RigVMTraits.h b/Engine/Source/Runtime/RigVM/Public/RigVMCore/RigVMTraits.h
--- a/Engine/Source/Runtime/RigVM/Public/RigVMCore/RigVMTraits.h
+++ b/Engine/Source/Runtime/RigVM/Public/RigVMCore/RigVMTraits.h
@@ -127,7 +127,7 @@ template<
 	typename T, 
 	typename TEnableIf<TIsArithmetic<T>::Value, bool>::Type* = nullptr
 >
-FORCEINLINE_DEBUGGABLE void RigVMInitialize(void* InPtr, int32 InCount = 1)
+void RigVMInitialize(void* InPtr, int32 InCount = 1)
 {
 	ensure(InCount >= 0);
 	FMemory::Memzero(InPtr, InCount * sizeof(T));
@@ -137,7 +137,7 @@ template<
 	typename T,
 	typename TEnableIf<TIsArithmetic<T>::Value, bool>::Type* = nullptr
 >
-FORCEINLINE_DEBUGGABLE void RigVMDestroy(void* InPtr, int32 InCount = 1)
+void RigVMDestroy(void* InPtr, int32 InCount = 1)
 {
 	// nothing to do
 }
@@ -146,7 +146,7 @@ template<
 	typename T,
 	typename TEnableIf<TIsArithmetic<T>::Value, bool>::Type* = nullptr
 >
-FORCEINLINE_DEBUGGABLE void RigVMCopy(void* InTargetPtr, const void* InSourcePtr, int32 InCount = 1)
+void RigVMCopy(void* InTargetPtr, const void* InSourcePtr, int32 InCount = 1)
 {
 	ensure(InCount >= 0);
 	FMemory::Memcpy(InTargetPtr, InSourcePtr, InCount * sizeof(T));
@@ -156,7 +156,7 @@ template<
 	typename T,
 	typename TEnableIf<TRigVMIsString<T>::Value>::Type* = nullptr
 >
-FORCEINLINE_DEBUGGABLE void RigVMInitialize(void* InPtr, int32 InCount)
+void RigVMInitialize(void* InPtr, int32 InCount)
 {
 	ensure(InCount >= 0);
 	FMemory::Memzero(InPtr, InCount * sizeof(FString));
@@ -172,7 +172,7 @@ template<
 	typename T,
 	typename TEnableIf<TRigVMIsString<T>::Value>::Type* = nullptr
 >
-FORCEINLINE_DEBUGGABLE void RigVMDestroy(void* InPtr, int32 InCount)
+void RigVMDestroy(void* InPtr, int32 InCount)
 {
 	ensure(InCount >= 0);
 	FString* Values = (FString*)InPtr;
@@ -186,7 +186,7 @@ template<
 	typename T,
 	typename TEnableIf<TRigVMIsString<T>::Value>::Type* = nullptr
 >
-FORCEINLINE_DEBUGGABLE void RigVMCopy(void* InTargetPtr, const void* InSourcePtr, int32 InCount = 1)
+void RigVMCopy(void* InTargetPtr, const void* InSourcePtr, int32 InCount = 1)
 {
 	ensure(InCount >= 0);
 
@@ -202,7 +202,7 @@ template<
 	typename T,
 	typename TEnableIf<TRigVMIsName<T>::Value, bool>::Type* = nullptr
 >
-FORCEINLINE_DEBUGGABLE void RigVMInitialize(void* InPtr, int32 InCount)
+void RigVMInitialize(void* InPtr, int32 InCount)
 {
 	ensure(InCount >= 0);
 	FMemory::Memzero(InPtr, InCount * sizeof(FName));
@@ -212,7 +212,7 @@ template<
 	typename T,
 	typename TEnableIf<TRigVMIsName<T>::Value, bool>::Type* = nullptr
 >
-FORCEINLINE_DEBUGGABLE void RigVMDestroy(void* InPtr, int32 InCount)
+void RigVMDestroy(void* InPtr, int32 InCount)
 {
 	ensure(InCount >= 0);
 	FMemory::Memzero(InPtr, InCount * sizeof(FName));
@@ -222,7 +222,7 @@ template<
 	typename T,
 	typename TEnableIf<TRigVMIsName<T>::Value, bool>::Type* = nullptr
 >
-FORCEINLINE_DEBUGGABLE void RigVMCopy(void* InTargetPtr, const void* InSourcePtr, int32 InCount = 1)
+void RigVMCopy(void* InTargetPtr, const void* InSourcePtr, int32 InCount = 1)
 {
 	ensure(InCount >= 0);
 	FMemory::Memcpy(InTargetPtr, InSourcePtr, InCount * sizeof(FName));
@@ -232,7 +232,7 @@ template <
 	typename T,
 	typename TEnableIf<TModels<CRigVMUStruct, T>::Value>::Type* = nullptr
 >
-FORCEINLINE_DEBUGGABLE void RigVMInitialize(void* InPtr, int32 InCount)
+void RigVMInitialize(void* InPtr, int32 InCount)
 {
 	ensure(InCount >= 0);
 	T::StaticStruct()->InitializeStruct(InPtr, InCount);
@@ -242,7 +242,7 @@ template <
 	typename T,
 	typename TEnableIf<TModels<CRigVMUStruct, T>::Value>::Type* = nullptr
 >
-FORCEINLINE_DEBUGGABLE void RigVMDestroy(void* InPtr, int32 InCount)
+void RigVMDestroy(void* InPtr, int32 InCount)
 {
 	ensure(InCount >= 0);
 	T::StaticStruct()->DestroyStruct(InPtr, InCount);
@@ -252,7 +252,7 @@ template <
 	typename T,
 	typename TEnableIf<TModels<CRigVMUStruct, T>::Value>::Type* = nullptr
 >
-FORCEINLINE_DEBUGGABLE void RigVMCopy(void* InTargetPtr, const void* InSourcePtr, int32 InCount = 1)
+void RigVMCopy(void* InTargetPtr, const void* InSourcePtr, int32 InCount = 1)
 {
 	ensure(InCount >= 0);
 	T::StaticStruct()->CopyScriptStruct(InTargetPtr, InSourcePtr, InCount);
@@ -262,7 +262,7 @@ template <
 	typename T,
 	typename TEnableIf<TIsEnum<T>::Value>::Type* = nullptr
 >
-FORCEINLINE_DEBUGGABLE void RigVMInitialize(void* InPtr, int32 InCount)
+void RigVMInitialize(void* InPtr, int32 InCount)
 {
 	FMemory::Memzero(InPtr, InCount * sizeof(T));
 }
@@ -271,7 +271,7 @@ template <
 	typename T,
 	typename TEnableIf<TIsEnum<T>::Value>::Type* = nullptr
 >
-FORCEINLINE_DEBUGGABLE void RigVMDestroy(void* InPtr, int32 InCount)
+void RigVMDestroy(void* InPtr, int32 InCount)
 {
 	// nothing to do
 }
@@ -280,7 +280,7 @@ template <
 	typename T,
 	typename TEnableIf<TIsEnum<T>::Value>::Type* = nullptr
 >
-FORCEINLINE_DEBUGGABLE void RigVMCopy(void* InTargetPtr, const void* InSourcePtr, int32 InCount = 1)
+void RigVMCopy(void* InTargetPtr, const void* InSourcePtr, int32 InCount = 1)
 {
 	FMemory::Memcpy(InTargetPtr, InSourcePtr, InCount * sizeof(T));
 }
@@ -288,7 +288,7 @@ template <
 	typename T,
 	typename TEnableIf<TRigVMIsBaseStructure<T>::Value, T>::Type* = nullptr
 >
-FORCEINLINE_DEBUGGABLE void RigVMInitialize(void* InPtr, int32 InCount)
+void RigVMInitialize(void* InPtr, int32 InCount)
 {
 	ensure(InCount >= 0);
 	TBaseStructure<T>::Get()->InitializeStruct(InPtr, InCount);
@@ -298,7 +298,7 @@ template <
 	typename T,
 	typename TEnableIf<TRigVMIsBaseStructure<T>::Value, T>::Type* = nullptr
 >
-FORCEINLINE_DEBUGGABLE void RigVMDestroy(void* InPtr, int32 InCount)
+void RigVMDestroy(void* InPtr, int32 InCount)
 {
 	// nothing to do
 }
@@ -307,7 +307,7 @@ template <
 	typename T,
 	typename TEnableIf<TRigVMIsBaseStructure<T>::Value, T>::Type* = nullptr
 >
-FORCEINLINE_DEBUGGABLE void RigVMCopy(void* InTargetPtr, const void* InSourcePtr, int32 InCount = 1)
+void RigVMCopy(void* InTargetPtr, const void* InSourcePtr, int32 InCount = 1)
 {
 	ensure(InCount >= 0);
 	FMemory::Memcpy(InTargetPtr, InSourcePtr, InCount * sizeof(T));
diff --git a/Engine/Source/Runtime/RigVM/Public/RigVMCore/RigVMUserWorkflow.h b/Engine/Source/Runtime/RigVM/Public/RigVMCore/RigVMUserWorkflow.h
--- a/Engine/Source/Runtime/RigVM/Public/RigVMCore/RigVMUserWorkflow.h
+++ b/Engine/Source/Runtime/RigVM/Public/RigVMCore/RigVMUserWorkflow.h
@@ -46,12 +46,12 @@ struct RIGVM_API FRigVMUserWorkflow
 
 public:
 
-	FORCEINLINE FRigVMUserWorkflow()
+	FRigVMUserWorkflow()
 	: Type(ERigVMUserWorkflowType::Invalid)
 	, OptionsClass(nullptr)
 	{}
 
-	FORCEINLINE FRigVMUserWorkflow(
+	FRigVMUserWorkflow(
 		const FString& InTitle,
 		const FString& InTooltip,
 		ERigVMUserWorkflowType InType,
@@ -64,19 +64,19 @@ public:
 	, OptionsClass(InOptionsClass)
 	{}
 	
-	FORCEINLINE virtual ~FRigVMUserWorkflow() {}
+	virtual ~FRigVMUserWorkflow() {}
 
-	FORCEINLINE bool IsValid() const
+	bool IsValid() const
 	{
 		return Type != ERigVMUserWorkflowType::Invalid &&
 			GetOptionsClass() != nullptr &&
 			(PerformDelegate.IsBound() || PerformDynamicDelegate.IsBound());
 	}
 
-	FORCEINLINE const FString& GetTitle() const { return Title; }
-	FORCEINLINE const FString& GetTooltip() const { return Tooltip; }
-	FORCEINLINE ERigVMUserWorkflowType GetType() const { return Type; }
-	FORCEINLINE UClass* GetOptionsClass() const { return OptionsClass; }
+	const FString& GetTitle() const { return Title; }
+	const FString& GetTooltip() const { return Tooltip; }
+	ERigVMUserWorkflowType GetType() const { return Type; }
+	UClass* GetOptionsClass() const { return OptionsClass; }
 
 	bool Perform(const URigVMUserWorkflowOptions* InOptions, UObject* InController) const;
 
@@ -110,15 +110,15 @@ class RIGVM_API URigVMUserWorkflowOptions : public UObject
 public:
 
 	UFUNCTION(BlueprintCallable, Category=Options)
-	FORCEINLINE bool IsValid() const { return Subject != nullptr; }
+	bool IsValid() const { return Subject != nullptr; }
 	
 	UFUNCTION(BlueprintCallable, Category=Options)
 	bool RequiresDialog() const;
 
 	template<typename T = UObject>
-	FORCEINLINE T* GetSubject() const { return Cast<T>(Subject.Get()); }
+	T* GetSubject() const { return Cast<T>(Subject.Get()); }
 
-	FORCEINLINE UObject* GetSubjectChecked() const
+	UObject* GetSubjectChecked() const
 	{
 		UObject* Object = Subject.Get();
 		check(Object);
@@ -126,7 +126,7 @@ public:
 	}
 
 	template<typename T>
-	FORCEINLINE T* GetSubjectChecked() const { return CastChecked<T>(Subject.Get()); }
+	T* GetSubjectChecked() const { return CastChecked<T>(Subject.Get()); }
 
 	const FRigVMUserWorkflow& GetWorkflow() const { return Workflow; }
 
@@ -139,19 +139,19 @@ public:
 	}
 
 	UFUNCTION(BlueprintCallable, Category=Options)
-	FORCEINLINE void ReportError(const FString& InMessage)
+	void ReportError(const FString& InMessage)
 	{
 		Report(EMessageSeverity::Error, InMessage);
 	}
 
 	UFUNCTION(BlueprintCallable, Category=Options)
-	FORCEINLINE void ReportWarning(const FString& InMessage)
+	void ReportWarning(const FString& InMessage)
 	{
 		Report(EMessageSeverity::Warning, InMessage);
 	}
 
 	UFUNCTION(BlueprintCallable, Category=Options)
-	FORCEINLINE void ReportInfo(const FString& InMessage)
+	void ReportInfo(const FString& InMessage)
 	{
 		Report(EMessageSeverity::Info, InMessage);
 	}
diff --git a/Engine/Source/Runtime/RigVM/Public/RigVMTypeUtils.h b/Engine/Source/Runtime/RigVM/Public/RigVMTypeUtils.h
--- a/Engine/Source/Runtime/RigVM/Public/RigVMTypeUtils.h
+++ b/Engine/Source/Runtime/RigVM/Public/RigVMTypeUtils.h
@@ -70,22 +70,22 @@ namespace RigVMTypeUtils
 	};
 
 	// Returns true if the type specified is an array
-	FORCEINLINE bool IsArrayType(const FString& InCPPType)
+	inline bool IsArrayType(const FString& InCPPType)
 	{
 		return InCPPType.StartsWith(TArrayPrefix);
 	}
 
-	FORCEINLINE FString ArrayTypeFromBaseType(const FString& InCPPType)
+	inline FString ArrayTypeFromBaseType(const FString& InCPPType)
 	{
 		return FString::Printf(TArrayTemplate, *InCPPType);
 	}
 
-	FORCEINLINE FString BaseTypeFromArrayType(const FString& InCPPType)
+	inline FString BaseTypeFromArrayType(const FString& InCPPType)
 	{
 		return InCPPType.RightChop(7).LeftChop(1).TrimStartAndEnd();
 	}
 
-	FORCEINLINE_DEBUGGABLE FString GetUniqueStructTypeName(const UScriptStruct* InScriptStruct)
+	inline FString GetUniqueStructTypeName(const UScriptStruct* InScriptStruct)
 	{
 		if (const UUserDefinedStruct* UserDefinedStruct = Cast<UUserDefinedStruct>(InScriptStruct))
 		{
@@ -95,7 +95,7 @@ namespace RigVMTypeUtils
 		return InScriptStruct->GetStructCPPName();
 	}
 
-	FORCEINLINE FString CPPTypeFromEnum(UEnum* InEnum)
+	inline FString CPPTypeFromEnum(UEnum* InEnum)
 	{
 		check(InEnum);
 
@@ -107,12 +107,12 @@ namespace RigVMTypeUtils
 		return CPPType;
 	}
 
-	FORCEINLINE bool IsUObjectType(const FString& InCPPType)
+	inline bool IsUObjectType(const FString& InCPPType)
 	{
 		return InCPPType.StartsWith(TObjectPtrPrefix);
 	}
 
-	FORCEINLINE bool IsInterfaceType(const FString& InCPPType)
+	inline bool IsInterfaceType(const FString& InCPPType)
 	{
 		return InCPPType.StartsWith(TScriptInterfacePrefix);
 	}
@@ -149,7 +149,7 @@ namespace RigVMTypeUtils
 
 
 
-	FORCEINLINE FString PostProcessCPPType(const FString& InCPPType, UObject* InCPPTypeObject)
+	inline FString PostProcessCPPType(const FString& InCPPType, UObject* InCPPTypeObject)
 	{
 		FString CPPType = InCPPType;
 	
