Control Rig: Fix UE_RIGVM_DEBUG_EXECUTION compilation

diff --git a/Engine/Source/Runtime/RigVM/Private/RigVMCore/RigVM.cpp b/Engine/Source/Runtime/RigVM/Private/RigVMCore/RigVM.cpp
--- a/Engine/Source/Runtime/RigVM/Private/RigVMCore/RigVM.cpp
+++ b/Engine/Source/Runtime/RigVM/Private/RigVMCore/RigVM.cpp
@@ -1891,16 +1891,15 @@ ERigVMExecuteResult URigVM::Execute(TArrayView<URigVMMemoryStorage*> Memory, con
 	StartProfiling();
 	
 #if UE_RIGVM_DEBUG_EXECUTION
-	FString DebugMemoryString;
-	TArray<FString> PreviousWorkMemory;
-	UEnum* InstanceOpCodeEnum = StaticEnum<ERigVMOpCode>();
+	ContextPublicData.DebugMemoryString;
+	ContextPublicData.InstanceOpCodeEnum = StaticEnum<ERigVMOpCode>();
 	URigVMMemoryStorage* LiteralMemory = GetLiteralMemory(false);
-	DebugMemoryString = FString("\n\nLiteral Memory\n\n");
+	ContextPublicData.DebugMemoryString = FString("\n\nLiteral Memory\n\n");
 	for (int32 PropertyIndex=0; PropertyIndex<LiteralMemory->Num(); ++PropertyIndex)
 	{
-		DebugMemoryString += FString::Printf(TEXT("%s: %s\n"), *LiteralMemory->GetProperties()[PropertyIndex]->GetFullName(), *LiteralMemory->GetDataAsString(PropertyIndex));				
+		ContextPublicData.DebugMemoryString += FString::Printf(TEXT("%s: %s\n"), *LiteralMemory->GetProperties()[PropertyIndex]->GetFullName(), *LiteralMemory->GetDataAsString(PropertyIndex));				
 	}
-	DebugMemoryString += FString(TEXT("\n\nWork Memory\n\n"));
+	ContextPublicData.DebugMemoryString += FString(TEXT("\n\nWork Memory\n\n"));
 	
 #endif
 	
@@ -1947,7 +1946,7 @@ ERigVMExecuteResult URigVM::ExecuteInstructions(int32 InFirstInstruction, int32
 		if (ShouldHaltAtInstruction(CurrentEntryName, ContextPublicData.InstructionIndex))
 		{
 #if UE_RIGVM_DEBUG_EXECUTION
-			ContextPublicData.Log(EMessageSeverity::Info, DebugMemoryString);					
+			ContextPublicData.Log(EMessageSeverity::Info, ContextPublicData.DebugMemoryString);					
 #endif
 			// we'll recursively exit all invoked
 			// entries here.
@@ -1989,16 +1988,16 @@ ERigVMExecuteResult URigVM::ExecuteInstructions(int32 InFirstInstruction, int32
 				Labels.Add(GetOperandLabel(Operand));
 			}
 
-			DebugMemoryString += FString::Printf(TEXT("Instruction %d: %s(%s)\n"), ContextPublicData.InstructionIndex, *FunctionNames[Op.FunctionIndex].ToString(), *FString::Join(Labels, TEXT(", ")));				
+			ContextPublicData.DebugMemoryString += FString::Printf(TEXT("Instruction %d: %s(%s)\n"), ContextPublicData.InstructionIndex, *FunctionNames[Op.FunctionIndex].ToString(), *FString::Join(Labels, TEXT(", ")));				
 		}
 		else if(Instruction.OpCode == ERigVMOpCode::Copy)
 		{
 			const FRigVMCopyOp& Op = ByteCode.GetOpAt<FRigVMCopyOp>(Instruction);
-			DebugMemoryString += FString::Printf(TEXT("Instruction %d: Copy %s -> %s\n"), ContextPublicData.InstructionIndex, *GetOperandLabel(Op.Source), *GetOperandLabel(Op.Target));				
+			ContextPublicData.DebugMemoryString += FString::Printf(TEXT("Instruction %d: Copy %s -> %s\n"), ContextPublicData.InstructionIndex, *GetOperandLabel(Op.Source), *GetOperandLabel(Op.Target));				
 		}
 		else
 		{
-			DebugMemoryString += FString::Printf(TEXT("Instruction %d: %s\n"), ContextPublicData.InstructionIndex, *InstanceOpCodeEnum->GetNameByIndex((uint8)Instruction.OpCode).ToString());				
+			ContextPublicData.DebugMemoryString += FString::Printf(TEXT("Instruction %d: %s\n"), ContextPublicData.InstructionIndex, *ContextPublicData.InstanceOpCodeEnum->GetNameByIndex((uint8)Instruction.OpCode).ToString());				
 		}
 		
 #endif
@@ -2270,7 +2269,7 @@ ERigVMExecuteResult URigVM::ExecuteInstructions(int32 InFirstInstruction, int32
 						ExecutionHalted().Broadcast(INDEX_NONE, nullptr, CurrentEntryName);
 					}
 #if UE_RIGVM_DEBUG_EXECUTION
-					ContextPublicData.Log(EMessageSeverity::Info, DebugMemoryString);					
+					ContextPublicData.Log(EMessageSeverity::Info, ContextPublicData.DebugMemoryString);					
 #endif
 #endif
 				}
@@ -2931,11 +2930,11 @@ ERigVMExecuteResult URigVM::ExecuteInstructions(int32 InFirstInstruction, int32
 		for (int32 PropertyIndex=0; PropertyIndex<WorkMemory->Num(); ++PropertyIndex, ++LineIndex)
 		{
 			FString Line = FString::Printf(TEXT("%s: %s"), *WorkMemory->GetProperties()[PropertyIndex]->GetFullName(), *WorkMemory->GetDataAsString(PropertyIndex));
-			if (PreviousWorkMemory.Num() > 0 && PreviousWorkMemory[PropertyIndex].StartsWith(TEXT(" -- ")))
+			if (ContextPublicData.PreviousWorkMemory.Num() > 0 && ContextPublicData.PreviousWorkMemory[PropertyIndex].StartsWith(TEXT(" -- ")))
 			{
-				PreviousWorkMemory[PropertyIndex].RightChopInline(4);
+				ContextPublicData.PreviousWorkMemory[PropertyIndex].RightChopInline(4);
 			}
-			if (PreviousWorkMemory.Num() == 0 || Line == PreviousWorkMemory[PropertyIndex])
+			if (ContextPublicData.PreviousWorkMemory.Num() == 0 || Line == ContextPublicData.PreviousWorkMemory[PropertyIndex])
 			{
 				CurrentWorkMemory.Add(Line);
 			}
@@ -2949,11 +2948,11 @@ ERigVMExecuteResult URigVM::ExecuteInstructions(int32 InFirstInstruction, int32
 			FString Value;
 			ExternalVariable.Property->ExportTextItem_Direct(Value, ExternalVariable.Memory, nullptr, nullptr, PPF_None);
 			FString Line = FString::Printf(TEXT("External %s: %s"), *ExternalVariable.Name.ToString(), *Value);
-			if (PreviousWorkMemory.Num() > 0 && PreviousWorkMemory[LineIndex].StartsWith(TEXT(" -- ")))
+			if (ContextPublicData.PreviousWorkMemory.Num() > 0 && ContextPublicData.PreviousWorkMemory[LineIndex].StartsWith(TEXT(" -- ")))
 			{
-				PreviousWorkMemory[LineIndex].RightChopInline(4);
+				ContextPublicData.PreviousWorkMemory[LineIndex].RightChopInline(4);
 			}
-			if (PreviousWorkMemory.Num() == 0 || Line == PreviousWorkMemory[LineIndex])
+			if (ContextPublicData.PreviousWorkMemory.Num() == 0 || Line == ContextPublicData.PreviousWorkMemory[LineIndex])
 			{
 				CurrentWorkMemory.Add(Line);
 			}
@@ -2963,8 +2962,8 @@ ERigVMExecuteResult URigVM::ExecuteInstructions(int32 InFirstInstruction, int32
 			}
 			++LineIndex;
 		}
-		DebugMemoryString += FString::Join(CurrentWorkMemory, TEXT("\n")) + FString(TEXT("\n\n"));
-		PreviousWorkMemory = CurrentWorkMemory;		
+		ContextPublicData.DebugMemoryString += FString::Join(CurrentWorkMemory, TEXT("\n")) + FString(TEXT("\n\n"));
+		ContextPublicData.PreviousWorkMemory = CurrentWorkMemory;		
 #endif
 #endif
 	}
diff --git a/Engine/Source/Runtime/RigVM/Public/RigVMCore/RigVMExecuteContext.h b/Engine/Source/Runtime/RigVM/Public/RigVMCore/RigVMExecuteContext.h
--- a/Engine/Source/Runtime/RigVM/Public/RigVMCore/RigVMExecuteContext.h
+++ b/Engine/Source/Runtime/RigVM/Public/RigVMCore/RigVMExecuteContext.h
@@ -315,6 +315,14 @@ protected:
 
 	mutable FRigVMNameCache* NameCache;
 
+#if UE_RIGVM_DEBUG_EXECUTION
+	FString DebugMemoryString;
+
+	TArray<FString> PreviousWorkMemory;
+
+	UEnum* InstanceOpCodeEnum;
+#endif
+
 	friend struct FRigVMExtendedExecuteContext;
 	friend class URigVM;
 	friend class URigVMNativized;
