Control Rig: Fix initialization of hidden properties in unit nodes

diff --git a/Engine/Source/Developer/RigVMDeveloper/Private/RigVMCompiler/RigVMCompiler.cpp b/Engine/Source/Developer/RigVMDeveloper/Private/RigVMCompiler/RigVMCompiler.cpp
--- a/Engine/Source/Developer/RigVMDeveloper/Private/RigVMCompiler/RigVMCompiler.cpp
+++ b/Engine/Source/Developer/RigVMDeveloper/Private/RigVMCompiler/RigVMCompiler.cpp
@@ -3676,11 +3676,6 @@ FRigVMOperand URigVMCompiler::FindOrAddRegister(const FRigVMVarExprAST* InVarExp
 			}
 		}
 
-		if(Pin->GetDirection() == ERigVMPinDirection::Hidden)
-		{
-			JoinedDefaultValue.Empty();
-		}
-	
 		Operand = WorkData.AddProperty(MemoryType, RegisterName, CPPType, Pin->GetCPPTypeObject(), JoinedDefaultValue);
 	}
 	ensure(Operand.IsValid());
diff --git a/Engine/Source/Runtime/RigVM/Public/RigVMCore/RigVMMemoryStorage.h b/Engine/Source/Runtime/RigVM/Public/RigVMCore/RigVMMemoryStorage.h
--- a/Engine/Source/Runtime/RigVM/Public/RigVMCore/RigVMMemoryStorage.h
+++ b/Engine/Source/Runtime/RigVM/Public/RigVMCore/RigVMMemoryStorage.h
@@ -787,7 +787,25 @@ private:
 			if(ArrayHelper.Num() <= InSliceIndex)
 			{
 				const int32 NumValuesToAdd = 1 + InSliceIndex - ArrayHelper.Num();
-				ArrayHelper.AddValues(NumValuesToAdd);
+				const int32 FirstAddedIndex = ArrayHelper.AddValues(NumValuesToAdd);
+
+				if (FirstAddedIndex > 0)
+				{
+					// Adding slices, we need to initialize the new values.
+					// Otherwise, if we are adding the first slice, we don't need to worry about initializing.
+					UObject* CDO = ArrayProperty->GetOwnerClass()->GetDefaultObject();
+					const uint8* DefaultArrayMemory = ArrayProperty->ContainerPtrToValuePtr<uint8>(CDO);		
+					FScriptArrayHelper DefaultArrayHelper(ArrayProperty, DefaultArrayMemory);
+					if (const uint8* DefaultElementMemory = DefaultArrayHelper.GetRawPtr(0))
+					{
+						FProperty* ElementProperty = ArrayProperty->Inner;
+						for (int32 i=FirstAddedIndex; i<ArrayHelper.Num(); ++i)
+						{
+							uint8* DestMemory = ArrayHelper.GetRawPtr(i);
+							ElementProperty->CopyCompleteValue(DestMemory, DefaultElementMemory);
+						}
+					}
+				}
 			}
 
 			return ArrayHelper.GetRawPtr(InSliceIndex);
