RigVM: Add GetPublicDataSafe to FRigVMExecuteContext

diff --git a/Engine/Plugins/Animation/ControlRig/Source/ControlRig/Private/ControlRig.cpp b/Engine/Plugins/Animation/ControlRig/Source/ControlRig/Private/ControlRig.cpp
--- a/Engine/Plugins/Animation/ControlRig/Source/ControlRig/Private/ControlRig.cpp
+++ b/Engine/Plugins/Animation/ControlRig/Source/ControlRig/Private/ControlRig.cpp
@@ -730,7 +730,7 @@ bool UControlRig::Execute(const EControlRigState InState, const FName& InEventNa
 	LatestExecutedState = InState;
 
 	FRigVMExtendedExecuteContext& ExtendedExecuteContext = GetVM()->GetContext();
-	FControlRigExecuteContext& PublicContext = ExtendedExecuteContext.GetPublicData<FControlRigExecuteContext>();
+	FControlRigExecuteContext& PublicContext = ExtendedExecuteContext.GetPublicDataSafe<FControlRigExecuteContext>();
 	PublicContext.SetDeltaTime(DeltaTime);
 	PublicContext.SetAbsoluteTime(AbsoluteTime);
 
diff --git a/Engine/Source/Runtime/RigVM/Public/RigVMCore/RigVMExecuteContext.h b/Engine/Source/Runtime/RigVM/Public/RigVMCore/RigVMExecuteContext.h
--- a/Engine/Source/Runtime/RigVM/Public/RigVMCore/RigVMExecuteContext.h
+++ b/Engine/Source/Runtime/RigVM/Public/RigVMCore/RigVMExecuteContext.h
@@ -436,6 +436,16 @@ struct RIGVM_API FRigVMExtendedExecuteContext
 		return *(ExecuteContextType*)PublicDataScope.GetStructMemory();
 	}
 
+	template<typename ExecuteContextType = FRigVMExecuteContext>
+	ExecuteContextType& GetPublicDataSafe()
+	{
+		if(!PublicDataScope.GetStruct()->IsChildOf(ExecuteContextType::StaticStruct()))
+		{
+			Initialize(ExecuteContextType::StaticStruct());
+		}
+		return *(ExecuteContextType*)PublicDataScope.GetStructMemory();
+	}
+
 	const FRigVMSlice& GetSlice() const
 	{
 		const int32 SliceOffset = (int32)SliceOffsets[GetPublicData<>().InstructionIndex];
