RigVM: New control flow Dispatch

diff --git a/Engine/Plugins/Animation/ControlRig/Source/ControlRigDeveloper/Private/Graph/ControlRigGraphNode.cpp b/Engine/Plugins/Animation/ControlRig/Source/ControlRigDeveloper/Private/Graph/ControlRigGraphNode.cpp
--- a/Engine/Plugins/Animation/ControlRig/Source/ControlRigDeveloper/Private/Graph/ControlRigGraphNode.cpp
+++ b/Engine/Plugins/Animation/ControlRig/Source/ControlRigDeveloper/Private/Graph/ControlRigGraphNode.cpp
@@ -674,7 +674,7 @@ bool UControlRigGraphNode::CreateGraphPinFromModelPin(const URigVMPin* InModelPi
 	UEdGraphPin* InParentPin)
 {
 	// don't create output pins for array elements
-	if(InDirection == EGPD_Output && InModelPin->IsArrayElement())
+	if(InDirection == EGPD_Output && InModelPin->IsArrayElement() && !InModelPin->GetParentPin()->IsFixedSizeArray())
 	{
 		return false;
 	}
diff --git a/Engine/Plugins/Animation/ControlRig/Source/ControlRigEditor/Private/ControlRigEditorModule.cpp b/Engine/Plugins/Animation/ControlRig/Source/ControlRigEditor/Private/ControlRigEditorModule.cpp
--- a/Engine/Plugins/Animation/ControlRig/Source/ControlRigEditor/Private/ControlRigEditorModule.cpp
+++ b/Engine/Plugins/Animation/ControlRig/Source/ControlRigEditor/Private/ControlRigEditorModule.cpp
@@ -1426,7 +1426,7 @@ void FControlRigEditorModule::GetContextMenuActions(const UControlRigGraphSchema
 							})
 						));
 		}
-					if(ModelPin->IsArrayElement() && !bIsExecutePin)
+					if(ModelPin->IsArrayElement())
 					{
 						FToolMenuSection& Section = Menu->AddSection("EdGraphSchemaPinArrays", LOCTEXT("PinArrays", "Arrays"));
 						Section.AddMenuEntry(
diff --git a/Engine/Source/Developer/RigVMDeveloper/Private/RigVMCompiler/RigVMAST.cpp b/Engine/Source/Developer/RigVMDeveloper/Private/RigVMCompiler/RigVMAST.cpp
--- a/Engine/Source/Developer/RigVMDeveloper/Private/RigVMCompiler/RigVMAST.cpp
+++ b/Engine/Source/Developer/RigVMDeveloper/Private/RigVMCompiler/RigVMAST.cpp
@@ -832,40 +832,46 @@ FRigVMExprAST* FRigVMParserAST::TraverseMutableNode(const FRigVMASTProxy& InNode
 
 		TraversePins(InNodeProxy, NodeExpr);
 
-		for (URigVMPin* SourcePin : Node->GetPins())
+		TArray<URigVMPin*> SourcePins = Node->GetPins().FilterByPredicate([](const URigVMPin* InPin) -> bool
 		{
-			if (SourcePin->GetDirection() == ERigVMPinDirection::Output || SourcePin->GetDirection() == ERigVMPinDirection::IO)
-			{
-				if (SourcePin->IsExecuteContext())
-				{
-					FRigVMASTProxy SourcePinProxy = InNodeProxy.GetSibling(SourcePin);
+			return (InPin->GetDirection() == ERigVMPinDirection::Output || InPin->GetDirection() == ERigVMPinDirection::IO) && InPin->IsExecuteContext();
+		});
 
-					FRigVMExprAST* ParentExpr = InParentExpr;
-					if (Node->IsControlFlowNode())
-					{
-						if (FRigVMExprAST** PinExpr = SubjectToExpression.Find(SourcePinProxy))
-						{
-							FRigVMBlockExprAST* BlockExpr = MakeExpr<FRigVMBlockExprAST>(FRigVMExprAST::EType::Block, FRigVMASTProxy());
-							BlockExpr->AddParent(*PinExpr);
-							BlockExpr->Name = SourcePin->GetFName();
-							ParentExpr = BlockExpr;
-						}
-					}
+		for(int32 SourcePinIndex = 0; SourcePinIndex < SourcePins.Num(); SourcePinIndex++)
+		{
+			URigVMPin* SourcePin = SourcePins[SourcePinIndex];
+			FRigVMASTProxy SourcePinProxy = InNodeProxy.GetSibling(SourcePin);
 
-					const TArray<int32>& LinkIndices = GetTargetLinkIndices(SourcePinProxy);
-					for(const int32 LinkIndex : LinkIndices)
-					{
-						const FRigVMASTLinkDescription& Link = Links[LinkIndex];
-						if(ShouldLinkBeSkipped(Link))
-						{
-							continue;
-						}
+			FRigVMExprAST* ParentExpr = InParentExpr;
 
-						URigVMNode* TargetNode = Link.TargetProxy.GetSubjectChecked<URigVMPin>()->GetNode();
-						FRigVMASTProxy TargetNodeProxy = Link.TargetProxy.GetSibling(TargetNode);
-						TraverseMutableNode(TargetNodeProxy, ParentExpr);
-					}
+			if(SourcePin->IsFixedSizeArray())
+			{
+				// also process elements of execute fixed arrays
+				SourcePins.Append(SourcePin->GetSubPins());
+			}
+			else if (Node->IsControlFlowNode())
+			{
+				if (FRigVMExprAST** PinExpr = SubjectToExpression.Find(SourcePinProxy))
+				{
+					FRigVMBlockExprAST* BlockExpr = MakeExpr<FRigVMBlockExprAST>(FRigVMExprAST::EType::Block, FRigVMASTProxy());
+					BlockExpr->AddParent(*PinExpr);
+					BlockExpr->Name = SourcePin->GetFName();
+					ParentExpr = BlockExpr;
+				}
+			}
+
+			const TArray<int32>& LinkIndices = GetTargetLinkIndices(SourcePinProxy);
+			for(const int32 LinkIndex : LinkIndices)
+			{
+				const FRigVMASTLinkDescription& Link = Links[LinkIndex];
+				if(ShouldLinkBeSkipped(Link))
+				{
+					continue;
 				}
+
+				URigVMNode* TargetNode = Link.TargetProxy.GetSubjectChecked<URigVMPin>()->GetNode();
+				FRigVMASTProxy TargetNodeProxy = Link.TargetProxy.GetSibling(TargetNode);
+				TraverseMutableNode(TargetNodeProxy, ParentExpr);
 			}
 		}
 	}
diff --git a/Engine/Source/Developer/RigVMDeveloper/Private/RigVMCompiler/RigVMCompiler.cpp b/Engine/Source/Developer/RigVMDeveloper/Private/RigVMCompiler/RigVMCompiler.cpp
--- a/Engine/Source/Developer/RigVMDeveloper/Private/RigVMCompiler/RigVMCompiler.cpp
+++ b/Engine/Source/Developer/RigVMDeveloper/Private/RigVMCompiler/RigVMCompiler.cpp
@@ -3460,9 +3460,13 @@ FRigVMOperand URigVMCompiler::FindOrAddRegister(const FRigVMVarExprAST* InVarExp
 					bValidHiddenPin = true;
 				}
 			}
-			else if(Pin->GetNode()->IsA<URigVMDispatchNode>())
+			else if(URigVMDispatchNode* DispatchNode = Cast<URigVMDispatchNode>(Pin->GetNode()))
 			{
 				bValidHiddenPin = true;
+				if(const FRigVMDispatchFactory* Factory = DispatchNode->GetFactory())
+				{
+					bValidHiddenPin = !Factory->HasArgumentMetaData(Pin->GetFName(), FRigVMStruct::SingletonMetaName);
+				}
 			}
 
 			if(bValidHiddenPin)
diff --git a/Engine/Source/Developer/RigVMDeveloper/Private/RigVMModel/Nodes/RigVMTemplateNode.cpp b/Engine/Source/Developer/RigVMDeveloper/Private/RigVMModel/Nodes/RigVMTemplateNode.cpp
--- a/Engine/Source/Developer/RigVMDeveloper/Private/RigVMModel/Nodes/RigVMTemplateNode.cpp
+++ b/Engine/Source/Developer/RigVMDeveloper/Private/RigVMModel/Nodes/RigVMTemplateNode.cpp
@@ -1,6 +1,7 @@
 // Copyright Epic Games, Inc. All Rights Reserved.
 
 #include "RigVMModel/Nodes/RigVMTemplateNode.h"
+#include "RigVMStringUtils.h"
 #include "RigVMCore/RigVMTemplate.h"
 #include "RigVMModel/RigVMController.h"
 #include "UObject/ObjectSaveContext.h"
@@ -724,11 +725,18 @@ FName URigVMTemplateNode::GetDisplayNameForPin(const FName& InRootPinName,
 		const FText DisplayNameText = Template->GetDisplayNameForArgument(InRootPinName, *PermutationIndicesPtr);
 		if(DisplayNameText.IsEmpty())
 		{
-			return InRootPinName;
+			return NAME_None;
+		}
+
+		FString DefaultDisplayName = InRootPinName.ToString();
+		FString Left, Right;
+		if(RigVMStringUtils::SplitPinPathAtEnd(DefaultDisplayName, Left, Right))
+		{
+			DefaultDisplayName = Right;
 		}
 
 		const FName DisplayName = *DisplayNameText.ToString();
-		if(DisplayName.IsEqual(InRootPinName))
+		if(DisplayName.IsEqual(*DefaultDisplayName))
 		{
 			return NAME_None;
 		}
diff --git a/Engine/Source/Developer/RigVMDeveloper/Private/RigVMModel/RigVMController.cpp b/Engine/Source/Developer/RigVMDeveloper/Private/RigVMModel/RigVMController.cpp
--- a/Engine/Source/Developer/RigVMDeveloper/Private/RigVMModel/RigVMController.cpp
+++ b/Engine/Source/Developer/RigVMDeveloper/Private/RigVMModel/RigVMController.cpp
@@ -16355,8 +16355,8 @@ bool URigVMController::FullyResolveTemplateNode(URigVMTemplateNode* InNode, int3
 
 				if(Factory)
 				{
-					const FText DisplayNameText = Factory->GetDisplayNameForArgument(MissingPin.GetName());
-					if(!DisplayNameText.IsEmpty())
+					const FName DisplayNameText = Factory->GetDisplayNameForArgument(MissingPin.GetName());
+					if(!DisplayNameText.IsNone())
 					{
 						Pin->DisplayName = *DisplayNameText.ToString();
 					}
diff --git a/Engine/Source/Developer/RigVMDeveloper/Private/RigVMModel/RigVMPin.cpp b/Engine/Source/Developer/RigVMDeveloper/Private/RigVMModel/RigVMPin.cpp
--- a/Engine/Source/Developer/RigVMDeveloper/Private/RigVMModel/RigVMPin.cpp
+++ b/Engine/Source/Developer/RigVMDeveloper/Private/RigVMModel/RigVMPin.cpp
@@ -357,6 +357,14 @@ FName URigVMPin::GetDisplayName() const
 {
 	if (DisplayName == NAME_None)
 	{
+		if(const URigVMTemplateNode* Node = Cast<URigVMTemplateNode>(GetNode()))
+		{
+			const FName DisplayNameForArgument = Node->GetDisplayNameForPin(*this->GetSegmentPath(true));
+			if(!DisplayNameForArgument.IsNone())
+			{
+				return DisplayNameForArgument;
+			}
+		}
 		return GetFName();
 	}
 
diff --git a/Engine/Source/Runtime/RigVM/Private/RigVMCore/RigVM.cpp b/Engine/Source/Runtime/RigVM/Private/RigVMCore/RigVM.cpp
--- a/Engine/Source/Runtime/RigVM/Private/RigVMCore/RigVM.cpp
+++ b/Engine/Source/Runtime/RigVM/Private/RigVMCore/RigVM.cpp
@@ -8,6 +8,7 @@
 #include "HAL/PlatformTLS.h"
 #include "Async/ParallelFor.h"
 #include "GenericPlatform/GenericPlatformSurvey.h"
+#include "RigVMCore/RigVMStruct.h"
 #include "UObject/UObjectHash.h"
 #include "UObject/UObjectIterator.h"
 
@@ -2862,7 +2863,7 @@ ERigVMExecuteResult URigVM::ExecuteInstructions(int32 InFirstInstruction, int32
 			{
 				const FRigVMJumpToBranchOp& Op = ByteCode.GetOpAt<FRigVMJumpToBranchOp>(Instruction);
 				// BranchLabel = Op.Arg
-				const FName& BranchLabel = *(FName*)CachedMemoryHandles[FirstHandleForInstruction[ContextPublicData.InstructionIndex]].GetData();
+				FName BranchLabel = *(FName*)CachedMemoryHandles[FirstHandleForInstruction[ContextPublicData.InstructionIndex]].GetData();
 
 				// iterate over the branches stored in the bytecode,
 				// starting at the first branch index stored in the operator.
@@ -2875,26 +2876,39 @@ ERigVMExecuteResult URigVM::ExecuteInstructions(int32 InFirstInstruction, int32
 					UE_LOG(LogRigVM, Error, TEXT("No branches in ByteCode - but JumpToBranch instruction %d found. Likely a corrupt VM. Exiting."), ContextPublicData.InstructionIndex);
 					return CurrentExecuteResult = ERigVMExecuteResult::Failed;
 				}
-					
-				for(int32 BranchIndex = Op.FirstBranchInfoIndex // start at the first branch known to this jump op
-					; BranchIndex < Branches.Num(); BranchIndex++)
+
+				for(int32 PassIndex = 0; PassIndex < 2; PassIndex++)
 				{
-					const FRigVMBranchInfo& Branch = Branches[BranchIndex];
-					if(Branch.InstructionIndex != ContextPublicData.InstructionIndex)
+					for(int32 BranchIndex = Op.FirstBranchInfoIndex // start at the first branch known to this jump op
+						; BranchIndex < Branches.Num(); BranchIndex++)
 					{
-						break;
+						const FRigVMBranchInfo& Branch = Branches[BranchIndex];
+						if(Branch.InstructionIndex != ContextPublicData.InstructionIndex)
+						{
+							break;
+						}
+						if(Branch.Label == BranchLabel)
+						{
+							ContextPublicData.InstructionIndex = Branch.FirstInstruction;
+							bBranchFound = true;
+							break;
+						}
 					}
-					if(Branch.Label == BranchLabel)
+
+					// if we don't find the branch - try to jump to the completed branch
+					if (!bBranchFound)
 					{
-						ContextPublicData.InstructionIndex = Branch.FirstInstruction;
-						bBranchFound = true;
-						break;
+						if(PassIndex == 0 && BranchLabel != FRigVMStruct::ControlFlowCompletedName)
+						{
+							UE_LOG(LogRigVM, Warning, TEXT("Branch '%s' was not found for instruction %d."), *BranchLabel.ToString(), ContextPublicData.InstructionIndex);
+							BranchLabel = FRigVMStruct::ControlFlowCompletedName;
+							continue;
+						}
+						
+						UE_LOG(LogRigVM, Error, TEXT("Branch '%s' was not found for instruction %d."), *BranchLabel.ToString(), ContextPublicData.InstructionIndex);
+						return ERigVMExecuteResult::Failed;
 					}
-				}
-				if (!bBranchFound)
-				{
-					UE_LOG(LogRigVM, Error, TEXT("Branch '%s' was not found for instruction %d. Exiting."), *BranchLabel.ToString(), ContextPublicData.InstructionIndex);
-					return CurrentExecuteResult = ERigVMExecuteResult::Failed;
+					break;
 				}
 				break;
 			}
diff --git a/Engine/Source/Runtime/RigVM/Private/RigVMCore/RigVMDispatchFactory.cpp b/Engine/Source/Runtime/RigVM/Private/RigVMCore/RigVMDispatchFactory.cpp
--- a/Engine/Source/Runtime/RigVM/Private/RigVMCore/RigVMDispatchFactory.cpp
+++ b/Engine/Source/Runtime/RigVM/Private/RigVMCore/RigVMDispatchFactory.cpp
@@ -3,10 +3,12 @@
 #include "RigVMCore/RigVMDispatchFactory.h"
 #include "RigVMCore/RigVMRegistry.h"
 #include "RigVMCore/RigVMStruct.h"
+#include "RigVMStringUtils.h"
 
 #include UE_INLINE_GENERATED_CPP_BY_NAME(RigVMDispatchFactory)
 
 const FString FRigVMDispatchFactory::DispatchPrefix = TEXT("DISPATCH_");
+const FString FRigVMDispatchFactory::TrueString = TEXT("True");
 
 FName FRigVMDispatchFactory::GetFactoryName() const
 {
@@ -21,6 +23,25 @@ FName FRigVMDispatchFactory::GetNextAggregateName(const FName& InLastAggregatePi
 	return FRigVMStruct().GetNextAggregateName(InLastAggregatePinName);
 }
 
+FName FRigVMDispatchFactory::GetDisplayNameForArgument(const FName& InArgumentName) const
+{
+	if(InArgumentName == FRigVMStruct::ExecuteContextName)
+	{
+		return FRigVMStruct::ExecuteName;
+	}
+	return NAME_None;
+}
+
+FString FRigVMDispatchFactory::GetArgumentMetaData(const FName& InArgumentName, const FName& InMetaDataKey) const
+{
+	if(InArgumentName == FRigVMStruct::ControlFlowBlockToRunName &&
+		InMetaDataKey == FRigVMStruct::SingletonMetaName)
+	{
+		return TrueString;
+	}
+	return FString();
+}
+
 FLinearColor FRigVMDispatchFactory::GetNodeColor() const
 {
 	if(const UScriptStruct* ScriptStruct = GetScriptStruct())
diff --git a/Engine/Source/Runtime/RigVM/Private/RigVMCore/RigVMTemplate.cpp b/Engine/Source/Runtime/RigVM/Private/RigVMCore/RigVMTemplate.cpp
--- a/Engine/Source/Runtime/RigVM/Private/RigVMCore/RigVMTemplate.cpp
+++ b/Engine/Source/Runtime/RigVM/Private/RigVMCore/RigVMTemplate.cpp
@@ -836,7 +836,12 @@ FText FRigVMTemplate::GetDisplayNameForArgument(const FName& InArgumentName, con
 {
 	if(const FRigVMDispatchFactory* Factory = GetDispatchFactory())
 	{
-		return Factory->GetDisplayNameForArgument(InArgumentName);
+		const FName DisplayName = Factory->GetDisplayNameForArgument(InArgumentName);
+		if(DisplayName.IsNone())
+		{
+			return FText();
+		}
+		return FText::FromName(DisplayName);
 	}
 
 	if(const FRigVMTemplateArgument* Argument = FindArgument(InArgumentName))
diff --git a/Engine/Source/Runtime/RigVM/Private/RigVMFunctions/RigVMDispatch_Core.cpp b/Engine/Source/Runtime/RigVM/Private/RigVMFunctions/RigVMDispatch_Core.cpp
--- a/Engine/Source/Runtime/RigVM/Private/RigVMFunctions/RigVMDispatch_Core.cpp
+++ b/Engine/Source/Runtime/RigVM/Private/RigVMFunctions/RigVMDispatch_Core.cpp
@@ -6,8 +6,6 @@
 
 #include UE_INLINE_GENERATED_CPP_BY_NAME(RigVMDispatch_Core)
 
-const FString FRigVMDispatch_CoreBase::TrueString = TEXT("True");
-
 TArray<FRigVMTemplateArgument> FRigVMDispatch_CoreEquals::GetArguments() const
 {
 	const TArray<FRigVMTemplateArgument::ETypeCategory> ValueCategories = {
diff --git a/Engine/Source/Runtime/RigVM/Private/RigVMFunctions/RigVMDispatch_Switch.cpp b/Engine/Source/Runtime/RigVM/Private/RigVMFunctions/RigVMDispatch_Switch.cpp
new file mode 100644
--- /dev/null
+++ b/Engine/Source/Runtime/RigVM/Private/RigVMFunctions/RigVMDispatch_Switch.cpp
@@ -0,0 +1,148 @@
+// Copyright Epic Games, Inc. All Rights Reserved.
+
+#include "RigVMFunctions/RigVMDispatch_Switch.h"
+#include "RigVMStringUtils.h"
+#include "RigVMCore/RigVMRegistry.h"
+#include "RigVMCore/RigVMStruct.h"
+
+#include UE_INLINE_GENERATED_CPP_BY_NAME(RigVMDispatch_Switch)
+
+const FName FRigVMDispatch_SwitchInt32::IndexName = TEXT("Index");
+const FName FRigVMDispatch_SwitchInt32::CasesName = TEXT("Cases");
+
+TArray<FRigVMTemplateArgument> FRigVMDispatch_SwitchInt32::GetArguments() const
+{
+	return {
+		FRigVMTemplateArgument(IndexName, ERigVMPinDirection::Input, RigVMTypeUtils::TypeIndex::Int32),
+		FRigVMTemplateArgument(FRigVMStruct::ControlFlowBlockToRunName, ERigVMPinDirection::Hidden, RigVMTypeUtils::TypeIndex::FName)
+	};
+}
+
+TArray<FRigVMExecuteArgument> FRigVMDispatch_SwitchInt32::GetExecuteArguments_Impl(const FRigVMDispatchContext& InContext) const
+{
+	static const TArray<FRigVMExecuteArgument> Arguments =
+	{
+		{FRigVMStruct::ExecuteContextName, ERigVMPinDirection::Input},
+		{CasesName, ERigVMPinDirection::Output, RigVMTypeUtils::TypeIndex::ExecuteArray},
+		{FRigVMStruct::ControlFlowCompletedName, ERigVMPinDirection::Output}
+	};
+	return Arguments;
+}
+
+#if WITH_EDITOR
+
+FString FRigVMDispatch_SwitchInt32::GetArgumentMetaData(const FName& InArgumentName, const FName& InMetaDataKey) const
+{
+	if(InArgumentName == CasesName)
+	{
+		if(InMetaDataKey == FRigVMStruct::FixedSizeArrayMetaName)
+		{
+			return TrueString;
+		}
+	}
+	return FRigVMDispatch_CoreBase::GetArgumentMetaData(InArgumentName, InMetaDataKey);
+}
+
+FString FRigVMDispatch_SwitchInt32::GetArgumentDefaultValue(const FName& InArgumentName,
+	TRigVMTypeIndex InTypeIndex) const
+{
+	if(InArgumentName == CasesName)
+	{
+		static const FString TwoCases = TEXT("((),())");
+		return TwoCases;
+	}
+	return FRigVMDispatch_CoreBase::GetArgumentDefaultValue(InArgumentName, InTypeIndex);
+}
+
+FName FRigVMDispatch_SwitchInt32::GetDisplayNameForArgument(const FName& InArgumentName) const
+{
+	static const FString CasesPrefix = CasesName.ToString() + TEXT(".");
+	const FString ArgumentNameString = InArgumentName.ToString();
+	if(ArgumentNameString.StartsWith(CasesPrefix))
+	{
+		FString Left, Right;
+		verify(RigVMStringUtils::SplitPinPathAtEnd(ArgumentNameString, Left, Right));
+		if(Right.IsNumeric())
+		{
+			const int32 CaseIndex = FCString::Atoi(*Right);
+			return GetCaseDisplayName(CaseIndex);
+		}
+	}
+	return FRigVMDispatch_CoreBase::GetDisplayNameForArgument(InArgumentName);
+}
+
+#endif
+
+const TArray<FName>& FRigVMDispatch_SwitchInt32::GetControlFlowBlocks_Impl(const FRigVMDispatchContext& InContext) const
+{
+	static const TArray<FName> DefaultBlocks = {FRigVMStruct::ControlFlowCompletedName};
+	
+	if(!InContext.StringRepresentation.IsEmpty())
+	{
+		static const FString CasesPrefix = CasesName.ToString() + TEXT("=");
+		TArray<FString> NameValuePairs = RigVMStringUtils::SplitDefaultValue(InContext.StringRepresentation);
+		for(const FString& NameValuePair : NameValuePairs)
+		{
+			if(NameValuePair.StartsWith(CasesPrefix))
+			{
+				const FString Values = NameValuePair.Mid(CasesPrefix.Len());
+
+				static TMap<FString, TArray<FName>> BlockCache;
+				if(const TArray<FName>* ExistingBlocks = BlockCache.Find(Values))
+				{
+					return *ExistingBlocks;
+				}
+				
+				TArray<FName>& Blocks = BlockCache.Add(Values);
+				
+				const TArray<FString> CaseNames = RigVMStringUtils::SplitDefaultValue(Values);
+				for(int32 CaseIndex = 0; CaseIndex < CaseNames.Num(); CaseIndex++)
+				{
+					Blocks.Add(GetCaseName(CaseIndex));
+				}
+				
+				Blocks.Append(DefaultBlocks);
+				return Blocks;
+			}
+		}
+	}
+
+	return DefaultBlocks;
+}
+
+FRigVMFunctionPtr FRigVMDispatch_SwitchInt32::GetDispatchFunctionImpl(const FRigVMTemplateTypeMap& InTypes) const
+{
+	return &FRigVMDispatch_SwitchInt32::Execute;
+}
+
+void FRigVMDispatch_SwitchInt32::Execute(FRigVMExtendedExecuteContext& InContext, FRigVMMemoryHandleArray Handles)
+{
+#if WITH_EDITOR
+	check(Handles[0].IsInt32());
+	check(Handles[1].IsName());
+#endif
+
+	const int32 Index = *(const int32*)Handles[0].GetData();
+	FName& BlockToRun = *(FName*)Handles[1].GetData();
+
+	if(BlockToRun.IsNone())
+	{
+		BlockToRun = GetCaseName(Index);
+	}
+	else
+	{
+		BlockToRun = FRigVMStruct::ControlFlowCompletedName;
+	}
+}
+
+FName FRigVMDispatch_SwitchInt32::GetCaseName(int32 InIndex)
+{
+	return FRigVMBranchInfo::GetFixedArrayLabel(CasesName, *FString::FromInt(InIndex));
+}
+
+FName FRigVMDispatch_SwitchInt32::GetCaseDisplayName(int32 InIndex)
+{
+	static constexpr TCHAR Format[] = TEXT("Case %d");
+	return *FString::Printf(Format, InIndex);
+}
+
diff --git a/Engine/Source/Runtime/RigVM/Public/RigVMCore/RigVMDispatchFactory.h b/Engine/Source/Runtime/RigVM/Public/RigVMCore/RigVMDispatchFactory.h
--- a/Engine/Source/Runtime/RigVM/Public/RigVMCore/RigVMDispatchFactory.h
+++ b/Engine/Source/Runtime/RigVM/Public/RigVMCore/RigVMDispatchFactory.h
@@ -88,10 +88,10 @@ public:
 	virtual FName GetNextAggregateName(const FName& InLastAggregatePinName) const;
 
 	// Returns the display name text for an argument 
-	virtual FText GetDisplayNameForArgument(const FName& InArgumentName) const { return FText::FromName(InArgumentName); }
+	virtual FName GetDisplayNameForArgument(const FName& InArgumentName) const;
 
 	// Returns meta data on the property of the permutations 
-	virtual FString GetArgumentMetaData(const FName& InArgumentName, const FName& InMetaDataKey) const { return FString(); }
+	virtual FString GetArgumentMetaData(const FName& InArgumentName, const FName& InMetaDataKey) const;
 
 	// Returns true if the factory provides metadata for a given argument
 	bool HasArgumentMetaData(const FName& InArgumentName, const FName& InMetaDataKey) const
@@ -206,6 +206,7 @@ protected:
 	virtual const TArray<FName>& GetControlFlowBlocks_Impl(const FRigVMDispatchContext& InContext) const;
 
 	static const FString DispatchPrefix;
+	static const FString TrueString;
 
 	UScriptStruct* FactoryScriptStruct;
 	mutable const FRigVMTemplate* CachedTemplate;
diff --git a/Engine/Source/Runtime/RigVM/Public/RigVMFunctions/RigVMDispatch_Core.h b/Engine/Source/Runtime/RigVM/Public/RigVMFunctions/RigVMDispatch_Core.h
--- a/Engine/Source/Runtime/RigVM/Public/RigVMFunctions/RigVMDispatch_Core.h
+++ b/Engine/Source/Runtime/RigVM/Public/RigVMFunctions/RigVMDispatch_Core.h
@@ -6,20 +6,16 @@
 #include "RigVMDispatch_Core.generated.h"
 
 USTRUCT(meta=(Abstract, Category = "Core", NodeColor = "0.762745, 1,0, 0.329412"))
-struct FRigVMDispatch_CoreBase : public FRigVMDispatchFactory
+struct RIGVM_API FRigVMDispatch_CoreBase : public FRigVMDispatchFactory
 {
 	GENERATED_BODY()
-
-protected:
-
-	static const FString TrueString;
 };
 
 /*
  * Compares any two values and return true if they are identical
  */
 USTRUCT(meta=(DisplayName = "Equals", Keywords = "Same,=="))
-struct FRigVMDispatch_CoreEquals : public FRigVMDispatch_CoreBase
+struct RIGVM_API FRigVMDispatch_CoreEquals : public FRigVMDispatch_CoreBase
 {
 	GENERATED_BODY()
 
@@ -85,7 +81,7 @@ protected:
  * Compares any two values and return true if they are identical
  */
 USTRUCT(meta=(DisplayName = "Not Equals", Keywords = "Different,!=,Xor"))
-struct FRigVMDispatch_CoreNotEquals : public FRigVMDispatch_CoreEquals
+struct RIGVM_API FRigVMDispatch_CoreNotEquals : public FRigVMDispatch_CoreEquals
 {
 	GENERATED_BODY()
 
diff --git a/Engine/Source/Runtime/RigVM/Public/RigVMFunctions/RigVMDispatch_If.h b/Engine/Source/Runtime/RigVM/Public/RigVMFunctions/RigVMDispatch_If.h
--- a/Engine/Source/Runtime/RigVM/Public/RigVMFunctions/RigVMDispatch_If.h
+++ b/Engine/Source/Runtime/RigVM/Public/RigVMFunctions/RigVMDispatch_If.h
@@ -10,7 +10,7 @@
  * Chooses between two values based on a condition
  */
 USTRUCT(meta=(DisplayName = "If", Category = "Execution", Keywords = "Branch,Condition", NodeColor = "0,1,0,1"))
-struct FRigVMDispatch_If : public FRigVMDispatch_CoreBase
+struct RIGVM_API FRigVMDispatch_If : public FRigVMDispatch_CoreBase
 {
 	GENERATED_BODY()
 
diff --git a/Engine/Source/Runtime/RigVM/Public/RigVMFunctions/RigVMDispatch_Print.h b/Engine/Source/Runtime/RigVM/Public/RigVMFunctions/RigVMDispatch_Print.h
--- a/Engine/Source/Runtime/RigVM/Public/RigVMFunctions/RigVMDispatch_Print.h
+++ b/Engine/Source/Runtime/RigVM/Public/RigVMFunctions/RigVMDispatch_Print.h
@@ -9,7 +9,7 @@
  * Prints any value to the log
  */
 USTRUCT(meta=(DisplayName = "Print", NodeColor = "0.8, 0, 0.2, 1"))
-struct FRigVMDispatch_Print : public FRigVMDispatchFactory
+struct RIGVM_API FRigVMDispatch_Print : public FRigVMDispatchFactory
 {
 	GENERATED_BODY()
 
diff --git a/Engine/Source/Runtime/RigVM/Public/RigVMFunctions/RigVMDispatch_Select.h b/Engine/Source/Runtime/RigVM/Public/RigVMFunctions/RigVMDispatch_Select.h
--- a/Engine/Source/Runtime/RigVM/Public/RigVMFunctions/RigVMDispatch_Select.h
+++ b/Engine/Source/Runtime/RigVM/Public/RigVMFunctions/RigVMDispatch_Select.h
@@ -10,7 +10,7 @@
  * Pick from a list of values based on an integer index
  */
 USTRUCT(meta=(DisplayName = "Select", Category = "Execution", Keywords = "Switch,Case", NodeColor = "0,1,0,1"))
-struct FRigVMDispatch_SelectInt32 : public FRigVMDispatch_CoreBase
+struct RIGVM_API FRigVMDispatch_SelectInt32 : public FRigVMDispatch_CoreBase
 {
 	GENERATED_BODY()
 
diff --git a/Engine/Source/Runtime/RigVM/Public/RigVMFunctions/RigVMDispatch_Switch.h b/Engine/Source/Runtime/RigVM/Public/RigVMFunctions/RigVMDispatch_Switch.h
new file mode 100644
--- /dev/null
+++ b/Engine/Source/Runtime/RigVM/Public/RigVMFunctions/RigVMDispatch_Switch.h
@@ -0,0 +1,42 @@
+// Copyright Epic Games, Inc. All Rights Reserved.
+
+#pragma once
+
+#include "RigVMCore/RigVMDispatchFactory.h"
+#include "RigVMDispatch_Core.h"
+#include "RigVMDispatch_Switch.generated.h"
+
+/*
+ * Run a branch based on an integer index
+ */
+USTRUCT(meta=(DisplayName = "Switch", Category = "Execution", Keywords = "Case", NodeColor = "0,1,0,1"))
+struct RIGVM_API FRigVMDispatch_SwitchInt32 : public FRigVMDispatch_CoreBase
+{
+	GENERATED_BODY()
+
+public:
+
+	FRigVMDispatch_SwitchInt32()
+	{}
+
+	virtual TArray<FRigVMTemplateArgument> GetArguments() const override;
+	virtual TArray<FRigVMExecuteArgument> GetExecuteArguments_Impl(const FRigVMDispatchContext& InContext) const override;
+
+#if WITH_EDITOR
+	virtual FString GetArgumentMetaData(const FName& InArgumentName, const FName& InMetaDataKey) const override;
+	virtual FString GetArgumentDefaultValue(const FName& InArgumentName, TRigVMTypeIndex InTypeIndex) const override;
+	virtual FName GetDisplayNameForArgument(const FName& InArgumentName) const override;
+#endif
+	
+	virtual const TArray<FName>& GetControlFlowBlocks_Impl(const FRigVMDispatchContext& InContext) const override;
+
+protected:
+
+	virtual FRigVMFunctionPtr GetDispatchFunctionImpl(const FRigVMTemplateTypeMap& InTypes) const override;
+	static void Execute(FRigVMExtendedExecuteContext& InContext, FRigVMMemoryHandleArray Handles);
+	static FName GetCaseName(int32 InIndex);
+	static FName GetCaseDisplayName(int32 InIndex);
+	
+	static const FName IndexName;
+	static const FName CasesName;
+};
