RigVM: Fix slice index access in if and select dispatches

diff --git a/Engine/Source/Runtime/RigVM/Private/RigVMFunctions/RigVMDispatch_If.cpp b/Engine/Source/Runtime/RigVM/Private/RigVMFunctions/RigVMDispatch_If.cpp
--- a/Engine/Source/Runtime/RigVM/Private/RigVMFunctions/RigVMDispatch_If.cpp
+++ b/Engine/Source/Runtime/RigVM/Private/RigVMFunctions/RigVMDispatch_If.cpp
@@ -74,7 +74,7 @@ void FRigVMDispatch_If::Execute(FRigVMExtendedExecuteContext& InContext, FRigVMM
 	FRigVMMemoryHandle& InputHandle = Condition ? Handles[1] : Handles[2];
 	if(InputHandle.IsLazy())
 	{
-		InputHandle.ComputeLazyValueIfNecessary();
+		InputHandle.ComputeLazyValueIfNecessary(InContext.GetSlice().GetIndex());
 	}
 
 	const uint8* Input = InputHandle.GetData();
diff --git a/Engine/Source/Runtime/RigVM/Private/RigVMFunctions/RigVMDispatch_Select.cpp b/Engine/Source/Runtime/RigVM/Private/RigVMFunctions/RigVMDispatch_Select.cpp
--- a/Engine/Source/Runtime/RigVM/Private/RigVMFunctions/RigVMDispatch_Select.cpp
+++ b/Engine/Source/Runtime/RigVM/Private/RigVMFunctions/RigVMDispatch_Select.cpp
@@ -137,7 +137,7 @@ void FRigVMDispatch_SelectInt32::Execute(FRigVMExtendedExecuteContext& InContext
 	FRigVMMemoryHandle& InputHandle = Handles[1 + Index];
 	if(InputHandle.IsLazy())
 	{
-		InputHandle.ComputeLazyValueIfNecessary();
+		InputHandle.ComputeLazyValueIfNecessary(InContext.GetSlice().GetIndex());
 	}
 
 	const uint8* Input = InputHandle.GetData();
diff --git a/Engine/Source/Runtime/RigVM/Public/RigVMCore/RigVMMemoryStorage.h b/Engine/Source/Runtime/RigVM/Public/RigVMCore/RigVMMemoryStorage.h
--- a/Engine/Source/Runtime/RigVM/Public/RigVMCore/RigVMMemoryStorage.h
+++ b/Engine/Source/Runtime/RigVM/Public/RigVMCore/RigVMMemoryStorage.h
@@ -141,7 +141,7 @@ public:
 	{}
 
 	ERigVMExecuteResult Execute();
-	ERigVMExecuteResult ExecuteIfRequired(int32 InSliceIndex = INDEX_NONE);
+	ERigVMExecuteResult ExecuteIfRequired(int32 InSliceIndex);
 
 private:
 
