Control Rig: Fix crash when collapsing nodes with no arguments

diff --git a/Engine/Source/Runtime/RigVM/Private/RigVMCore/RigVMTemplate.cpp b/Engine/Source/Runtime/RigVM/Private/RigVMCore/RigVMTemplate.cpp
--- a/Engine/Source/Runtime/RigVM/Private/RigVMCore/RigVMTemplate.cpp
+++ b/Engine/Source/Runtime/RigVM/Private/RigVMCore/RigVMTemplate.cpp
@@ -1146,7 +1146,11 @@ bool FRigVMTemplate::ArgumentSupportsTypeIndex(const FName& InArgumentName, TRig
 
 const FRigVMFunction* FRigVMTemplate::GetPrimaryPermutation() const
 {
-	return GetPermutation(0);
+	if (NumPermutations() > 0)
+	{
+		return GetPermutation(0);
+	}
+	return nullptr;
 }
 
 const FRigVMFunction* FRigVMTemplate::GetPermutation(int32 InIndex) const
