Control Rig: Fix storage of bool array properties

diff --git a/Engine/Source/Runtime/RigVM/Private/RigVMCore/RigVMMemoryStorage.cpp b/Engine/Source/Runtime/RigVM/Private/RigVMCore/RigVMMemoryStorage.cpp
--- a/Engine/Source/Runtime/RigVM/Private/RigVMCore/RigVMMemoryStorage.cpp
+++ b/Engine/Source/Runtime/RigVM/Private/RigVMCore/RigVMMemoryStorage.cpp
@@ -831,7 +831,9 @@ FProperty* URigVMMemoryStorageGeneratorClass::AddProperty(URigVMMemoryStorageGen
 			const FString BaseCPPType = InProperty.GetTailCPPType();
 			if(BaseCPPType.Equals(BoolString, ESearchCase::IgnoreCase))
 			{
-				(*ValuePropertyPtr) = new FBoolProperty(PropertyOwner, InProperty.Name, RF_Public);;
+				FBoolProperty* BoolProperty = new FBoolProperty(PropertyOwner, InProperty.Name, RF_Public); 
+				BoolProperty->SetBoolSize(sizeof(bool), true);
+				(*ValuePropertyPtr) = BoolProperty;
 			}
 			else if(BaseCPPType.Equals(Int32String, ESearchCase::IgnoreCase) ||
 				BaseCPPType.Equals(IntString, ESearchCase::IgnoreCase))
