RigVM: Guard if and select nodes against non matching types

diff --git a/Engine/Source/Runtime/RigVM/Private/RigVMCore/RigVMDispatchFactory.cpp b/Engine/Source/Runtime/RigVM/Private/RigVMCore/RigVMDispatchFactory.cpp
--- a/Engine/Source/Runtime/RigVM/Private/RigVMCore/RigVMDispatchFactory.cpp
+++ b/Engine/Source/Runtime/RigVM/Private/RigVMCore/RigVMDispatchFactory.cpp
@@ -3,7 +3,7 @@
 #include "RigVMCore/RigVMDispatchFactory.h"
 #include "RigVMCore/RigVMRegistry.h"
 #include "RigVMCore/RigVMStruct.h"
-#include "RigVMStringUtils.h"
+#include "RigVMCore/RigVMMemoryStorage.h"
 
 #include UE_INLINE_GENERATED_CPP_BY_NAME(RigVMDispatchFactory)
 
@@ -225,6 +225,12 @@ FString FRigVMDispatchFactory::GetPermutationNameImpl(const FRigVMTemplateTypeMa
 	return FString::Printf(Format, *GetFactoryName().ToString(), *TypePairStrings);
 }
 
+bool FRigVMDispatchFactory::CopyProperty(const FProperty* InTargetProperty, uint8* InTargetPtr,
+	const FProperty* InSourceProperty, const uint8* InSourcePtr)
+{
+	return URigVMMemoryStorage::CopyProperty(InTargetProperty, InTargetPtr, InSourceProperty, InSourcePtr);
+}
+
 const FRigVMTemplate* FRigVMDispatchFactory::GetTemplate() const
 {
 	static bool bIsDispatchingTemplate = false;
diff --git a/Engine/Source/Runtime/RigVM/Private/RigVMFunctions/RigVMDispatch_If.cpp b/Engine/Source/Runtime/RigVM/Private/RigVMFunctions/RigVMDispatch_If.cpp
--- a/Engine/Source/Runtime/RigVM/Private/RigVMFunctions/RigVMDispatch_If.cpp
+++ b/Engine/Source/Runtime/RigVM/Private/RigVMFunctions/RigVMDispatch_If.cpp
@@ -63,10 +63,6 @@ void FRigVMDispatch_If::Execute(FRigVMExtendedExecuteContext& InContext, FRigVMM
 
 #if WITH_EDITOR
 	check(Handles[0].IsBool());
-	const FProperty* PropertyB = Handles[2].GetResolvedProperty(); 
-	const FProperty* PropertyC = Handles[3].GetResolvedProperty(); 
-	check(Property->SameType(PropertyB));
-	check(Property->SameType(PropertyC));
 #endif
 
 	const bool& Condition = *(const bool*)Handles[0].GetData();
@@ -79,6 +75,9 @@ void FRigVMDispatch_If::Execute(FRigVMExtendedExecuteContext& InContext, FRigVMM
 
 	const uint8* Input = InputHandle.GetData();
 	uint8* Result = Handles[3].GetData();
-	Property->CopyCompleteValue(Result, Input);
+
+	(void)CopyProperty(
+		Handles[3].GetProperty(), Result,
+		InputHandle.GetProperty(), Input);
 }
 
diff --git a/Engine/Source/Runtime/RigVM/Private/RigVMFunctions/RigVMDispatch_Select.cpp b/Engine/Source/Runtime/RigVM/Private/RigVMFunctions/RigVMDispatch_Select.cpp
--- a/Engine/Source/Runtime/RigVM/Private/RigVMFunctions/RigVMDispatch_Select.cpp
+++ b/Engine/Source/Runtime/RigVM/Private/RigVMFunctions/RigVMDispatch_Select.cpp
@@ -119,11 +119,6 @@ void FRigVMDispatch_SelectInt32::Execute(FRigVMExtendedExecuteContext& InContext
 
 #if WITH_EDITOR
 	check(Handles[0].IsInt32());
-	for(int32 Case = 0; Case < NumCases; Case++)
-	{
-		const FProperty* PropertyB = Handles[1 + Case].GetResolvedProperty(); 
-		check(Property->SameType(PropertyB));
-	}
 #endif
 
 	int32 Index = *(const int32*)Handles[0].GetData();
@@ -142,6 +137,10 @@ void FRigVMDispatch_SelectInt32::Execute(FRigVMExtendedExecuteContext& InContext
 
 	const uint8* Input = InputHandle.GetData();
 	uint8* Result = Handles.Last().GetData();
-	Property->CopyCompleteValue(Result, Input);
+
+	// copy property performs compatibility check and will report accordingly
+	(void)CopyProperty(
+		Handles.Last().GetProperty(), Result,
+		InputHandle.GetProperty(), Input);
 }
 
diff --git a/Engine/Source/Runtime/RigVM/Public/RigVMCore/RigVMDispatchFactory.h b/Engine/Source/Runtime/RigVM/Public/RigVMCore/RigVMDispatchFactory.h
--- a/Engine/Source/Runtime/RigVM/Public/RigVMCore/RigVMDispatchFactory.h
+++ b/Engine/Source/Runtime/RigVM/Public/RigVMCore/RigVMDispatchFactory.h
@@ -164,6 +164,12 @@ protected:
 
 	virtual TArray<FRigVMExecuteArgument> GetExecuteArguments_Impl(const FRigVMDispatchContext& InContext) const { return TArray<FRigVMExecuteArgument>(); }
 
+	static bool CopyProperty(
+		const FProperty* InTargetProperty,
+		uint8* InTargetPtr,
+		const FProperty* InSourceProperty,
+		const uint8* InSourcePtr);
+
 	template <
 	typename T,
 	typename TEnableIf<TRigVMIsBaseStructure<T>::Value, T>::Type* = nullptr
