Control Rig: Move draw interface and log to RigVM

diff --git a/Engine/Plugins/Animation/ControlRig/Config/BaseControlRig.ini b/Engine/Plugins/Animation/ControlRig/Config/BaseControlRig.ini
--- a/Engine/Plugins/Animation/ControlRig/Config/BaseControlRig.ini
+++ b/Engine/Plugins/Animation/ControlRig/Config/BaseControlRig.ini
@@ -82,3 +82,6 @@
 +StructRedirects=(OldName="/Script/ControlRig.RigUnit_StringSplit",NewName="/Script/RigVM.RigVMFunction_StringSplit")
 +StructRedirects=(OldName="/Script/ControlRig.RigUnit_StringJoin",NewName="/Script/RigVM.RigVMFunction_StringJoin")
 +StructRedirects=(OldName="/Script/ControlRig.RigUnit_StringPadInteger",NewName="/Script/RigVM.RigVMFunction_StringPadInteger")
++StructRedirects=(OldName="/Script/ControlRig.ControlRigDrawInstruction",NewName="/Script/RigVM.RigVMDrawInstruction")
++StructRedirects=(OldName="/Script/ControlRig.ControlRigDrawContainer",NewName="/Script/RigVM.RigVMDrawContainer")
++StructRedirects=(OldName="/Script/ControlRig.ControlRigDrawInterface",NewName="/Script/RigVM.RigVMDrawInterface")
diff --git a/Engine/Plugins/Animation/ControlRig/Source/ControlRig/Private/AnimNode_ControlRigBase.cpp b/Engine/Plugins/Animation/ControlRig/Source/ControlRig/Private/AnimNode_ControlRigBase.cpp
--- a/Engine/Plugins/Animation/ControlRig/Source/ControlRig/Private/AnimNode_ControlRigBase.cpp
+++ b/Engine/Plugins/Animation/ControlRig/Source/ControlRig/Private/AnimNode_ControlRigBase.cpp
@@ -830,7 +830,7 @@ void FAnimNode_ControlRigBase::QueueControlRigDrawInstructions(UControlRig* Cont
 
 	if (ControlRig && Proxy)
 	{
-		for (const FControlRigDrawInstruction& Instruction : ControlRig->GetDrawInterface())
+		for (const FRigVMDrawInstruction& Instruction : ControlRig->GetDrawInterface())
 		{
 			if (!Instruction.IsValid())
 			{
@@ -840,7 +840,7 @@ void FAnimNode_ControlRigBase::QueueControlRigDrawInstructions(UControlRig* Cont
 			FTransform InstructionTransform = Instruction.Transform * Proxy->GetComponentTransform();
 			switch (Instruction.PrimitiveType)
 			{
-				case EControlRigDrawSettings::Points:
+				case ERigVMDrawSettings::Points:
 				{
 					for (const FVector& Point : Instruction.Positions)
 					{
@@ -848,7 +848,7 @@ void FAnimNode_ControlRigBase::QueueControlRigDrawInstructions(UControlRig* Cont
 					}
 					break;
 				}
-				case EControlRigDrawSettings::Lines:
+				case ERigVMDrawSettings::Lines:
 				{
 					const TArray<FVector>& Points = Instruction.Positions;
 
@@ -858,7 +858,7 @@ void FAnimNode_ControlRigBase::QueueControlRigDrawInstructions(UControlRig* Cont
 					}
 					break;
 				}
-				case EControlRigDrawSettings::LineStrip:
+				case ERigVMDrawSettings::LineStrip:
 				{
 					const TArray<FVector>& Points = Instruction.Positions;
 
@@ -869,7 +869,7 @@ void FAnimNode_ControlRigBase::QueueControlRigDrawInstructions(UControlRig* Cont
 					break;
 				}
 
-				case EControlRigDrawSettings::DynamicMesh:
+				case ERigVMDrawSettings::DynamicMesh:
 				{
 					// TODO: Add support for this if anyone is actually using it. Currently it is only defined and referenced in an unused API, DrawCone in Control Rig.
 					break;
diff --git a/Engine/Plugins/Animation/ControlRig/Source/ControlRig/Private/ControlRig.cpp b/Engine/Plugins/Animation/ControlRig/Source/ControlRig/Private/ControlRig.cpp
--- a/Engine/Plugins/Animation/ControlRig/Source/ControlRig/Private/ControlRig.cpp
+++ b/Engine/Plugins/Animation/ControlRig/Source/ControlRig/Private/ControlRig.cpp
@@ -797,7 +797,7 @@ bool UControlRig::Execute(const FName& InEventName)
 	{
 		DrawInterface.Reset();
 	}
-	Context.DrawInterface = &DrawInterface;
+	PublicContext.SetDrawInterface(&DrawInterface);
 
 	// setup the animation attribute container
 	Context.AnimAttributeContainer = ExternalAnimAttributeContainer;
@@ -809,7 +809,7 @@ bool UControlRig::Execute(const FName& InEventName)
 	{
 		DrawContainer = CDO->DrawContainer;
 	}
-	Context.DrawContainer = &DrawContainer;
+	PublicContext.SetDrawContainer(&DrawContainer);
 
 	// setup the data source registry
 	Context.DataSourceRegistry = GetDataSourceRegistry();
@@ -911,7 +911,7 @@ bool UControlRig::Execute(const FName& InEventName)
 
 #if WITH_EDITOR
 	// setup the log and VM settings
-	Context.Log = ControlRigLog;
+	PublicContext.SetLog(ControlRigLog);
 	if (ControlRigLog != nullptr)
 	{
 		ControlRigLog->Reset();
@@ -1127,7 +1127,7 @@ bool UControlRig::Execute(const FName& InEventName)
 		}
 		else if(bIsEventLastInQueue)
 		{
-			for (const FControlRigLog::FLogEntry& Entry : ControlRigLog->Entries)
+			for (const FRigVMLog::FLogEntry& Entry : ControlRigLog->Entries)
 			{
 				if (Entry.FunctionName == NAME_None || Entry.InstructionIndex == INDEX_NONE || Entry.Message.IsEmpty())
 				{
@@ -1169,9 +1169,9 @@ bool UControlRig::Execute(const FName& InEventName)
 		HandleHierarchyEvent(GetHierarchy(), EventContext);
 	}
 
-	if (Context.DrawInterface && Context.DrawContainer && bIsEventLastInQueue) 
+	if (PublicContext.GetDrawInterface() && PublicContext.GetDrawContainer() && bIsEventLastInQueue) 
 	{
-		Context.DrawInterface->Instructions.Append(Context.DrawContainer->Instructions);
+		PublicContext.GetDrawInterface()->Instructions.Append(PublicContext.GetDrawContainer()->Instructions);
 
 		FRigHierarchyValidityBracket ValidityBracket(GetHierarchy());
 		
@@ -1185,7 +1185,7 @@ bool UControlRig::Execute(const FName& InEventName)
 				Settings.LimitEnabled.Contains(FRigControlLimitEnabled(true, true)))
 			{
 				FTransform Transform = GetHierarchy()->GetGlobalControlOffsetTransformByIndex(ControlElement->GetIndex());
-				FControlRigDrawInstruction Instruction(EControlRigDrawSettings::Lines, Settings.ShapeColor, 0.f, Transform);
+				FRigVMDrawInstruction Instruction(ERigVMDrawSettings::Lines, Settings.ShapeColor, 0.f, Transform);
 
 				switch (Settings.ControlType)
 				{
@@ -1272,7 +1272,7 @@ bool UControlRig::Execute(const FName& InEventName)
 							break;
 						}
 
-						Instruction.PrimitiveType = EControlRigDrawSettings::LineStrip;
+						Instruction.PrimitiveType = ERigVMDrawSettings::LineStrip;
 						FVector3f MinPos = Settings.MinimumValue.Get<FVector3f>();
 						FVector3f MaxPos = Settings.MaximumValue.Get<FVector3f>();
 
@@ -1390,7 +1390,7 @@ bool UControlRig::Execute(const FName& InEventName)
 
 				if (Instruction.Positions.Num() > 0)
 				{
-					DrawInterface.Instructions.Add(Instruction);
+					DrawInterface.DrawInstruction(Instruction);
 				}
 			}
 
diff --git a/Engine/Plugins/Animation/ControlRig/Source/ControlRig/Private/ControlRigComponent.cpp b/Engine/Plugins/Animation/ControlRig/Source/ControlRig/Private/ControlRigComponent.cpp
--- a/Engine/Plugins/Animation/ControlRig/Source/ControlRig/Private/ControlRigComponent.cpp
+++ b/Engine/Plugins/Animation/ControlRig/Source/ControlRig/Private/ControlRigComponent.cpp
@@ -10,6 +10,7 @@
 #include "AnimCustomInstanceHelper.h"
 #include "ControlRigObjectBinding.h"
 #include "SceneManagement.h"
+#include "Math/ControlRigMathLibrary.h"
 
 #include UE_INLINE_GENERATED_CPP_BY_NAME(ControlRigComponent)
 
@@ -223,11 +224,11 @@ FBoxSphereBounds UControlRigComponent::CalcBounds(const FTransform& LocalToWorld
 		// Get bounding box for the debug drawings if they are drawn 
 		if (bShowDebugDrawing)
 		{ 
-			const FControlRigDrawInterface& DrawInterface = ControlRig->GetDrawInterface();
+			const FRigVMDrawInterface& DrawInterface = ControlRig->GetDrawInterface();
 
 			for (int32 InstructionIndex = 0; InstructionIndex < DrawInterface.Num(); InstructionIndex++)
 			{
-				const FControlRigDrawInstruction& Instruction = DrawInterface[InstructionIndex];
+				const FRigVMDrawInstruction& Instruction = DrawInterface[InstructionIndex];
 
 				FTransform Transform = Instruction.Transform * GetComponentToWorld();
 				for (const FVector& Position : Instruction.Positions)
@@ -2007,11 +2008,11 @@ void FControlRigSceneProxy::GetDynamicMeshElements(const TArray<const FSceneView
 
 			if (ControlRigComponent->bShowDebugDrawing)
 			{ 
-				const FControlRigDrawInterface& DrawInterface = ControlRigComponent->ControlRig->GetDrawInterface();
+				const FRigVMDrawInterface& DrawInterface = ControlRigComponent->ControlRig->GetDrawInterface();
 
 				for (int32 InstructionIndex = 0; InstructionIndex < DrawInterface.Num(); InstructionIndex++)
 				{
-					const FControlRigDrawInstruction& Instruction = DrawInterface[InstructionIndex];
+					const FRigVMDrawInstruction& Instruction = DrawInterface[InstructionIndex];
 					if (Instruction.Positions.Num() == 0)
 					{
 						continue;
@@ -2020,7 +2021,7 @@ void FControlRigSceneProxy::GetDynamicMeshElements(const TArray<const FSceneView
 					FTransform InstructionTransform = Instruction.Transform * ControlRigComponent->GetComponentToWorld();
 					switch (Instruction.PrimitiveType)
 					{
-						case EControlRigDrawSettings::Points:
+						case ERigVMDrawSettings::Points:
 						{
 							for (const FVector& Point : Instruction.Positions)
 							{
@@ -2028,7 +2029,7 @@ void FControlRigSceneProxy::GetDynamicMeshElements(const TArray<const FSceneView
 							}
 							break;
 						}
-						case EControlRigDrawSettings::Lines:
+						case ERigVMDrawSettings::Lines:
 						{
 							const TArray<FVector>& Points = Instruction.Positions;
 							PDI->AddReserveLines(SDPG_Foreground, Points.Num() / 2, false, Instruction.Thickness > SMALL_NUMBER);
@@ -2038,7 +2039,7 @@ void FControlRigSceneProxy::GetDynamicMeshElements(const TArray<const FSceneView
 							}
 							break;
 						}
-						case EControlRigDrawSettings::LineStrip:
+						case ERigVMDrawSettings::LineStrip:
 						{
 							const TArray<FVector>& Points = Instruction.Positions;
 							PDI->AddReserveLines(SDPG_Foreground, Points.Num() - 1, false, Instruction.Thickness > SMALL_NUMBER);
diff --git a/Engine/Plugins/Animation/ControlRig/Source/ControlRig/Private/ControlRigLog.cpp b/Engine/Plugins/Animation/ControlRig/Source/ControlRig/Private/ControlRigLog.cpp
deleted file mode 100644
diff --git a/Engine/Plugins/Animation/ControlRig/Source/ControlRig/Private/ControlRigTestData.cpp b/Engine/Plugins/Animation/ControlRig/Source/ControlRig/Private/ControlRigTestData.cpp
--- a/Engine/Plugins/Animation/ControlRig/Source/ControlRig/Private/ControlRigTestData.cpp
+++ b/Engine/Plugins/Animation/ControlRig/Source/ControlRig/Private/ControlRigTestData.cpp
@@ -338,7 +338,7 @@ bool UControlRigTestData::SetupReplay(UControlRig* InControlRig, bool bGroundTru
 			const FRigPose& ExpectedPose = Frame.Pose;
 
 			// draw differences of the pose result of the rig onto the screen
-			FControlRigDrawInterface& DrawInterface = InControlRig->GetDrawInterface();
+			FRigVMDrawInterface& DrawInterface = InControlRig->GetDrawInterface();
 			for(const FRigPoseElement& ExpectedPoseElement : ExpectedPose)
 			{
 				const int32 CurrentPoseIndex = CurrentPose.GetIndex(ExpectedPoseElement.Index.GetKey());
diff --git a/Engine/Plugins/Animation/ControlRig/Source/ControlRig/Private/Drawing/ControlRigDrawContainer.cpp b/Engine/Plugins/Animation/ControlRig/Source/ControlRig/Private/Drawing/ControlRigDrawContainer.cpp
deleted file mode 100644
diff --git a/Engine/Plugins/Animation/ControlRig/Source/ControlRig/Private/Drawing/ControlRigDrawInstruction.cpp b/Engine/Plugins/Animation/ControlRig/Source/ControlRig/Private/Drawing/ControlRigDrawInstruction.cpp
deleted file mode 100644
diff --git a/Engine/Plugins/Animation/ControlRig/Source/ControlRig/Private/Drawing/ControlRigDrawInterface.cpp b/Engine/Plugins/Animation/ControlRig/Source/ControlRig/Private/Drawing/ControlRigDrawInterface.cpp
deleted file mode 100644
diff --git a/Engine/Plugins/Animation/ControlRig/Source/ControlRig/Private/Rigs/RigHierarchy.cpp b/Engine/Plugins/Animation/ControlRig/Source/ControlRig/Private/Rigs/RigHierarchy.cpp
--- a/Engine/Plugins/Animation/ControlRig/Source/ControlRig/Private/Rigs/RigHierarchy.cpp
+++ b/Engine/Plugins/Animation/ControlRig/Source/ControlRig/Private/Rigs/RigHierarchy.cpp
@@ -2800,15 +2800,15 @@ void URigHierarchy::SetTransform(FRigTransformElement* InTransformElement, const
 
 						if(bChildFound)
 						{
-							const FRigUnitContext& UnitContext = ExecuteContext->GetPublicData<FControlRigExecuteContext>().UnitContext;
-							if(UnitContext.Log)
+							const FControlRigExecuteContext& CRContext = ExecuteContext->GetPublicData<FControlRigExecuteContext>();
+							if(CRContext.GetLog())
 							{
 								static constexpr TCHAR MessageFormat[] = TEXT("Setting transform of parent (%s) after setting child (%s).\nThis may lead to unexpected results.");
 								const FString& Message = FString::Printf(
 									MessageFormat,
 									*InTransformElement->GetName().ToString(),
 									*Child->GetName().ToString());
-								UnitContext.Log->Report(
+								CRContext.GetLog()->Report(
 									EMessageSeverity::Info,
 									ExecuteContext->GetPublicData<>().GetFunctionName(),
 									ExecuteContext->GetPublicData<>().GetInstructionIndex(),
diff --git a/Engine/Plugins/Animation/ControlRig/Source/ControlRig/Private/Units/Debug/RigUnit_DebugBezier.cpp b/Engine/Plugins/Animation/ControlRig/Source/ControlRig/Private/Units/Debug/RigUnit_DebugBezier.cpp
--- a/Engine/Plugins/Animation/ControlRig/Source/ControlRig/Private/Units/Debug/RigUnit_DebugBezier.cpp
+++ b/Engine/Plugins/Animation/ControlRig/Source/ControlRig/Private/Units/Debug/RigUnit_DebugBezier.cpp
@@ -30,18 +30,13 @@ FRigUnit_DebugBezierItemSpace_Execute()
 {
     DECLARE_SCOPE_HIERARCHICAL_COUNTER_RIGUNIT()
 
-	if (ExecuteContext.UnitContext.DrawInterface == nullptr || !bEnabled)
-	{
-		return;
-	}
-
 	FTransform Transform = WorldOffset;
 	if (Space.IsValid())
 	{
 		Transform = Transform * ExecuteContext.Hierarchy->GetGlobalTransform(Space);
 	}
 
-	ExecuteContext.UnitContext.DrawInterface->DrawBezier(Transform, Bezier, MinimumU, MaximumU, Color, Thickness, Detail);
+	DrawBezier(ExecuteContext, Transform, Bezier, MinimumU, MaximumU, Color, Thickness, Detail);
 }
 
 FRigVMStructUpgradeInfo FRigUnit_DebugBezierItemSpace::GetUpgradeInfo() const
@@ -49,3 +44,32 @@ FRigVMStructUpgradeInfo FRigUnit_DebugBezierItemSpace::GetUpgradeInfo() const
 	// this node is no longer supported
 	return FRigVMStructUpgradeInfo();
 }
+
+void FRigUnit_DebugBezierItemSpace::DrawBezier(const FControlRigExecuteContext& InContext, const FTransform& WorldOffset, const FCRFourPointBezier& InBezier, float MinimumU, float MaximumU, const FLinearColor& Color, float Thickness, int32 Detail)
+{
+	FRigVMDrawInterface* DrawInterface = InContext.GetDrawInterface();
+	if(DrawInterface == nullptr)
+	{
+		return;
+	}
+	
+	if (!DrawInterface->IsEnabled())
+	{
+		return;
+	}
+
+	int32 Count = FMath::Clamp<int32>(Detail, 4, 64);
+	FRigVMDrawInstruction Instruction(ERigVMDrawSettings::LineStrip, Color, Thickness, WorldOffset);
+	Instruction.Positions.SetNumUninitialized(Count);
+
+	float T = MinimumU;
+	float Step = (MaximumU - MinimumU) / float(Detail-1);
+	for(int32 Index=0;Index<Count;Index++)
+	{
+		FVector Tangent;
+		FControlRigMathLibrary::FourPointBezier(InBezier, T, Instruction.Positions[Index], Tangent);
+		T += Step;
+	}
+
+	DrawInterface->DrawInstruction(Instruction);
+}
diff --git a/Engine/Plugins/Animation/ControlRig/Source/ControlRig/Private/Units/Debug/RigUnit_DebugBezier.h b/Engine/Plugins/Animation/ControlRig/Source/ControlRig/Private/Units/Debug/RigUnit_DebugBezier.h
--- a/Engine/Plugins/Animation/ControlRig/Source/ControlRig/Private/Units/Debug/RigUnit_DebugBezier.h
+++ b/Engine/Plugins/Animation/ControlRig/Source/ControlRig/Private/Units/Debug/RigUnit_DebugBezier.h
@@ -107,4 +107,6 @@ struct CONTROLRIG_API FRigUnit_DebugBezierItemSpace : public FRigUnit_DebugBaseM
 
 	RIGVM_METHOD()
 	virtual FRigVMStructUpgradeInfo GetUpgradeInfo() const override;
+
+	static void DrawBezier(const FControlRigExecuteContext& InContext, const FTransform& WorldOffset, const FCRFourPointBezier& InBezier, float MinimumU, float MaximumU, const FLinearColor& Color, float Thickness, int32 Detail);
 };
diff --git a/Engine/Plugins/Animation/ControlRig/Source/ControlRig/Private/Units/Debug/RigUnit_DebugHierarchy.cpp b/Engine/Plugins/Animation/ControlRig/Source/ControlRig/Private/Units/Debug/RigUnit_DebugHierarchy.cpp
--- a/Engine/Plugins/Animation/ControlRig/Source/ControlRig/Private/Units/Debug/RigUnit_DebugHierarchy.cpp
+++ b/Engine/Plugins/Animation/ControlRig/Source/ControlRig/Private/Units/Debug/RigUnit_DebugHierarchy.cpp
@@ -9,15 +9,10 @@ FRigUnit_DebugHierarchy_Execute()
 {
     DECLARE_SCOPE_HIERARCHICAL_COUNTER_RIGUNIT()
 
-	if (ExecuteContext.UnitContext.DrawInterface == nullptr || !bEnabled)
-	{
-		return;
-	}
-
 	URigHierarchy* Hierarchy = ExecuteContext.Hierarchy;
 	if (Hierarchy)
 	{
-		ExecuteContext.UnitContext.DrawInterface->DrawHierarchy(WorldOffset, Hierarchy, EControlRigDrawHierarchyMode::Axes, Scale, Color, Thickness);
+		DrawHierarchy(ExecuteContext, WorldOffset, Hierarchy, EControlRigDrawHierarchyMode::Axes, Scale, Color, Thickness, nullptr);
 	}
 }
 
@@ -25,15 +20,114 @@ FRigUnit_DebugPose_Execute()
 {
 	DECLARE_SCOPE_HIERARCHICAL_COUNTER_RIGUNIT()
 
-	if (ExecuteContext.UnitContext.DrawInterface == nullptr || !bEnabled)
+	URigHierarchy* Hierarchy = ExecuteContext.Hierarchy;
+	if (Hierarchy)
+	{
+		FRigUnit_DebugHierarchy::DrawHierarchy(ExecuteContext, WorldOffset, Hierarchy, EControlRigDrawHierarchyMode::Axes, Scale, Color, Thickness, &Pose);
+	}
+}
+
+void FRigUnit_DebugHierarchy::DrawHierarchy(const FControlRigExecuteContext& InContext, const FTransform& WorldOffset, URigHierarchy* Hierarchy, EControlRigDrawHierarchyMode::Type Mode, float Scale, const FLinearColor& Color, float Thickness, const FRigPose* InPose)
+{
+	FRigVMDrawInterface* DrawInterface = InContext.GetDrawInterface();
+	if(DrawInterface == nullptr)
+	{
+		return;
+	}
+	
+	if (!DrawInterface->IsEnabled())
 	{
 		return;
 	}
 
-	URigHierarchy* Hierarchy = ExecuteContext.Hierarchy;
-	if (Hierarchy)
+	switch (Mode)
 	{
-		ExecuteContext.UnitContext.DrawInterface->DrawHierarchy(WorldOffset, Hierarchy, EControlRigDrawHierarchyMode::Axes, Scale, Color, Thickness, &Pose);
+		case EControlRigDrawHierarchyMode::Axes:
+		{
+			FRigVMDrawInstruction InstructionX(ERigVMDrawSettings::Lines, FLinearColor::Red, Thickness, WorldOffset);
+			FRigVMDrawInstruction InstructionY(ERigVMDrawSettings::Lines, FLinearColor::Green, Thickness, WorldOffset);
+			FRigVMDrawInstruction InstructionZ(ERigVMDrawSettings::Lines, FLinearColor::Blue, Thickness, WorldOffset);
+			FRigVMDrawInstruction InstructionParent(ERigVMDrawSettings::Lines, Color, Thickness, WorldOffset);
+			InstructionX.Positions.Reserve(Hierarchy->Num() * 2);
+			InstructionY.Positions.Reserve(Hierarchy->Num() * 2);
+			InstructionZ.Positions.Reserve(Hierarchy->Num() * 2);
+			InstructionParent.Positions.Reserve(Hierarchy->Num() * 6);
+
+			Hierarchy->ForEach<FRigTransformElement>([&](FRigTransformElement* Child) -> bool
+			{
+				auto GetTransformLambda = [InPose, Hierarchy](FRigTransformElement* InElement, FTransform& Transform) -> bool
+				{
+					bool bValid = true;
+					if(InPose)
+					{
+						const int32 ElementIndex = InPose->GetIndex(InElement->GetKey());
+						if(ElementIndex != INDEX_NONE)
+						{
+							Transform = InPose->operator[](ElementIndex).GlobalTransform;
+							return true;
+						}
+						else
+						{
+							bValid = false;
+						}
+					}
+
+					Transform = Hierarchy->GetTransform(InElement, ERigTransformType::CurrentGlobal);
+					return bValid;
+				};
+				
+				FTransform Transform = FTransform::Identity;
+				if(!GetTransformLambda(Child, Transform))
+				{
+					return true;
+				}
+
+				const FVector P0 = Transform.GetLocation();
+				const FVector PX = Transform.TransformPosition(FVector(Scale, 0.f, 0.f));
+				const FVector PY = Transform.TransformPosition(FVector(0.f, Scale, 0.f));
+				const FVector PZ = Transform.TransformPosition(FVector(0.f, 0.f, Scale));
+				InstructionX.Positions.Add(P0);
+				InstructionX.Positions.Add(PX);
+				InstructionY.Positions.Add(P0);
+				InstructionY.Positions.Add(PY);
+				InstructionZ.Positions.Add(P0);
+				InstructionZ.Positions.Add(PZ);
+
+				FRigBaseElementParentArray Parents = Hierarchy->GetParents(Child);
+				TArray<FRigElementWeight> Weights = Hierarchy->GetParentWeightArray(Child);
+				
+				for (int32 ParentIndex = 0; ParentIndex < Parents.Num(); ParentIndex++)
+				{
+					if(FRigTransformElement* ParentTransformElement = Cast<FRigTransformElement>(Parents[ParentIndex]))
+					{
+						if(Weights.IsValidIndex(ParentIndex))
+						{
+							if(Weights[ParentIndex].IsAlmostZero())
+							{
+								continue;
+							}
+						}
+						
+						FTransform ParentTransform = FTransform::Identity;
+						GetTransformLambda(ParentTransformElement, ParentTransform);
+
+						const FVector P1 = ParentTransform.GetLocation();
+						InstructionParent.Positions.Add(P0);
+						InstructionParent.Positions.Add(P1);
+					}
+				}
+				return true;
+			});
+
+			DrawInterface->DrawInstruction(InstructionX);
+			DrawInterface->DrawInstruction(InstructionY);
+			DrawInterface->DrawInstruction(InstructionZ);
+			DrawInterface->DrawInstruction(InstructionParent);
+			break;
+		}
+		default:
+		{
+			break;
+		}
 	}
 }
-
diff --git a/Engine/Plugins/Animation/ControlRig/Source/ControlRig/Private/Units/Debug/RigUnit_DebugHierarchy.h b/Engine/Plugins/Animation/ControlRig/Source/ControlRig/Private/Units/Debug/RigUnit_DebugHierarchy.h
--- a/Engine/Plugins/Animation/ControlRig/Source/ControlRig/Private/Units/Debug/RigUnit_DebugHierarchy.h
+++ b/Engine/Plugins/Animation/ControlRig/Source/ControlRig/Private/Units/Debug/RigUnit_DebugHierarchy.h
@@ -3,9 +3,21 @@
 #pragma once
 
 #include "RigUnit_DebugBase.h"
-#include "Drawing/ControlRigDrawInterface.h"
 #include "RigUnit_DebugHierarchy.generated.h"
 
+UENUM()
+namespace EControlRigDrawHierarchyMode
+{
+	enum Type
+	{
+		/** Draw as axes */
+		Axes,
+
+		/** MAX - invalid */
+		Max UMETA(Hidden),
+	};
+}
+
 /**
  * Draws vectors on each bone in the viewport across the entire hierarchy
  */
@@ -40,6 +52,8 @@ struct CONTROLRIG_API FRigUnit_DebugHierarchy : public FRigUnit_DebugBaseMutable
 
 	UPROPERTY(meta = (Input))
 	bool bEnabled;
+
+	static void DrawHierarchy(const FControlRigExecuteContext& InContext, const FTransform& WorldOffset, URigHierarchy* Hierarchy, EControlRigDrawHierarchyMode::Type Mode, float Scale, const FLinearColor& Color, float Thickness, const FRigPose* InPose);
 };
 
 /**
diff --git a/Engine/Plugins/Animation/ControlRig/Source/ControlRig/Private/Units/Debug/RigUnit_DebugLine.cpp b/Engine/Plugins/Animation/ControlRig/Source/ControlRig/Private/Units/Debug/RigUnit_DebugLine.cpp
--- a/Engine/Plugins/Animation/ControlRig/Source/ControlRig/Private/Units/Debug/RigUnit_DebugLine.cpp
+++ b/Engine/Plugins/Animation/ControlRig/Source/ControlRig/Private/Units/Debug/RigUnit_DebugLine.cpp
@@ -38,7 +38,7 @@ FRigUnit_DebugLineItemSpace_Execute()
 {
     DECLARE_SCOPE_HIERARCHICAL_COUNTER_RIGUNIT()
 
-	if (ExecuteContext.UnitContext.DrawInterface == nullptr || !bEnabled)
+	if (ExecuteContext.GetDrawInterface() == nullptr || !bEnabled)
 	{
 		return;
 	}
@@ -51,5 +51,5 @@ FRigUnit_DebugLineItemSpace_Execute()
 		DrawB = Transform.TransformPosition(DrawB);
 	}
 
-	ExecuteContext.UnitContext.DrawInterface->DrawLine(WorldOffset, DrawA, DrawB, Color, Thickness);
+	ExecuteContext.GetDrawInterface()->DrawLine(WorldOffset, DrawA, DrawB, Color, Thickness);
 }
diff --git a/Engine/Plugins/Animation/ControlRig/Source/ControlRig/Private/Units/Debug/RigUnit_DebugLineStrip.cpp b/Engine/Plugins/Animation/ControlRig/Source/ControlRig/Private/Units/Debug/RigUnit_DebugLineStrip.cpp
--- a/Engine/Plugins/Animation/ControlRig/Source/ControlRig/Private/Units/Debug/RigUnit_DebugLineStrip.cpp
+++ b/Engine/Plugins/Animation/ControlRig/Source/ControlRig/Private/Units/Debug/RigUnit_DebugLineStrip.cpp
@@ -36,7 +36,7 @@ FRigUnit_DebugLineStripItemSpace_Execute()
 {
     DECLARE_SCOPE_HIERARCHICAL_COUNTER_RIGUNIT()
 
-	if (ExecuteContext.UnitContext.DrawInterface == nullptr || !bEnabled)
+	if (ExecuteContext.GetDrawInterface() == nullptr || !bEnabled)
 	{
 		return;
 	}
@@ -50,11 +50,11 @@ FRigUnit_DebugLineStripItemSpace_Execute()
 		{
 			PointsTransformed.Add(Transform.TransformPosition(Point));
 		}
-		ExecuteContext.UnitContext.DrawInterface->DrawLineStrip(WorldOffset, PointsTransformed, Color, Thickness);
+		ExecuteContext.GetDrawInterface()->DrawLineStrip(WorldOffset, PointsTransformed, Color, Thickness);
 	}
 	else
 	{
-		ExecuteContext.UnitContext.DrawInterface->DrawLineStrip(WorldOffset, TArrayView<const FVector>(Points.GetData(), Points.Num()), Color, Thickness);
+		ExecuteContext.GetDrawInterface()->DrawLineStrip(WorldOffset, TArrayView<const FVector>(Points.GetData(), Points.Num()), Color, Thickness);
 	}
 }
 
diff --git a/Engine/Plugins/Animation/ControlRig/Source/ControlRig/Private/Units/Debug/RigUnit_DebugPoint.cpp b/Engine/Plugins/Animation/ControlRig/Source/ControlRig/Private/Units/Debug/RigUnit_DebugPoint.cpp
--- a/Engine/Plugins/Animation/ControlRig/Source/ControlRig/Private/Units/Debug/RigUnit_DebugPoint.cpp
+++ b/Engine/Plugins/Animation/ControlRig/Source/ControlRig/Private/Units/Debug/RigUnit_DebugPoint.cpp
@@ -10,7 +10,7 @@ FRigUnit_DebugPoint_Execute()
 {
     DECLARE_SCOPE_HIERARCHICAL_COUNTER_RIGUNIT()
 
-	if (ExecuteContext.UnitContext.DrawInterface == nullptr || !bEnabled)
+	if (ExecuteContext.GetDrawInterface() == nullptr || !bEnabled)
 	{
 		return;
 	}
@@ -28,12 +28,12 @@ FRigUnit_DebugPoint_Execute()
 	{
 		case ERigUnitDebugPointMode::Point:
 		{
-			ExecuteContext.UnitContext.DrawInterface->DrawPoint(WorldOffset, DrawVector, Scale, Color);
+			ExecuteContext.GetDrawInterface()->DrawPoint(WorldOffset, DrawVector, Scale, Color);
 			break;
 		}
 		case ERigUnitDebugPointMode::Vector:
 		{
-			ExecuteContext.UnitContext.DrawInterface->DrawLine(WorldOffset, Center, DrawVector, Color, Thickness);
+			ExecuteContext.GetDrawInterface()->DrawLine(WorldOffset, Center, DrawVector, Color, Thickness);
 			break;
 		}
 	}
@@ -59,7 +59,7 @@ FRigUnit_DebugPointMutable_Execute()
 {
     DECLARE_SCOPE_HIERARCHICAL_COUNTER_RIGUNIT()
 
-	if (ExecuteContext.UnitContext.DrawInterface == nullptr || !bEnabled)
+	if (ExecuteContext.GetDrawInterface() == nullptr || !bEnabled)
 	{
 		return;
 	}
@@ -77,12 +77,12 @@ FRigUnit_DebugPointMutable_Execute()
 	{
 		case ERigUnitDebugPointMode::Point:
 		{
-			ExecuteContext.UnitContext.DrawInterface->DrawPoint(WorldOffset, DrawVector, Scale, Color);
+			ExecuteContext.GetDrawInterface()->DrawPoint(WorldOffset, DrawVector, Scale, Color);
 			break;
 		}
 		case ERigUnitDebugPointMode::Vector:
 		{
-			ExecuteContext.UnitContext.DrawInterface->DrawLine(WorldOffset, Center, DrawVector, Color, Thickness);
+			ExecuteContext.GetDrawInterface()->DrawLine(WorldOffset, Center, DrawVector, Color, Thickness);
 			break;
 		}
 	}
diff --git a/Engine/Plugins/Animation/ControlRig/Source/ControlRig/Private/Units/Debug/RigUnit_DebugPrimitives.cpp b/Engine/Plugins/Animation/ControlRig/Source/ControlRig/Private/Units/Debug/RigUnit_DebugPrimitives.cpp
--- a/Engine/Plugins/Animation/ControlRig/Source/ControlRig/Private/Units/Debug/RigUnit_DebugPrimitives.cpp
+++ b/Engine/Plugins/Animation/ControlRig/Source/ControlRig/Private/Units/Debug/RigUnit_DebugPrimitives.cpp
@@ -38,7 +38,7 @@ FRigUnit_DebugRectangleItemSpace_Execute()
 {
     DECLARE_SCOPE_HIERARCHICAL_COUNTER_RIGUNIT()
 
-	if (ExecuteContext.UnitContext.DrawInterface == nullptr || !bEnabled)
+	if (ExecuteContext.GetDrawInterface() == nullptr || !bEnabled)
 	{
 		return;
 	}
@@ -49,7 +49,7 @@ FRigUnit_DebugRectangleItemSpace_Execute()
 		DrawTransform = DrawTransform * ExecuteContext.Hierarchy->GetGlobalTransform(Space);
 	}
 
-	ExecuteContext.UnitContext.DrawInterface->DrawRectangle(WorldOffset, DrawTransform, Scale, Color, Thickness);
+	ExecuteContext.GetDrawInterface()->DrawRectangle(WorldOffset, DrawTransform, Scale, Color, Thickness);
 }
 
 FRigUnit_DebugArc_Execute()
@@ -91,7 +91,7 @@ FRigUnit_DebugArcItemSpace_Execute()
 {
     DECLARE_SCOPE_HIERARCHICAL_COUNTER_RIGUNIT()
 
-	if (ExecuteContext.UnitContext.DrawInterface == nullptr || !bEnabled)
+	if (ExecuteContext.GetDrawInterface() == nullptr || !bEnabled)
 	{
 		return;
 	}
@@ -102,5 +102,5 @@ FRigUnit_DebugArcItemSpace_Execute()
 		DrawTransform = DrawTransform * ExecuteContext.Hierarchy->GetGlobalTransform(Space);
 	}
 
-	ExecuteContext.UnitContext.DrawInterface->DrawArc(WorldOffset, DrawTransform, Radius, FMath::DegreesToRadians(MinimumDegrees), FMath::DegreesToRadians(MaximumDegrees), Color, Thickness, Detail);
+	ExecuteContext.GetDrawInterface()->DrawArc(WorldOffset, DrawTransform, Radius, FMath::DegreesToRadians(MinimumDegrees), FMath::DegreesToRadians(MaximumDegrees), Color, Thickness, Detail);
 }
diff --git a/Engine/Plugins/Animation/ControlRig/Source/ControlRig/Private/Units/Debug/RigUnit_DebugTransform.cpp b/Engine/Plugins/Animation/ControlRig/Source/ControlRig/Private/Units/Debug/RigUnit_DebugTransform.cpp
--- a/Engine/Plugins/Animation/ControlRig/Source/ControlRig/Private/Units/Debug/RigUnit_DebugTransform.cpp
+++ b/Engine/Plugins/Animation/ControlRig/Source/ControlRig/Private/Units/Debug/RigUnit_DebugTransform.cpp
@@ -10,7 +10,7 @@ FRigUnit_DebugTransform_Execute()
 {
     DECLARE_SCOPE_HIERARCHICAL_COUNTER_RIGUNIT()
 
-	if (ExecuteContext.UnitContext.DrawInterface == nullptr || !bEnabled)
+	if (ExecuteContext.GetDrawInterface() == nullptr || !bEnabled)
 	{
 		return;
 	}
@@ -25,18 +25,18 @@ FRigUnit_DebugTransform_Execute()
 	{
 		case ERigUnitDebugTransformMode::Axes:
 		{
-			ExecuteContext.UnitContext.DrawInterface->DrawAxes(WorldOffset, DrawTransform, Scale, Thickness);
+			ExecuteContext.GetDrawInterface()->DrawAxes(WorldOffset, DrawTransform, Scale, Thickness);
 			break;
 		}
 		case ERigUnitDebugTransformMode::Point:
 		{
-			ExecuteContext.UnitContext.DrawInterface->DrawPoint(WorldOffset, DrawTransform.GetLocation(), Scale, Color);
+			ExecuteContext.GetDrawInterface()->DrawPoint(WorldOffset, DrawTransform.GetLocation(), Scale, Color);
 			break;
 		}
 		case ERigUnitDebugTransformMode::Box:
 		{
 			DrawTransform.SetScale3D(DrawTransform.GetScale3D() * Scale);
-			ExecuteContext.UnitContext.DrawInterface->DrawBox(WorldOffset, DrawTransform, Color, Thickness);
+			ExecuteContext.GetDrawInterface()->DrawBox(WorldOffset, DrawTransform, Color, Thickness);
 			break;
 		}
 	}
@@ -93,7 +93,7 @@ FRigUnit_DebugTransformMutableItemSpace_Execute()
 {
     DECLARE_SCOPE_HIERARCHICAL_COUNTER_RIGUNIT()
 
-	if (ExecuteContext.UnitContext.DrawInterface == nullptr || !bEnabled)
+	if (ExecuteContext.GetDrawInterface() == nullptr || !bEnabled)
 	{
 		return;
 	}
@@ -108,18 +108,18 @@ FRigUnit_DebugTransformMutableItemSpace_Execute()
 	{
 		case ERigUnitDebugTransformMode::Axes:
 		{
-			ExecuteContext.UnitContext.DrawInterface->DrawAxes(WorldOffset, DrawTransform, Scale, Thickness);
+			ExecuteContext.GetDrawInterface()->DrawAxes(WorldOffset, DrawTransform, Scale, Thickness);
 			break;
 		}
 		case ERigUnitDebugTransformMode::Point:
 		{
-			ExecuteContext.UnitContext.DrawInterface->DrawPoint(WorldOffset, DrawTransform.GetLocation(), Scale, Color);
+			ExecuteContext.GetDrawInterface()->DrawPoint(WorldOffset, DrawTransform.GetLocation(), Scale, Color);
 			break;
 		}
 		case ERigUnitDebugTransformMode::Box:
 		{
 			DrawTransform.SetScale3D(DrawTransform.GetScale3D() * Scale);
-			ExecuteContext.UnitContext.DrawInterface->DrawBox(WorldOffset, DrawTransform, Color, Thickness);
+			ExecuteContext.GetDrawInterface()->DrawBox(WorldOffset, DrawTransform, Color, Thickness);
 			break;
 		}
 	}
@@ -129,7 +129,7 @@ FRigUnit_DebugTransformArrayMutable_Execute()
 {
     DECLARE_SCOPE_HIERARCHICAL_COUNTER_RIGUNIT()
 
-	if (ExecuteContext.UnitContext.DrawInterface == nullptr || !bEnabled)
+	if (ExecuteContext.GetDrawInterface() == nullptr || !bEnabled)
 	{
 		return;
 	}
@@ -157,18 +157,18 @@ FRigUnit_DebugTransformArrayMutable_Execute()
 		{
 			case ERigUnitDebugTransformMode::Axes:
 			{
-				ExecuteContext.UnitContext.DrawInterface->DrawAxes(WorldOffset, DrawTransform, Scale, Thickness);
+				ExecuteContext.GetDrawInterface()->DrawAxes(WorldOffset, DrawTransform, Scale, Thickness);
 				break;
 			}
 			case ERigUnitDebugTransformMode::Point:
 			{
-				ExecuteContext.UnitContext.DrawInterface->DrawPoint(WorldOffset, DrawTransform.GetLocation(), Scale, Color);
+				ExecuteContext.GetDrawInterface()->DrawPoint(WorldOffset, DrawTransform.GetLocation(), Scale, Color);
 				break;
 			}
 			case ERigUnitDebugTransformMode::Box:
 			{
 				DrawTransform.SetScale3D(DrawTransform.GetScale3D() * Scale);
-				ExecuteContext.UnitContext.DrawInterface->DrawBox(WorldOffset, DrawTransform, Color, Thickness);
+				ExecuteContext.GetDrawInterface()->DrawBox(WorldOffset, DrawTransform, Color, Thickness);
 				break;
 			}
 		}
@@ -196,7 +196,7 @@ FRigUnit_DebugTransformArrayMutableItemSpace_Execute()
 {
 	DECLARE_SCOPE_HIERARCHICAL_COUNTER_RIGUNIT()
 
-	if (ExecuteContext.UnitContext.DrawInterface == nullptr || !bEnabled || Transforms.Num() == 0)
+	if (ExecuteContext.GetDrawInterface() == nullptr || !bEnabled || Transforms.Num() == 0)
 	{
 		return;
 	}
@@ -221,7 +221,7 @@ FRigUnit_DebugTransformArrayMutableItemSpace_Execute()
 		{
 			if(Transforms.IsValidIndex(ParentIndices[Index]))
 			{
-				ExecuteContext.UnitContext.DrawInterface->DrawLine(
+				ExecuteContext.GetDrawInterface()->DrawLine(
 					WorldOffset,
 					Transforms[Index].GetTranslation(),
 					Transforms[ParentIndices[Index]].GetTranslation(),
diff --git a/Engine/Plugins/Animation/ControlRig/Source/ControlRig/Private/Units/Debug/RigUnit_VisualDebug.cpp b/Engine/Plugins/Animation/ControlRig/Source/ControlRig/Private/Units/Debug/RigUnit_VisualDebug.cpp
--- a/Engine/Plugins/Animation/ControlRig/Source/ControlRig/Private/Units/Debug/RigUnit_VisualDebug.cpp
+++ b/Engine/Plugins/Animation/ControlRig/Source/ControlRig/Private/Units/Debug/RigUnit_VisualDebug.cpp
@@ -29,7 +29,7 @@ FRigUnit_VisualDebugVectorItemSpace_Execute()
 {
     DECLARE_SCOPE_HIERARCHICAL_COUNTER_RIGUNIT()
 
-	if (ExecuteContext.UnitContext.DrawInterface == nullptr || !bEnabled)
+	if (ExecuteContext.GetDrawInterface() == nullptr || !bEnabled)
 	{
 		return;
 	}
@@ -44,12 +44,12 @@ FRigUnit_VisualDebugVectorItemSpace_Execute()
 	{
 		case ERigUnitVisualDebugPointMode::Point:
 		{
-			ExecuteContext.UnitContext.DrawInterface->DrawPoint(WorldOffset, Value, Thickness, Color);
+			ExecuteContext.GetDrawInterface()->DrawPoint(WorldOffset, Value, Thickness, Color);
 			break;
 		}
 		case ERigUnitVisualDebugPointMode::Vector:
 		{
-			ExecuteContext.UnitContext.DrawInterface->DrawLine(WorldOffset, FVector::ZeroVector, Value * Scale, Color, Thickness);
+			ExecuteContext.GetDrawInterface()->DrawLine(WorldOffset, FVector::ZeroVector, Value * Scale, Color, Thickness);
 			break;
 		}
 		default:
@@ -111,7 +111,7 @@ FRigUnit_VisualDebugTransformItemSpace_Execute()
 {
     DECLARE_SCOPE_HIERARCHICAL_COUNTER_RIGUNIT()
 
-	if (ExecuteContext.UnitContext.DrawInterface == nullptr || !bEnabled)
+	if (ExecuteContext.GetDrawInterface() == nullptr || !bEnabled)
 	{
 		return;
 	}
@@ -122,5 +122,5 @@ FRigUnit_VisualDebugTransformItemSpace_Execute()
 		WorldOffset = ExecuteContext.Hierarchy->GetGlobalTransform(Space);
 	}
 
-	ExecuteContext.UnitContext.DrawInterface->DrawAxes(WorldOffset, Value, Scale, Thickness);
+	ExecuteContext.GetDrawInterface()->DrawAxes(WorldOffset, Value, Scale, Thickness);
 }
diff --git a/Engine/Plugins/Animation/ControlRig/Source/ControlRig/Private/Units/Drawing/RigUnit_DrawContainer.cpp b/Engine/Plugins/Animation/ControlRig/Source/ControlRig/Private/Units/Drawing/RigUnit_DrawContainer.cpp
--- a/Engine/Plugins/Animation/ControlRig/Source/ControlRig/Private/Units/Drawing/RigUnit_DrawContainer.cpp
+++ b/Engine/Plugins/Animation/ControlRig/Source/ControlRig/Private/Units/Drawing/RigUnit_DrawContainer.cpp
@@ -9,15 +9,15 @@ FRigUnit_DrawContainerGetInstruction_Execute()
 {
     DECLARE_SCOPE_HIERARCHICAL_COUNTER_RIGUNIT()
 
-    if(ExecuteContext.UnitContext.DrawContainer == nullptr)
+    if(ExecuteContext.GetDrawContainer() == nullptr)
     {
     	return;
     }
 	
-	int32 Index = ExecuteContext.UnitContext.DrawContainer->GetIndex(InstructionName);
+	int32 Index = ExecuteContext.GetDrawContainer()->GetIndex(InstructionName);
 	if (Index != INDEX_NONE)
 	{
-		const FControlRigDrawInstruction& Instruction = (*ExecuteContext.UnitContext.DrawContainer)[Index];
+		const FRigVMDrawInstruction& Instruction = (*ExecuteContext.GetDrawContainer())[Index];
 		Color = Instruction.Color;
 		Transform = Instruction.Transform;
 	}
@@ -38,15 +38,15 @@ FRigUnit_DrawContainerSetThickness_Execute()
 {
 	DECLARE_SCOPE_HIERARCHICAL_COUNTER_RIGUNIT()
 
-	if (ExecuteContext.UnitContext.DrawContainer == nullptr)
+	if (ExecuteContext.GetDrawContainer() == nullptr)
 	{
 		return;
 	}
 
-	int32 Index = ExecuteContext.UnitContext.DrawContainer->GetIndex(InstructionName);
+	int32 Index = ExecuteContext.GetDrawContainer()->GetIndex(InstructionName);
 	if (Index != INDEX_NONE)
 	{
-		FControlRigDrawInstruction& Instruction = (*ExecuteContext.UnitContext.DrawContainer)[Index];
+		FRigVMDrawInstruction& Instruction = (*ExecuteContext.GetDrawContainer())[Index];
 		Instruction.Thickness = Thickness;
 	}
 }
@@ -61,15 +61,15 @@ FRigUnit_DrawContainerSetColor_Execute()
 {
 	DECLARE_SCOPE_HIERARCHICAL_COUNTER_RIGUNIT()
 
-	if (ExecuteContext.UnitContext.DrawContainer == nullptr)
+	if (ExecuteContext.GetDrawContainer() == nullptr)
 	{
 		return;
 	}
 
-	int32 Index = ExecuteContext.UnitContext.DrawContainer->GetIndex(InstructionName);
+	int32 Index = ExecuteContext.GetDrawContainer()->GetIndex(InstructionName);
 	if (Index != INDEX_NONE)
 	{
-		FControlRigDrawInstruction& Instruction = (*ExecuteContext.UnitContext.DrawContainer)[Index];
+		FRigVMDrawInstruction& Instruction = (*ExecuteContext.GetDrawContainer())[Index];
 		Instruction.Color = Color;
 	}
 }
@@ -84,15 +84,15 @@ FRigUnit_DrawContainerSetTransform_Execute()
 {
 	DECLARE_SCOPE_HIERARCHICAL_COUNTER_RIGUNIT()
 
-	if (ExecuteContext.UnitContext.DrawContainer == nullptr)
+	if (ExecuteContext.GetDrawContainer() == nullptr)
 	{
 		return;
 	}
 
-	int32 Index = ExecuteContext.UnitContext.DrawContainer->GetIndex(InstructionName);
+	int32 Index = ExecuteContext.GetDrawContainer()->GetIndex(InstructionName);
 	if (Index != INDEX_NONE)
 	{
-		FControlRigDrawInstruction& Instruction = (*ExecuteContext.UnitContext.DrawContainer)[Index];
+		FRigVMDrawInstruction& Instruction = (*ExecuteContext.GetDrawContainer())[Index];
 		Instruction.Transform = Transform;
 	}
 }
diff --git a/Engine/Plugins/Animation/ControlRig/Source/ControlRig/Private/Units/Highlevel/Harmonics/RigUnit_ChainHarmonics.cpp b/Engine/Plugins/Animation/ControlRig/Source/ControlRig/Private/Units/Highlevel/Harmonics/RigUnit_ChainHarmonics.cpp
--- a/Engine/Plugins/Animation/ControlRig/Source/ControlRig/Private/Units/Highlevel/Harmonics/RigUnit_ChainHarmonics.cpp
+++ b/Engine/Plugins/Animation/ControlRig/Source/ControlRig/Private/Units/Highlevel/Harmonics/RigUnit_ChainHarmonics.cpp
@@ -215,7 +215,7 @@ FRigUnit_ChainHarmonicsPerItem_Execute()
 
 	Time = Time + Speed * ExecuteContext.GetDeltaTime();
 
-	if (ExecuteContext.UnitContext.DrawInterface != nullptr && bDrawDebug)
+	if (ExecuteContext.GetDrawInterface() != nullptr && bDrawDebug)
 	{
 		HierarchyLine.SetNum(Items.Num());
 		for (int32 Index = 0; Index < Items.Num(); Index++)
@@ -223,9 +223,9 @@ FRigUnit_ChainHarmonicsPerItem_Execute()
 			HierarchyLine[Index] = Hierarchy->GetGlobalTransform(Items[Index]).GetLocation();
 		}
 
-		ExecuteContext.UnitContext.DrawInterface->DrawLineStrip(DrawWorldOffset, HierarchyLine, FLinearColor::Yellow, 0.f);
-		ExecuteContext.UnitContext.DrawInterface->DrawLines(DrawWorldOffset, VelocityLines, FLinearColor(0.3f, 0.3f, 1.f), 0.f);
-		ExecuteContext.UnitContext.DrawInterface->DrawPoints(DrawWorldOffset, PendulumPosition, 3.f, FLinearColor::Blue);
+		ExecuteContext.GetDrawInterface()->DrawLineStrip(DrawWorldOffset, HierarchyLine, FLinearColor::Yellow, 0.f);
+		ExecuteContext.GetDrawInterface()->DrawLines(DrawWorldOffset, VelocityLines, FLinearColor(0.3f, 0.3f, 1.f), 0.f);
+		ExecuteContext.GetDrawInterface()->DrawPoints(DrawWorldOffset, PendulumPosition, 3.f, FLinearColor::Blue);
 		return;
 	}
 
diff --git a/Engine/Plugins/Animation/ControlRig/Source/ControlRig/Private/Units/Highlevel/Hierarchy/RigUnit_AimBone.cpp b/Engine/Plugins/Animation/ControlRig/Source/ControlRig/Private/Units/Highlevel/Hierarchy/RigUnit_AimBone.cpp
--- a/Engine/Plugins/Animation/ControlRig/Source/ControlRig/Private/Units/Highlevel/Hierarchy/RigUnit_AimBone.cpp
+++ b/Engine/Plugins/Animation/ControlRig/Source/ControlRig/Private/Units/Highlevel/Hierarchy/RigUnit_AimBone.cpp
@@ -48,17 +48,17 @@ FRigUnit_AimBoneMath_Execute()
 			}
 		}
 
-		if (ExecuteContext.UnitContext.DrawInterface != nullptr && DebugSettings.bEnabled)
+		if (ExecuteContext.GetDrawInterface() != nullptr && DebugSettings.bEnabled)
 		{
 			const FLinearColor Color = FLinearColor(0.f, 1.f, 1.f, 1.f);
 			if (Primary.Kind == EControlRigVectorKind::Direction)
 			{
-				ExecuteContext.UnitContext.DrawInterface->DrawLine(DebugSettings.WorldOffset, Result.GetLocation(), Result.GetLocation() + Target * DebugSettings.Scale, Color);
+				ExecuteContext.GetDrawInterface()->DrawLine(DebugSettings.WorldOffset, Result.GetLocation(), Result.GetLocation() + Target * DebugSettings.Scale, Color);
 			}
 			else
 			{
-				ExecuteContext.UnitContext.DrawInterface->DrawLine(DebugSettings.WorldOffset, Result.GetLocation(), Target, Color);
-				ExecuteContext.UnitContext.DrawInterface->DrawBox(DebugSettings.WorldOffset, FTransform(FQuat::Identity, Target, FVector(1.f, 1.f, 1.f) * DebugSettings.Scale * 0.1f), Color);
+				ExecuteContext.GetDrawInterface()->DrawLine(DebugSettings.WorldOffset, Result.GetLocation(), Target, Color);
+				ExecuteContext.GetDrawInterface()->DrawBox(DebugSettings.WorldOffset, FTransform(FQuat::Identity, Target, FVector(1.f, 1.f, 1.f) * DebugSettings.Scale * 0.1f), Color);
 			}
 		}
 
@@ -102,17 +102,17 @@ FRigUnit_AimBoneMath_Execute()
 			}
 		}
 
-		if (ExecuteContext.UnitContext.DrawInterface != nullptr && DebugSettings.bEnabled)
+		if (ExecuteContext.GetDrawInterface() != nullptr && DebugSettings.bEnabled)
 		{
 			const FLinearColor Color = FLinearColor(0.f, 0.2f, 1.f, 1.f);
 			if (Secondary.Kind == EControlRigVectorKind::Direction)
 			{
-				ExecuteContext.UnitContext.DrawInterface->DrawLine(DebugSettings.WorldOffset, Result.GetLocation(), Result.GetLocation() + Target * DebugSettings.Scale, Color);
+				ExecuteContext.GetDrawInterface()->DrawLine(DebugSettings.WorldOffset, Result.GetLocation(), Result.GetLocation() + Target * DebugSettings.Scale, Color);
 			}
 			else
 			{
-				ExecuteContext.UnitContext.DrawInterface->DrawLine(DebugSettings.WorldOffset, Result.GetLocation(), Target, Color);
-				ExecuteContext.UnitContext.DrawInterface->DrawBox(DebugSettings.WorldOffset, FTransform(FQuat::Identity, Target, FVector(1.f, 1.f, 1.f) * DebugSettings.Scale * 0.1f), Color);
+				ExecuteContext.GetDrawInterface()->DrawLine(DebugSettings.WorldOffset, Result.GetLocation(), Target, Color);
+				ExecuteContext.GetDrawInterface()->DrawBox(DebugSettings.WorldOffset, FTransform(FQuat::Identity, Target, FVector(1.f, 1.f, 1.f) * DebugSettings.Scale * 0.1f), Color);
 			}
 		}
 
diff --git a/Engine/Plugins/Animation/ControlRig/Source/ControlRig/Private/Units/Highlevel/Hierarchy/RigUnit_FitChainToCurve.cpp b/Engine/Plugins/Animation/ControlRig/Source/ControlRig/Private/Units/Highlevel/Hierarchy/RigUnit_FitChainToCurve.cpp
--- a/Engine/Plugins/Animation/ControlRig/Source/ControlRig/Private/Units/Highlevel/Hierarchy/RigUnit_FitChainToCurve.cpp
+++ b/Engine/Plugins/Animation/ControlRig/Source/ControlRig/Private/Units/Highlevel/Hierarchy/RigUnit_FitChainToCurve.cpp
@@ -1,6 +1,7 @@
 // Copyright Epic Games, Inc. All Rights Reserved.
 
 #include "Units/Highlevel/Hierarchy/RigUnit_FitChainToCurve.h"
+#include "Units/Debug/RigUnit_DebugBezier.h"
 #include "Units/RigUnitContext.h"
 
 #include UE_INLINE_GENERATED_CPP_BY_NAME(RigUnit_FitChainToCurve)
@@ -440,17 +441,17 @@ FRigUnit_FitChainToCurveItemArray_Execute()
 		}
 	}
 
-	if (ExecuteContext.UnitContext.DrawInterface != nullptr && DebugSettings.bEnabled)
+	if (ExecuteContext.GetDrawInterface() != nullptr && DebugSettings.bEnabled)
 	{
-		ExecuteContext.UnitContext.DrawInterface->DrawBezier(DebugSettings.WorldOffset, Bezier, 0.f, 1.f, DebugSettings.CurveColor, DebugSettings.Scale, 64);
-		ExecuteContext.UnitContext.DrawInterface->DrawPoint(DebugSettings.WorldOffset, Bezier.A, DebugSettings.Scale * 6, DebugSettings.CurveColor);
-		ExecuteContext.UnitContext.DrawInterface->DrawPoint(DebugSettings.WorldOffset, Bezier.B, DebugSettings.Scale * 6, DebugSettings.CurveColor);
-		ExecuteContext.UnitContext.DrawInterface->DrawPoint(DebugSettings.WorldOffset, Bezier.C, DebugSettings.Scale * 6, DebugSettings.CurveColor);
-		ExecuteContext.UnitContext.DrawInterface->DrawPoint(DebugSettings.WorldOffset, Bezier.D, DebugSettings.Scale * 6, DebugSettings.CurveColor);
-		ExecuteContext.UnitContext.DrawInterface->DrawLineStrip(DebugSettings.WorldOffset, CurvePositions, DebugSettings.SegmentsColor, DebugSettings.Scale);
-		ExecuteContext.UnitContext.DrawInterface->DrawPoints(DebugSettings.WorldOffset, CurvePositions, DebugSettings.Scale * 4.f, DebugSettings.SegmentsColor);
-		// ExecuteContext.UnitContext.DrawInterface->DrawPoints(DebugSettings.WorldOffset, CurvePositions, DebugSettings.Scale * 3.f, FLinearColor::Blue);
-		// ExecuteContext.UnitContext.DrawInterface->DrawPoints(DebugSettings.WorldOffset, ItemPositions, DebugSettings.Scale * 6.f, DebugSettings.SegmentsColor);
+		FRigUnit_DebugBezierItemSpace::DrawBezier(ExecuteContext, DebugSettings.WorldOffset, Bezier, 0.f, 1.f, DebugSettings.CurveColor, DebugSettings.Scale, 64);
+		ExecuteContext.GetDrawInterface()->DrawPoint(DebugSettings.WorldOffset, Bezier.A, DebugSettings.Scale * 6, DebugSettings.CurveColor);
+		ExecuteContext.GetDrawInterface()->DrawPoint(DebugSettings.WorldOffset, Bezier.B, DebugSettings.Scale * 6, DebugSettings.CurveColor);
+		ExecuteContext.GetDrawInterface()->DrawPoint(DebugSettings.WorldOffset, Bezier.C, DebugSettings.Scale * 6, DebugSettings.CurveColor);
+		ExecuteContext.GetDrawInterface()->DrawPoint(DebugSettings.WorldOffset, Bezier.D, DebugSettings.Scale * 6, DebugSettings.CurveColor);
+		ExecuteContext.GetDrawInterface()->DrawLineStrip(DebugSettings.WorldOffset, CurvePositions, DebugSettings.SegmentsColor, DebugSettings.Scale);
+		ExecuteContext.GetDrawInterface()->DrawPoints(DebugSettings.WorldOffset, CurvePositions, DebugSettings.Scale * 4.f, DebugSettings.SegmentsColor);
+		// ExecuteContext.GetDrawInterface()->DrawPoints(DebugSettings.WorldOffset, CurvePositions, DebugSettings.Scale * 3.f, FLinearColor::Blue);
+		// ExecuteContext.GetDrawInterface()->DrawPoints(DebugSettings.WorldOffset, ItemPositions, DebugSettings.Scale * 6.f, DebugSettings.SegmentsColor);
 	}
 }
 
diff --git a/Engine/Plugins/Animation/ControlRig/Source/ControlRig/Private/Units/Highlevel/Hierarchy/RigUnit_SphericalPoseReader.cpp b/Engine/Plugins/Animation/ControlRig/Source/ControlRig/Private/Units/Highlevel/Hierarchy/RigUnit_SphericalPoseReader.cpp
--- a/Engine/Plugins/Animation/ControlRig/Source/ControlRig/Private/Units/Highlevel/Hierarchy/RigUnit_SphericalPoseReader.cpp
+++ b/Engine/Plugins/Animation/ControlRig/Source/ControlRig/Private/Units/Highlevel/Hierarchy/RigUnit_SphericalPoseReader.cpp
@@ -93,7 +93,7 @@ FRigUnit_SphericalPoseReader_Execute()
 	{
 		// avoid NaNs from DistanceToEllipse, and guaranteed to be inside inner ellipse
 		OutputParam = 1.0f;
-		Debug.DrawDebug(WorldOffset, ExecuteContext.UnitContext.DrawInterface, InnerRegion, OuterRegion, DriverNormal, Driver2D, OutputParam);
+		Debug.DrawDebug(WorldOffset, ExecuteContext.GetDrawInterface(), InnerRegion, OuterRegion, DriverNormal, Driver2D, OutputParam);
 		return;
 	}
 	
@@ -117,7 +117,7 @@ FRigUnit_SphericalPoseReader_Execute()
 	OutputParam = CalcOutputParam(InnerEllipseResults, OuterEllipseResults);
 
 	// do all debug drawing
-	Debug.DrawDebug(WorldOffset, ExecuteContext.UnitContext.DrawInterface, InnerRegion, OuterRegion, DriverNormal, Driver2D, OutputParam);
+	Debug.DrawDebug(WorldOffset, ExecuteContext.GetDrawInterface(), InnerRegion, OuterRegion, DriverNormal, Driver2D, OutputParam);
 }
 
 void FRigUnit_SphericalPoseReader::RemapAndConvertInputs(
diff --git a/Engine/Plugins/Animation/ControlRig/Source/ControlRig/Private/Units/Highlevel/Hierarchy/RigUnit_SphericalPoseReader.h b/Engine/Plugins/Animation/ControlRig/Source/ControlRig/Private/Units/Highlevel/Hierarchy/RigUnit_SphericalPoseReader.h
--- a/Engine/Plugins/Animation/ControlRig/Source/ControlRig/Private/Units/Highlevel/Hierarchy/RigUnit_SphericalPoseReader.h
+++ b/Engine/Plugins/Animation/ControlRig/Source/ControlRig/Private/Units/Highlevel/Hierarchy/RigUnit_SphericalPoseReader.h
@@ -2,7 +2,6 @@
 
 #pragma once
 
-#include "Drawing/ControlRigDrawInterface.h"
 #include "Units/Highlevel/RigUnit_HighlevelBase.h"
 #include "RigUnit_SphericalPoseReader.generated.h"
 
@@ -93,7 +92,7 @@ struct CONTROLRIG_API FSphericalPoseReaderDebugSettings
 	
 	void DrawDebug(
 		const FTransform WorldOffset,
-		FControlRigDrawInterface* DrawInterface,
+		FRigVMDrawInterface* DrawInterface,
 		const FSphericalRegion& InnerRegion,
 		const FSphericalRegion& OuterRegion,
 		const FVector DriverNormal,
diff --git a/Engine/Plugins/Animation/ControlRig/Source/ControlRig/Private/Units/Highlevel/Hierarchy/RigUnit_SpringIK.cpp b/Engine/Plugins/Animation/ControlRig/Source/ControlRig/Private/Units/Highlevel/Hierarchy/RigUnit_SpringIK.cpp
--- a/Engine/Plugins/Animation/ControlRig/Source/ControlRig/Private/Units/Highlevel/Hierarchy/RigUnit_SpringIK.cpp
+++ b/Engine/Plugins/Animation/ControlRig/Source/ControlRig/Private/Units/Highlevel/Hierarchy/RigUnit_SpringIK.cpp
@@ -301,11 +301,10 @@ FRigUnit_SpringIK_Execute()
 		Hierarchy->SetGlobalTransform(CachedBones[PointIndex], Transform, bPropagateToChildren);
 	}
 
-	if (ExecuteContext.UnitContext.DrawInterface != nullptr && DebugSettings.bEnabled)
+	if (ExecuteContext.GetDrawInterface() != nullptr && DebugSettings.bEnabled)
 	{
-		ExecuteContext.UnitContext.DrawInterface->DrawPointSimulation(DebugSettings.WorldOffset, Simulation, DebugSettings.Color, DebugSettings.Scale * 0.25f);
-		ExecuteContext.UnitContext.DrawInterface->DrawLine(DebugSettings.WorldOffset, PoleTarget, FirstPoint, DebugSettings.Color, 0.f);
-		ExecuteContext.UnitContext.DrawInterface->DrawLine(DebugSettings.WorldOffset, PoleTarget, LastPoint, DebugSettings.Color, 0.f);
-		ExecuteContext.UnitContext.DrawInterface->DrawBox(DebugSettings.WorldOffset, FTransform(FQuat::Identity, PoleTarget, FVector::OneVector * DebugSettings.Scale * 10.f), DebugSettings.Color);
+		ExecuteContext.GetDrawInterface()->DrawLine(DebugSettings.WorldOffset, PoleTarget, FirstPoint, DebugSettings.Color, 0.f);
+		ExecuteContext.GetDrawInterface()->DrawLine(DebugSettings.WorldOffset, PoleTarget, LastPoint, DebugSettings.Color, 0.f);
+		ExecuteContext.GetDrawInterface()->DrawBox(DebugSettings.WorldOffset, FTransform(FQuat::Identity, PoleTarget, FVector::OneVector * DebugSettings.Scale * 10.f), DebugSettings.Color);
 	}
 }
diff --git a/Engine/Plugins/Animation/ControlRig/Source/ControlRig/Private/Units/Highlevel/Hierarchy/RigUnit_TwoBoneIKSimple.cpp b/Engine/Plugins/Animation/ControlRig/Source/ControlRig/Private/Units/Highlevel/Hierarchy/RigUnit_TwoBoneIKSimple.cpp
--- a/Engine/Plugins/Animation/ControlRig/Source/ControlRig/Private/Units/Highlevel/Hierarchy/RigUnit_TwoBoneIKSimple.cpp
+++ b/Engine/Plugins/Animation/ControlRig/Source/ControlRig/Private/Units/Highlevel/Hierarchy/RigUnit_TwoBoneIKSimple.cpp
@@ -145,14 +145,14 @@ FRigUnit_TwoBoneIKSimplePerItem_Execute()
 
 	FControlRigMathLibrary::SolveBasicTwoBoneIK(TransformA, TransformB, TransformC, PoleTarget, PrimaryAxis, SecondaryAxis, SecondaryAxisWeight, LengthA, LengthB, bEnableStretch, StretchStartRatio, StretchMaximumRatio);
 
-	if (ExecuteContext.UnitContext.DrawInterface != nullptr && DebugSettings.bEnabled)
+	if (ExecuteContext.GetDrawInterface() != nullptr && DebugSettings.bEnabled)
 	{
 		const FLinearColor Dark = FLinearColor(0.f, 0.2f, 1.f, 1.f);
 		const FLinearColor Bright = FLinearColor(0.f, 1.f, 1.f, 1.f);
-		ExecuteContext.UnitContext.DrawInterface->DrawLine(DebugSettings.WorldOffset, TransformA.GetLocation(), TransformB.GetLocation(), Dark);
-		ExecuteContext.UnitContext.DrawInterface->DrawLine(DebugSettings.WorldOffset, TransformB.GetLocation(), TransformC.GetLocation(), Dark);
-		ExecuteContext.UnitContext.DrawInterface->DrawLine(DebugSettings.WorldOffset, TransformB.GetLocation(), PoleTarget, Bright);
-		ExecuteContext.UnitContext.DrawInterface->DrawBox(DebugSettings.WorldOffset, FTransform(FQuat::Identity, PoleTarget, FVector(1.f, 1.f, 1.f) * DebugSettings.Scale * 0.1f), Bright);
+		ExecuteContext.GetDrawInterface()->DrawLine(DebugSettings.WorldOffset, TransformA.GetLocation(), TransformB.GetLocation(), Dark);
+		ExecuteContext.GetDrawInterface()->DrawLine(DebugSettings.WorldOffset, TransformB.GetLocation(), TransformC.GetLocation(), Dark);
+		ExecuteContext.GetDrawInterface()->DrawLine(DebugSettings.WorldOffset, TransformB.GetLocation(), PoleTarget, Bright);
+		ExecuteContext.GetDrawInterface()->DrawBox(DebugSettings.WorldOffset, FTransform(FQuat::Identity, PoleTarget, FVector(1.f, 1.f, 1.f) * DebugSettings.Scale * 0.1f), Bright);
 	}
 
 	if (Weight < 1.0f - SMALL_NUMBER)
diff --git a/Engine/Plugins/Animation/ControlRig/Source/ControlRig/Private/Units/Math/RigUnit_MathVector.cpp b/Engine/Plugins/Animation/ControlRig/Source/ControlRig/Private/Units/Math/RigUnit_MathVector.cpp
--- a/Engine/Plugins/Animation/ControlRig/Source/ControlRig/Private/Units/Math/RigUnit_MathVector.cpp
+++ b/Engine/Plugins/Animation/ControlRig/Source/ControlRig/Private/Units/Math/RigUnit_MathVector.cpp
@@ -433,7 +433,7 @@ FRigUnit_MathVectorClampSpatially_Execute()
 	DECLARE_SCOPE_HIERARCHICAL_COUNTER_RIGUNIT()
 	Result = FControlRigMathLibrary::ClampSpatially(Value, Axis, Type, Minimum, Maximum, Space);
 
-	if (ExecuteContext.UnitContext.DrawInterface != nullptr && bDrawDebug)
+	if (ExecuteContext.GetDrawInterface() != nullptr && bDrawDebug)
 	{
 		switch (Type)
 		{
@@ -464,8 +464,8 @@ FRigUnit_MathVectorClampSpatially_Execute()
 					}
 				}
 
-				ExecuteContext.UnitContext.DrawInterface->DrawLine(Space, Points[0], Points[1], DebugColor, DebugThickness);
-				ExecuteContext.UnitContext.DrawInterface->DrawPoints(Space, Points, DebugThickness * 8.f, DebugColor);
+				ExecuteContext.GetDrawInterface()->DrawLine(Space, Points[0], Points[1], DebugColor, DebugThickness);
+				ExecuteContext.GetDrawInterface()->DrawPoints(Space, Points, DebugThickness * 8.f, DebugColor);
 
 				break;
 			}
@@ -491,9 +491,9 @@ FRigUnit_MathVectorClampSpatially_Execute()
 				}
 				if (Minimum > SMALL_NUMBER)
 				{
-					ExecuteContext.UnitContext.DrawInterface->DrawArc(Space, CircleTransform, Minimum, 0.f, PI * 2.f, DebugColor, DebugThickness, 16);
+					ExecuteContext.GetDrawInterface()->DrawArc(Space, CircleTransform, Minimum, 0.f, PI * 2.f, DebugColor, DebugThickness, 16);
 				}
-				ExecuteContext.UnitContext.DrawInterface->DrawArc(Space, CircleTransform, Maximum, 0.f, PI * 2.f, DebugColor, DebugThickness, 16);
+				ExecuteContext.GetDrawInterface()->DrawArc(Space, CircleTransform, Maximum, 0.f, PI * 2.f, DebugColor, DebugThickness, 16);
 				break;
 			}
 			default:
@@ -502,21 +502,21 @@ FRigUnit_MathVectorClampSpatially_Execute()
 				FTransform CircleTransform = FTransform::Identity;
 				if (Minimum > SMALL_NUMBER)
 				{
-					ExecuteContext.UnitContext.DrawInterface->DrawArc(Space, CircleTransform, Minimum, 0.f, PI * 2.f, DebugColor, DebugThickness, 16);
+					ExecuteContext.GetDrawInterface()->DrawArc(Space, CircleTransform, Minimum, 0.f, PI * 2.f, DebugColor, DebugThickness, 16);
 				}
-				ExecuteContext.UnitContext.DrawInterface->DrawArc(Space, CircleTransform, Maximum, 0.f, PI * 2.f, DebugColor, DebugThickness, 16);
+				ExecuteContext.GetDrawInterface()->DrawArc(Space, CircleTransform, Maximum, 0.f, PI * 2.f, DebugColor, DebugThickness, 16);
 				CircleTransform.SetRotation(FQuat(FVector(0.f, 1.f, 0.f), PI * 0.5f));
 				if (Minimum > SMALL_NUMBER)
 				{
-					ExecuteContext.UnitContext.DrawInterface->DrawArc(Space, CircleTransform, Minimum, 0.f, PI * 2.f, DebugColor, DebugThickness, 16);
+					ExecuteContext.GetDrawInterface()->DrawArc(Space, CircleTransform, Minimum, 0.f, PI * 2.f, DebugColor, DebugThickness, 16);
 				}
-				ExecuteContext.UnitContext.DrawInterface->DrawArc(Space, CircleTransform, Maximum, 0.f, PI * 2.f, DebugColor, DebugThickness, 16);
+				ExecuteContext.GetDrawInterface()->DrawArc(Space, CircleTransform, Maximum, 0.f, PI * 2.f, DebugColor, DebugThickness, 16);
 				CircleTransform.SetRotation(FQuat(FVector(1.f, 0.f, 0.f), PI * 0.5f));
 				if (Minimum > SMALL_NUMBER)
 				{
-					ExecuteContext.UnitContext.DrawInterface->DrawArc(Space, CircleTransform, Minimum, 0.f, PI * 2.f, DebugColor, DebugThickness, 16);
+					ExecuteContext.GetDrawInterface()->DrawArc(Space, CircleTransform, Minimum, 0.f, PI * 2.f, DebugColor, DebugThickness, 16);
 				}
-				ExecuteContext.UnitContext.DrawInterface->DrawArc(Space, CircleTransform, Maximum, 0.f, PI * 2.f, DebugColor, DebugThickness, 16);
+				ExecuteContext.GetDrawInterface()->DrawArc(Space, CircleTransform, Maximum, 0.f, PI * 2.f, DebugColor, DebugThickness, 16);
 				break;
 			}
 
diff --git a/Engine/Plugins/Animation/ControlRig/Source/ControlRig/Private/Units/Simulation/RigUnit_PointSimulation.cpp b/Engine/Plugins/Animation/ControlRig/Source/ControlRig/Private/Units/Simulation/RigUnit_PointSimulation.cpp
--- a/Engine/Plugins/Animation/ControlRig/Source/ControlRig/Private/Units/Simulation/RigUnit_PointSimulation.cpp
+++ b/Engine/Plugins/Animation/ControlRig/Source/ControlRig/Private/Units/Simulation/RigUnit_PointSimulation.cpp
@@ -249,11 +249,6 @@ FRigUnit_PointSimulation_Execute()
 		Bezier.C = Simulation.GetPointInterpolated(2).Position;
 		Bezier.D = Simulation.GetPointInterpolated(3).Position;
 	}
-
-	if (ExecuteContext.UnitContext.DrawInterface != nullptr && DebugSettings.bEnabled)
-	{
-		ExecuteContext.UnitContext.DrawInterface->DrawPointSimulation(DebugSettings.WorldOffset, Simulation, DebugSettings.Color, DebugSettings.Scale, DebugSettings.CollisionScale, DebugSettings.bDrawPointsAsSpheres);
-	}
 }
 
 FRigVMStructUpgradeInfo FRigUnit_PointSimulation::GetUpgradeInfo() const
diff --git a/Engine/Plugins/Animation/ControlRig/Source/ControlRig/Private/Validation/ControlRigNumericalValidationPass.cpp b/Engine/Plugins/Animation/ControlRig/Source/ControlRig/Private/Validation/ControlRigNumericalValidationPass.cpp
--- a/Engine/Plugins/Animation/ControlRig/Source/ControlRig/Private/Validation/ControlRigNumericalValidationPass.cpp
+++ b/Engine/Plugins/Animation/ControlRig/Source/ControlRig/Private/Validation/ControlRigNumericalValidationPass.cpp
@@ -153,7 +153,7 @@ void UControlRigNumericalValidationPass::OnEvent(UControlRig* InControlRig, cons
 
 				if (InControlRig->GetHierarchy()->IsSelected(Key))
 				{
-					if (FControlRigDrawInterface* DrawInterface = InContext->GetDrawInterface())
+					if (FRigVMDrawInterface* DrawInterface = InContext->GetDrawInterface())
 					{
 						DrawInterface->DrawAxes(FTransform::Identity, A, 50.f);
 						DrawInterface->DrawAxes(FTransform::Identity, B, 50.f);
diff --git a/Engine/Plugins/Animation/ControlRig/Source/ControlRig/Public/ControlRig.h b/Engine/Plugins/Animation/ControlRig/Source/ControlRig/Public/ControlRig.h
--- a/Engine/Plugins/Animation/ControlRig/Source/ControlRig/Public/ControlRig.h
+++ b/Engine/Plugins/Animation/ControlRig/Source/ControlRig/Public/ControlRig.h
@@ -485,11 +485,11 @@ private:
 
 public:
 	
-	const FControlRigDrawInterface& GetDrawInterface() const { return DrawInterface; };
-	FControlRigDrawInterface& GetDrawInterface() { return DrawInterface; };
+	const FRigVMDrawInterface& GetDrawInterface() const { return DrawInterface; };
+	FRigVMDrawInterface& GetDrawInterface() { return DrawInterface; };
 
-	const FControlRigDrawContainer& GetDrawContainer() const { return DrawContainer; };
-	FControlRigDrawContainer& GetDrawContainer() { return DrawContainer; };
+	const FRigVMDrawContainer& GetDrawContainer() const { return DrawContainer; };
+	FRigVMDrawContainer& GetDrawContainer() { return DrawContainer; };
 
 	const FRigControlElementCustomization* GetControlCustomization(const FRigElementKey& InControl) const;
 	void SetControlCustomization(const FRigElementKey& InControl, const FRigControlElementCustomization& InCustomization);
@@ -526,7 +526,7 @@ protected:
 	TSharedPtr<IControlRigObjectBinding> ObjectBinding;
 
 #if WITH_EDITOR
-	FControlRigLog* ControlRigLog;
+	FRigVMLog* ControlRigLog;
 	bool bEnableControlRigLogging;
 #endif
 
@@ -587,10 +587,10 @@ private:
 #endif
 	
 	UPROPERTY()
-	FControlRigDrawContainer DrawContainer;
+	FRigVMDrawContainer DrawContainer;
 
 	/** The draw interface for the units to use */
-	FControlRigDrawInterface DrawInterface;
+	FRigVMDrawInterface DrawInterface;
 
 	/** The registry to access data source */
 	UPROPERTY(transient)
diff --git a/Engine/Plugins/Animation/ControlRig/Source/ControlRig/Public/ControlRigLog.h b/Engine/Plugins/Animation/ControlRig/Source/ControlRig/Public/ControlRigLog.h
deleted file mode 100644
diff --git a/Engine/Plugins/Animation/ControlRig/Source/ControlRig/Public/ControlRigValidationPass.h b/Engine/Plugins/Animation/ControlRig/Source/ControlRig/Public/ControlRigValidationPass.h
--- a/Engine/Plugins/Animation/ControlRig/Source/ControlRig/Public/ControlRigValidationPass.h
+++ b/Engine/Plugins/Animation/ControlRig/Source/ControlRig/Public/ControlRigValidationPass.h
@@ -29,7 +29,7 @@ public:
 	FControlRigValidationClearDelegate& OnClear() { return ClearDelegate;  }
 	FControlRigValidationReportDelegate& OnReport() { return ReportDelegate; }
 
-	FControlRigDrawInterface* GetDrawInterface() { return DrawInterface; }
+	FRigVMDrawInterface* GetDrawInterface() { return DrawInterface; }
 
 	FString GetDisplayNameForEvent(const FName& InEventName) const;
 
@@ -38,7 +38,7 @@ private:
 	FControlRigValidationClearDelegate ClearDelegate;;
 	FControlRigValidationReportDelegate ReportDelegate;
 
-	FControlRigDrawInterface* DrawInterface;
+	FRigVMDrawInterface* DrawInterface;
 
 	friend class UControlRigValidator;
 };
diff --git a/Engine/Plugins/Animation/ControlRig/Source/ControlRig/Public/Drawing/ControlRigDrawContainer.h b/Engine/Plugins/Animation/ControlRig/Source/ControlRig/Public/Drawing/ControlRigDrawContainer.h
deleted file mode 100644
diff --git a/Engine/Plugins/Animation/ControlRig/Source/ControlRig/Public/Drawing/ControlRigDrawInstruction.h b/Engine/Plugins/Animation/ControlRig/Source/ControlRig/Public/Drawing/ControlRigDrawInstruction.h
deleted file mode 100644
diff --git a/Engine/Plugins/Animation/ControlRig/Source/ControlRig/Public/Drawing/ControlRigDrawInterface.h b/Engine/Plugins/Animation/ControlRig/Source/ControlRig/Public/Drawing/ControlRigDrawInterface.h
deleted file mode 100644
diff --git a/Engine/Plugins/Animation/ControlRig/Source/ControlRig/Public/Units/RigUnitContext.h b/Engine/Plugins/Animation/ControlRig/Source/ControlRig/Public/Units/RigUnitContext.h
--- a/Engine/Plugins/Animation/ControlRig/Source/ControlRig/Public/Units/RigUnitContext.h
+++ b/Engine/Plugins/Animation/ControlRig/Source/ControlRig/Public/Units/RigUnitContext.h
@@ -5,12 +5,11 @@
 #include "CoreMinimal.h"
 #include "Rigs/RigHierarchyContainer.h"
 #include "Rigs/RigCurveContainer.h"
-#include "ControlRigLog.h"
 #include "AnimationDataSource.h"
 #include "Animation/AttributesRuntime.h"
-#include "Drawing/ControlRigDrawInterface.h"
 #include "GameFramework/Actor.h"
 #include "Components/SceneComponent.h"
+#include "RigVMCore/RigVMExecuteContext.h"
 #include "RigUnitContext.generated.h"
 
 /**
@@ -47,8 +46,6 @@ struct FRigUnitContext
 	/** default constructor */
 	FRigUnitContext()
 		: AnimAttributeContainer(nullptr)
-		, DrawInterface(nullptr)
-		, DrawContainer(nullptr)
 		, DataSourceRegistry(nullptr)
 		, InteractionType((uint8)EControlRigInteractionType::None)
 		, ElementsBeingInteracted()
@@ -56,21 +53,12 @@ struct FRigUnitContext
 		, OwningComponent(nullptr)
 		, OwningActor(nullptr)
 		, World(nullptr)
-#if WITH_EDITOR
-		, Log(nullptr)
-#endif
 	{
 	}
 
 	/** An external anim attribute container */
 	UE::Anim::FStackAttributeContainer* AnimAttributeContainer;
 	
-	/** The draw interface for the units to use */
-	FControlRigDrawInterface* DrawInterface;
-
-	/** The draw container for the units to use */
-	FControlRigDrawContainer* DrawContainer;
-
 	/** The registry to access data source */
 	const UAnimationDataSourceRegistry* DataSourceRegistry;
 
@@ -95,11 +83,6 @@ struct FRigUnitContext
 	/** The world this rig is running in */
 	const UWorld* World;
 
-#if WITH_EDITOR
-	/** A handle to the compiler log */
-	FControlRigLog* Log;
-#endif
-
 	/**
 	 * Returns a given data source and cast it to the expected class.
 	 *
@@ -198,9 +181,9 @@ struct FControlRigExecuteContext : public FRigVMExecuteContext
 
 #if WITH_EDITOR
 #define UE_CONTROLRIG_RIGUNIT_REPORT(Severity, Format, ...) \
-if(ExecuteContext.UnitContext.Log != nullptr) \
+if(ExecuteContext.GetLog() != nullptr) \
 { \
-	ExecuteContext.UnitContext.Log->Report(EMessageSeverity::Severity, ExecuteContext.GetFunctionName(), ExecuteContext.GetInstructionIndex(), FString::Printf((Format), ##__VA_ARGS__)); \
+	ExecuteContext.GetLog()->Report(EMessageSeverity::Severity, ExecuteContext.GetFunctionName(), ExecuteContext.GetInstructionIndex(), FString::Printf((Format), ##__VA_ARGS__)); \
 }
 #define UE_CONTROLRIG_RIGUNIT_LOG_MESSAGE(Format, ...) UE_CONTROLRIG_RIGUNIT_REPORT(Info, (Format), ##__VA_ARGS__)
 #define UE_CONTROLRIG_RIGUNIT_REPORT_WARNING(Format, ...) UE_CONTROLRIG_RIGUNIT_REPORT(Warning, (Format), ##__VA_ARGS__)
diff --git a/Engine/Plugins/Animation/ControlRig/Source/ControlRigDeveloper/Private/Graph/ControlRigGraph.cpp b/Engine/Plugins/Animation/ControlRig/Source/ControlRigDeveloper/Private/Graph/ControlRigGraph.cpp
--- a/Engine/Plugins/Animation/ControlRig/Source/ControlRigDeveloper/Private/Graph/ControlRigGraph.cpp
+++ b/Engine/Plugins/Animation/ControlRig/Source/ControlRigDeveloper/Private/Graph/ControlRigGraph.cpp
@@ -150,7 +150,7 @@ void UControlRigGraph::Serialize(FArchive& Ar)
 
 TArray<TSharedPtr<FString>> UControlRigGraph::EmptyElementNameList;
 
-void UControlRigGraph::CacheNameLists(URigHierarchy* InHierarchy, const FControlRigDrawContainer* DrawContainer, TArray<TSoftObjectPtr<UControlRigShapeLibrary>> ShapeLibraries)
+void UControlRigGraph::CacheNameLists(URigHierarchy* InHierarchy, const FRigVMDrawContainer* DrawContainer, TArray<TSoftObjectPtr<UControlRigShapeLibrary>> ShapeLibraries)
 {
 	if (UControlRigGraph* OuterGraph = Cast<UControlRigGraph>(GetOuter()))
 	{
@@ -189,7 +189,7 @@ void UControlRigGraph::CacheNameLists(URigHierarchy* InHierarchy, const FControl
 
 		LastHierarchyTopologyVersion = InHierarchy->GetTopologyVersion();
 	}
-	CacheNameList<FControlRigDrawContainer>(*DrawContainer, DrawingNameList);
+	CacheNameList<FRigVMDrawContainer>(*DrawContainer, DrawingNameList);
 
 	EntryNameList.Reset();
 	EntryNameList.Add(MakeShared<FString>(FName(NAME_None).ToString()));
diff --git a/Engine/Plugins/Animation/ControlRig/Source/ControlRigDeveloper/Public/ControlRigBlueprint.h b/Engine/Plugins/Animation/ControlRig/Source/ControlRigDeveloper/Public/ControlRigBlueprint.h
--- a/Engine/Plugins/Animation/ControlRig/Source/ControlRigDeveloper/Public/ControlRigBlueprint.h
+++ b/Engine/Plugins/Animation/ControlRig/Source/ControlRigDeveloper/Public/ControlRigBlueprint.h
@@ -18,7 +18,6 @@
 #include "RigVMModel/RigVMClient.h"
 #include "RigVMCompiler/RigVMCompiler.h"
 #include "ControlRigValidationPass.h"
-#include "Drawing/ControlRigDrawContainer.h"
 #include "RigVMCore/RigVMGraphFunctionDefinition.h"
 
 #if WITH_EDITOR
@@ -491,7 +490,7 @@ public:
 #endif
 
 	UPROPERTY(EditAnywhere, Category = "Drawing")
-	FControlRigDrawContainer DrawContainer;
+	FRigVMDrawContainer DrawContainer;
 
 #if WITH_EDITOR
 	/** Remove a transient / temporary control used to interact with a pin */
diff --git a/Engine/Plugins/Animation/ControlRig/Source/ControlRigDeveloper/Public/Graph/ControlRigGraph.h b/Engine/Plugins/Animation/ControlRig/Source/ControlRigDeveloper/Public/Graph/ControlRigGraph.h
--- a/Engine/Plugins/Animation/ControlRig/Source/ControlRigDeveloper/Public/Graph/ControlRigGraph.h
+++ b/Engine/Plugins/Animation/ControlRig/Source/ControlRigDeveloper/Public/Graph/ControlRigGraph.h
@@ -9,7 +9,6 @@
 #include "Rigs/RigHierarchyController.h"
 #include "RigVMModel/RigVMGraph.h"
 #include "RigVMCore/RigVM.h"
-#include "Drawing/ControlRigDrawContainer.h"
 #include "ControlRigGraph.generated.h"
 
 class UControlRigBlueprint;
@@ -66,7 +65,7 @@ public:
 		return GetElementNameList(ERigElementType::Curve);
 	}
 
-	void CacheNameLists(URigHierarchy* InHierarchy, const FControlRigDrawContainer* DrawContainer, TArray<TSoftObjectPtr<UControlRigShapeLibrary>> ShapeLibraries);
+	void CacheNameLists(URigHierarchy* InHierarchy, const FRigVMDrawContainer* DrawContainer, TArray<TSoftObjectPtr<UControlRigShapeLibrary>> ShapeLibraries);
 	const TArray<TSharedPtr<FString>>* GetElementNameList(ERigElementType InElementType = ERigElementType::Bone) const;
 	const TArray<TSharedPtr<FString>>* GetElementNameList(URigVMPin* InPin = nullptr) const;
 	const TArray<TSharedPtr<FString>> GetSelectedElementsNameList() const;
diff --git a/Engine/Plugins/Animation/ControlRig/Source/ControlRigEditor/Private/ControlRigDrawingDetails.cpp b/Engine/Plugins/Animation/ControlRig/Source/ControlRigEditor/Private/ControlRigDrawingDetails.cpp
--- a/Engine/Plugins/Animation/ControlRig/Source/ControlRigEditor/Private/ControlRigDrawingDetails.cpp
+++ b/Engine/Plugins/Animation/ControlRig/Source/ControlRigEditor/Private/ControlRigDrawingDetails.cpp
@@ -12,8 +12,6 @@
 #include "DesktopPlatformModule.h"
 #include "Framework/Application/SlateApplication.h"
 #include "FbxImporter.h"
-#include "Drawing/ControlRigDrawContainer.h"
-#include "Drawing/ControlRigDrawInstruction.h"
 #include "ScopedTransaction.h"
 #include "Dialogs/Dialogs.h"
 #include "SKismetInspector.h"
@@ -185,8 +183,8 @@ void FControlRigDrawContainerDetails::ImportCurvesFromFBX(const FString& InFileP
 
 		for (fbxsdk::FbxNode* FbxCurveNode : FbxCurveNodes)
 		{
-			FControlRigDrawInstruction Instruction;
-			Instruction.PrimitiveType = InSettings.bMergeCurves ? EControlRigDrawSettings::Lines : EControlRigDrawSettings::LineStrip;
+			FRigVMDrawInstruction Instruction;
+			Instruction.PrimitiveType = InSettings.bMergeCurves ? ERigVMDrawSettings::Lines : ERigVMDrawSettings::LineStrip;
 			Instruction.Transform = UnFbx::FFbxDataConverter::ConvertTransform(FbxCurveNode->EvaluateGlobalTransform());
 			Instruction.Transform.SetLocation(Instruction.Transform.GetLocation() * InSettings.Scale);
 
@@ -287,7 +285,7 @@ void FControlRigDrawContainerDetails::ImportCurvesFromFBX(const FString& InFileP
 				}
 				Instruction.Positions.Reset();
 
-				if (Instruction.PrimitiveType == EControlRigDrawSettings::Lines)
+				if (Instruction.PrimitiveType == ERigVMDrawSettings::Lines)
 				{
 					Instruction.Positions.Reserve((Line.Num() - 1) * 2 + (LineIsClosed[LineIndex] ? 2 : 0));
 				}
@@ -304,7 +302,7 @@ void FControlRigDrawContainerDetails::ImportCurvesFromFBX(const FString& InFileP
 
 				if (LineIsClosed[LineIndex])
 				{
-					if (Instruction.PrimitiveType == EControlRigDrawSettings::Lines)
+					if (Instruction.PrimitiveType == ERigVMDrawSettings::Lines)
 					{
 						Instruction.Positions.Add(Instruction.Positions.Last());
 						Instruction.Positions.Add(Instruction.Positions[0]);
diff --git a/Engine/Plugins/Animation/ControlRig/Source/ControlRigEditor/Private/ControlRigEditorModule.cpp b/Engine/Plugins/Animation/ControlRig/Source/ControlRigEditor/Private/ControlRigEditorModule.cpp
--- a/Engine/Plugins/Animation/ControlRig/Source/ControlRigEditor/Private/ControlRigEditorModule.cpp
+++ b/Engine/Plugins/Animation/ControlRig/Source/ControlRigEditor/Private/ControlRigEditorModule.cpp
@@ -173,7 +173,7 @@ void FControlRigEditorModule::StartupModule()
 	PropertiesToUnregisterOnShutdown.Add(FControlRigPythonSettings::StaticStruct()->GetFName());
 	PropertyEditorModule.RegisterCustomPropertyTypeLayout(PropertiesToUnregisterOnShutdown.Last(), FOnGetPropertyTypeCustomizationInstance::CreateStatic(&FControlRigPythonLogDetails::MakeInstance));
 
-	PropertiesToUnregisterOnShutdown.Add(FControlRigDrawContainer::StaticStruct()->GetFName());
+	PropertiesToUnregisterOnShutdown.Add(FRigVMDrawContainer::StaticStruct()->GetFName());
 	PropertyEditorModule.RegisterCustomPropertyTypeLayout(PropertiesToUnregisterOnShutdown.Last(), FOnGetPropertyTypeCustomizationInstance::CreateStatic(&FControlRigDrawContainerDetails::MakeInstance));
 
 	PropertiesToUnregisterOnShutdown.Add(FControlRigEnumControlProxyValue::StaticStruct()->GetFName());
diff --git a/Engine/Plugins/Animation/ControlRig/Source/ControlRigEditor/Private/EditMode/ControlRigEditMode.cpp b/Engine/Plugins/Animation/ControlRig/Source/ControlRigEditor/Private/EditMode/ControlRigEditMode.cpp
--- a/Engine/Plugins/Animation/ControlRig/Source/ControlRigEditor/Private/EditMode/ControlRigEditMode.cpp
+++ b/Engine/Plugins/Animation/ControlRig/Source/ControlRigEditor/Private/EditMode/ControlRigEditMode.cpp
@@ -31,7 +31,6 @@
 #include "ControlRigBlueprintGeneratedClass.h"
 #include "IControlRigObjectBinding.h"
 #include "Kismet2/BlueprintEditorUtils.h"
-#include "Drawing/ControlRigDrawInterface.h"
 #include "ControlRigBlueprint.h"
 #include "ControlRigGizmoActor.h"
 #include "Framework/MultiBox/MultiBoxBuilder.h"
@@ -881,7 +880,7 @@ void FControlRigEditMode::Render(const FSceneView* View, FViewport* Viewport, FP
 						}
 					}
 				}
-				for (const FControlRigDrawInstruction& Instruction : ControlRig->DrawInterface)
+				for (const FRigVMDrawInstruction& Instruction : ControlRig->DrawInterface)
 				{
 					if (!Instruction.IsValid())
 					{
@@ -891,7 +890,7 @@ void FControlRigEditMode::Render(const FSceneView* View, FViewport* Viewport, FP
 					FTransform InstructionTransform = Instruction.Transform * ComponentTransform;
 					switch (Instruction.PrimitiveType)
 					{
-					case EControlRigDrawSettings::Points:
+					case ERigVMDrawSettings::Points:
 					{
 						for (const FVector& Point : Instruction.Positions)
 						{
@@ -899,7 +898,7 @@ void FControlRigEditMode::Render(const FSceneView* View, FViewport* Viewport, FP
 						}
 						break;
 					}
-					case EControlRigDrawSettings::Lines:
+					case ERigVMDrawSettings::Lines:
 					{
 						const TArray<FVector>& Points = Instruction.Positions;
 						PDI->AddReserveLines(SDPG_Foreground, Points.Num() / 2, false, Instruction.Thickness > SMALL_NUMBER);
@@ -909,7 +908,7 @@ void FControlRigEditMode::Render(const FSceneView* View, FViewport* Viewport, FP
 						}
 						break;
 					}
-					case EControlRigDrawSettings::LineStrip:
+					case ERigVMDrawSettings::LineStrip:
 					{
 						const TArray<FVector>& Points = Instruction.Positions;
 						PDI->AddReserveLines(SDPG_Foreground, Points.Num() - 1, false, Instruction.Thickness > SMALL_NUMBER);
@@ -920,7 +919,7 @@ void FControlRigEditMode::Render(const FSceneView* View, FViewport* Viewport, FP
 						break;
 					}
 
-					case EControlRigDrawSettings::DynamicMesh:
+					case ERigVMDrawSettings::DynamicMesh:
 					{
 						FDynamicMeshBuilder MeshBuilder(PDI->View->GetFeatureLevel());
 						MeshBuilder.AddVertices(Instruction.MeshVerts);
diff --git a/Engine/Plugins/Animation/ControlRig/Source/ControlRigEditor/Private/Editor/ControlRigEditor.cpp b/Engine/Plugins/Animation/ControlRig/Source/ControlRigEditor/Private/Editor/ControlRigEditor.cpp
--- a/Engine/Plugins/Animation/ControlRig/Source/ControlRigEditor/Private/Editor/ControlRigEditor.cpp
+++ b/Engine/Plugins/Animation/ControlRig/Source/ControlRigEditor/Private/Editor/ControlRigEditor.cpp
@@ -3332,7 +3332,7 @@ void FControlRigEditor::HandleControlRigExecutedEvent(UControlRig* InControlRig,
 								if(DebuggedControlRig->ControlRigLog)
 								{
 									DebuggedControlRig->ControlRigLog->Entries.Add(
-										FControlRigLog::FLogEntry(EMessageSeverity::Warning, InEventName, InstructionIndex, Message
+										FRigVMLog::FLogEntry(EMessageSeverity::Warning, InEventName, InstructionIndex, Message
 									));
 								}
 
@@ -6943,7 +6943,7 @@ void FControlRigEditor::UpdateGraphCompilerErrors()
 			}
 
 			// update the nodes' error messages
-			for (const FControlRigLog::FLogEntry& Entry : ControlRigLog.Entries)
+			for (const FRigVMLog::FLogEntry& Entry : ControlRigLog.Entries)
 			{
 				URigVMNode* ModelNode = Cast<URigVMNode>(ByteCode.GetSubjectForInstruction(Entry.InstructionIndex));
 				if (ModelNode == nullptr)
diff --git a/Engine/Plugins/Animation/ControlRig/Source/ControlRigEditor/Private/Editor/ControlRigEditor.h b/Engine/Plugins/Animation/ControlRig/Source/ControlRigEditor/Private/Editor/ControlRigEditor.h
--- a/Engine/Plugins/Animation/ControlRig/Source/ControlRigEditor/Private/Editor/ControlRigEditor.h
+++ b/Engine/Plugins/Animation/ControlRig/Source/ControlRigEditor/Private/Editor/ControlRigEditor.h
@@ -8,7 +8,6 @@
 #include "DragAndDrop/GraphNodeDragDropOp.h"
 #include "ControlRigDefines.h"
 #include "Units/RigUnitContext.h"
-#include "ControlRigLog.h"
 #include "IPersonaViewport.h"
 #include "Framework/MultiBox/MultiBoxBuilder.h"
 #include "RigVMModel/RigVMGraph.h"
@@ -456,7 +455,7 @@ protected:
 	bool bExecutionControlRig;
 
 	/** The log to use for errors resulting from the init phase of the units */
-	FControlRigLog ControlRigLog;
+	FRigVMLog ControlRigLog;
 	/** Once the log is collected update the graph */
 	void UpdateGraphCompilerErrors();
 
diff --git a/Engine/Plugins/Animation/ControlRig/Source/ControlRigEditor/Private/Editor/SControlRigStackView.cpp b/Engine/Plugins/Animation/ControlRig/Source/ControlRigEditor/Private/Editor/SControlRigStackView.cpp
--- a/Engine/Plugins/Animation/ControlRig/Source/ControlRigEditor/Private/Editor/SControlRigStackView.cpp
+++ b/Engine/Plugins/Animation/ControlRig/Source/ControlRigEditor/Private/Editor/SControlRigStackView.cpp
@@ -789,8 +789,8 @@ void SControlRigStackView::RefreshTreeView(URigVM* InVM)
 			UControlRig* ControlRig = ControlRigEditor.Pin()->ControlRig;
 			if(ControlRig && ControlRig->ControlRigLog)
 			{
-				const TArray<FControlRigLog::FLogEntry>& LogEntries = ControlRig->ControlRigLog->Entries;
-				for (const FControlRigLog::FLogEntry& LogEntry : LogEntries)
+				const TArray<FRigVMLog::FLogEntry>& LogEntries = ControlRig->ControlRigLog->Entries;
+				for (const FRigVMLog::FLogEntry& LogEntry : LogEntries)
 				{
 					if (Operators.Num() <= LogEntry.InstructionIndex)
 					{
diff --git a/Engine/Plugins/Animation/ControlRigSpline/Source/ControlRigSpline/Private/ControlRigSplineUnits.cpp b/Engine/Plugins/Animation/ControlRigSpline/Source/ControlRigSpline/Private/ControlRigSplineUnits.cpp
--- a/Engine/Plugins/Animation/ControlRigSpline/Source/ControlRigSpline/Private/ControlRigSplineUnits.cpp
+++ b/Engine/Plugins/Animation/ControlRigSpline/Source/ControlRigSpline/Private/ControlRigSplineUnits.cpp
@@ -138,7 +138,7 @@ FRigUnit_TangentFromControlRigSpline_Execute()
 
 FRigUnit_DrawControlRigSpline_Execute()
 {
-	if (ExecuteContext.UnitContext.DrawInterface == nullptr)
+	if (ExecuteContext.GetDrawInterface() == nullptr)
 	{
 		return;
 	}
@@ -157,7 +157,7 @@ FRigUnit_DrawControlRigSpline_Execute()
 #endif	
 
 	int32 Count = FMath::Clamp<int32>(Detail, 4, 64);
-	FControlRigDrawInstruction Instruction(EControlRigDrawSettings::LineStrip, Color, Thickness);
+	FRigVMDrawInstruction Instruction(ERigVMDrawSettings::LineStrip, Color, Thickness);
 	Instruction.Positions.SetNumUninitialized(Count);
 
 	float T = 0;
@@ -169,7 +169,7 @@ FRigUnit_DrawControlRigSpline_Execute()
 		T += Step;
 	}
 
-	ExecuteContext.UnitContext.DrawInterface->Instructions.Add(Instruction);
+	ExecuteContext.GetDrawInterface()->Instructions.Add(Instruction);
 }
 
 FRigUnit_GetLengthControlRigSpline_Execute()
@@ -618,10 +618,10 @@ FRigUnit_FitChainToSplineCurveItemArray_Execute()
 		}
 	}
 
-	if (ExecuteContext.UnitContext.DrawInterface != nullptr && DebugSettings.bEnabled)
+	if (ExecuteContext.GetDrawInterface() != nullptr && DebugSettings.bEnabled)
 	{
 		int32 Count = 64;
-		FControlRigDrawInstruction Instruction(EControlRigDrawSettings::LineStrip, DebugSettings.CurveColor, DebugSettings.Scale, DebugSettings.WorldOffset);
+		FRigVMDrawInstruction Instruction(ERigVMDrawSettings::LineStrip, DebugSettings.CurveColor, DebugSettings.Scale, DebugSettings.WorldOffset);
 		Instruction.Positions.SetNumUninitialized(Count);
 
 		float T = 0;
@@ -632,15 +632,15 @@ FRigUnit_FitChainToSplineCurveItemArray_Execute()
 			Instruction.Positions[Index] = Spline.PositionAtParam(T);
 			T += Step;
 		}
-		ExecuteContext.UnitContext.DrawInterface->Instructions.Add(Instruction);
+		ExecuteContext.GetDrawInterface()->Instructions.Add(Instruction);
 
 		for (auto Point : Spline.SplineData->GetControlPoints())
 		{
-			ExecuteContext.UnitContext.DrawInterface->DrawPoint(DebugSettings.WorldOffset, Point, DebugSettings.Scale * 6, DebugSettings.CurveColor);
+			ExecuteContext.GetDrawInterface()->DrawPoint(DebugSettings.WorldOffset, Point, DebugSettings.Scale * 6, DebugSettings.CurveColor);
 		}
 
-		ExecuteContext.UnitContext.DrawInterface->DrawLineStrip(DebugSettings.WorldOffset, CurvePositions, DebugSettings.SegmentsColor, DebugSettings.Scale);
-		ExecuteContext.UnitContext.DrawInterface->DrawPoints(DebugSettings.WorldOffset, CurvePositions, DebugSettings.Scale * 4.f, DebugSettings.SegmentsColor);
+		ExecuteContext.GetDrawInterface()->DrawLineStrip(DebugSettings.WorldOffset, CurvePositions, DebugSettings.SegmentsColor, DebugSettings.Scale);
+		ExecuteContext.GetDrawInterface()->DrawPoints(DebugSettings.WorldOffset, CurvePositions, DebugSettings.Scale * 4.f, DebugSettings.SegmentsColor);
 	}
 }
 
diff --git a/Engine/Plugins/Animation/IKRig/Source/IKRig/Private/AnimNodes/AnimNode_IKRig.cpp b/Engine/Plugins/Animation/IKRig/Source/IKRig/Private/AnimNodes/AnimNode_IKRig.cpp
--- a/Engine/Plugins/Animation/IKRig/Source/IKRig/Private/AnimNodes/AnimNode_IKRig.cpp
+++ b/Engine/Plugins/Animation/IKRig/Source/IKRig/Private/AnimNodes/AnimNode_IKRig.cpp
@@ -6,7 +6,6 @@
 #include "IKRigDefinition.h"
 #include "IKRigSolver.h"
 #include "ActorComponents/IKRigInterface.h"
-#include "Drawing/ControlRigDrawInterface.h"
 #include "Animation/AnimInstanceProxy.h"
 #include "Algo/ForEach.h"
 #include "SceneManagement.h"
diff --git a/Engine/Plugins/Experimental/FullBodyIK/Source/FullBodyIK/Private/RigUnit_FullbodyIK.cpp b/Engine/Plugins/Experimental/FullBodyIK/Source/FullBodyIK/Private/RigUnit_FullbodyIK.cpp
--- a/Engine/Plugins/Experimental/FullBodyIK/Source/FullBodyIK/Private/RigUnit_FullbodyIK.cpp
+++ b/Engine/Plugins/Experimental/FullBodyIK/Source/FullBodyIK/Private/RigUnit_FullbodyIK.cpp
@@ -346,7 +346,7 @@ FRigUnit_FullbodyIK_Execute()
 		///////////////////////////////////////////////////////////////////////////
 		// debug draw start
 		///////////////////////////////////////////////////////////////////////////
-		if (bDebugEnabled && ExecuteContext.UnitContext.DrawInterface != nullptr)
+		if (bDebugEnabled && ExecuteContext.GetDrawInterface() != nullptr)
 		{
 			const int32 DebugDataNum = DebugData.Num();
 			if (DebugData.Num() > 0)
@@ -395,12 +395,12 @@ FRigUnit_FullbodyIK_Execute()
 							if (Data.ParentLinkIndex != INDEX_NONE)
 							{
 								const FFBIKLinkData& ParentData = LocalLink[Data.ParentLinkIndex];
-								ExecuteContext.UnitContext.DrawInterface->DrawLine(Offset, Data.GetPreviousTransform().GetLocation(), ParentData.GetPreviousTransform().GetLocation(), DrawColor, LineThickness);
+								ExecuteContext.GetDrawInterface()->DrawLine(Offset, Data.GetPreviousTransform().GetLocation(), ParentData.GetPreviousTransform().GetLocation(), DrawColor, LineThickness);
 							}
 
 							if (DebugOption.bDrawDebugAxes)
 							{
-								ExecuteContext.UnitContext.DrawInterface->DrawAxes(Offset, Data.GetPreviousTransform(), DebugOption.DrawSize);
+								ExecuteContext.GetDrawInterface()->DrawAxes(Offset, Data.GetPreviousTransform(), DebugOption.DrawSize);
 							}
 						}
 					}
@@ -413,17 +413,17 @@ FRigUnit_FullbodyIK_Execute()
 							if (EffectorTarget.bPositionEnabled)
 							{
 								// draw effector target locations
-								ExecuteContext.UnitContext.DrawInterface->DrawBox(Offset, FTransform(EffectorTarget.Position), FLinearColor::Yellow, DebugOption.DrawSize);
+								ExecuteContext.GetDrawInterface()->DrawBox(Offset, FTransform(EffectorTarget.Position), FLinearColor::Yellow, DebugOption.DrawSize);
 							}
 
 							// draw effector link location
-							ExecuteContext.UnitContext.DrawInterface->DrawBox(Offset, LocalLink[Iter.Key()].GetPreviousTransform(), FLinearColor::Green, DebugOption.DrawSize);
+							ExecuteContext.GetDrawInterface()->DrawBox(Offset, LocalLink[Iter.Key()].GetPreviousTransform(), FLinearColor::Green, DebugOption.DrawSize);
 						}
 
 						for (int32 Index = 0; Index < DebugData[DebugIndex].TargetVectorSources.Num(); ++Index)
 						{
 							// draw arrow to the target
-							ExecuteContext.UnitContext.DrawInterface->DrawLine(Offset, DebugData[DebugIndex].TargetVectorSources[Index].GetLocation(), 
+							ExecuteContext.GetDrawInterface()->DrawLine(Offset, DebugData[DebugIndex].TargetVectorSources[Index].GetLocation(), 
 								DebugData[DebugIndex].TargetVectorSources[Index].GetLocation() + DebugData[DebugIndex].TargetVectors[Index], FLinearColor::Red);
 						}
 					}
@@ -446,7 +446,7 @@ FRigUnit_FullbodyIK_Execute()
 							{
 								FTransform ConstraintFrame = LinkData[*Found].GetTransform();
 								ConstraintFrame.ConcatenateRotation(FQuat(Constraints[Index].OffsetRotation));
-								ExecuteContext.UnitContext.DrawInterface->DrawAxes(Offset, ConstraintFrame, 2.f);
+								ExecuteContext.GetDrawInterface()->DrawAxes(Offset, ConstraintFrame, 2.f);
 							}
 						}
 					}
@@ -463,7 +463,7 @@ FRigUnit_FullbodyIK_Execute()
 						// base is parent transform but in their space, we can get there by inversing local ref rotation
 						FTransform BaseTransform = FTransform(LocalRefRotation).GetRelativeTransformReverse(RotationTransform);
 						BaseTransform.ConcatenateRotation(LimitConstraint.BaseFrameOffset);
-						ExecuteContext.UnitContext.DrawInterface->DrawAxes(Offset, BaseTransform, 5.f, 1.f);
+						ExecuteContext.GetDrawInterface()->DrawAxes(Offset, BaseTransform, 5.f, 1.f);
 
 						// current transform
 						const FQuat LocalRotation = BaseTransform.GetRotation().Inverse() * RotationTransform.GetRotation();
@@ -473,7 +473,7 @@ FRigUnit_FullbodyIK_Execute()
 						RotationTransform.SetLocation(BaseTransform.GetLocation());
 
 						// draw ref pose on their current transform
-						ExecuteContext.UnitContext.DrawInterface->DrawAxes(Offset, RotationTransform, 10.f, 1.f);
+						ExecuteContext.GetDrawInterface()->DrawAxes(Offset, RotationTransform, 10.f, 1.f);
 
 						FVector XAxis = BaseTransform.GetUnitAxis(EAxis::X);
 						FVector YAxis = BaseTransform.GetUnitAxis(EAxis::Y);
@@ -484,7 +484,7 @@ FRigUnit_FullbodyIK_Execute()
 							FTransform XAxisConeTM(YAxis, XAxis ^ YAxis, XAxis, BaseTransform.GetTranslation());
 							XAxisConeTM.SetRotation(FQuat(XAxis, 0.f) * XAxisConeTM.GetRotation());
 							XAxisConeTM.SetScale3D(FVector(30.f));
-							ExecuteContext.UnitContext.DrawInterface->DrawCone(Offset, XAxisConeTM, LimitConstraint.Limit.X, 0.0f, 24, false, FLinearColor::Red, GEngine->ConstraintLimitMaterialX->GetRenderProxy());
+							ExecuteContext.GetDrawInterface()->DrawCone(Offset, XAxisConeTM, LimitConstraint.Limit.X, 0.0f, 24, false, FLinearColor::Red, GEngine->ConstraintLimitMaterialX->GetRenderProxy());
 						}
 
 						if (LimitConstraint.bYLimitSet)
@@ -492,7 +492,7 @@ FRigUnit_FullbodyIK_Execute()
 							FTransform YAxisConeTM(ZAxis, YAxis ^ ZAxis, YAxis, BaseTransform.GetTranslation());
 							YAxisConeTM.SetRotation(FQuat(YAxis, 0.f) * YAxisConeTM.GetRotation());
 							YAxisConeTM.SetScale3D(FVector(30.f));
-							ExecuteContext.UnitContext.DrawInterface->DrawCone(Offset, YAxisConeTM, LimitConstraint.Limit.Y, 0.0f, 24, false, FLinearColor::Green, GEngine->ConstraintLimitMaterialY->GetRenderProxy());
+							ExecuteContext.GetDrawInterface()->DrawCone(Offset, YAxisConeTM, LimitConstraint.Limit.Y, 0.0f, 24, false, FLinearColor::Green, GEngine->ConstraintLimitMaterialY->GetRenderProxy());
 						}
 
 						if (LimitConstraint.bZLimitSet)
@@ -500,7 +500,7 @@ FRigUnit_FullbodyIK_Execute()
 							FTransform ZAxisConeTM(XAxis, ZAxis ^ XAxis, ZAxis, BaseTransform.GetTranslation());
 							ZAxisConeTM.SetRotation(FQuat(ZAxis, 0.f) * ZAxisConeTM.GetRotation());
 							ZAxisConeTM.SetScale3D(FVector(30.f));
-							ExecuteContext.UnitContext.DrawInterface->DrawCone(Offset, ZAxisConeTM, LimitConstraint.Limit.Z, 0.0f, 24, false, FLinearColor::Blue, GEngine->ConstraintLimitMaterialZ->GetRenderProxy());
+							ExecuteContext.GetDrawInterface()->DrawCone(Offset, ZAxisConeTM, LimitConstraint.Limit.Z, 0.0f, 24, false, FLinearColor::Blue, GEngine->ConstraintLimitMaterialZ->GetRenderProxy());
 						}
 					}
 
@@ -524,8 +524,8 @@ FRigUnit_FullbodyIK_Execute()
 						Positions.Add(JointTarget);
 						Positions.Add(RootTransform.GetLocation());
 
-						ExecuteContext.UnitContext.DrawInterface->DrawLines(Offset, Positions, FLinearColor::Gray, 1.2f);
-						ExecuteContext.UnitContext.DrawInterface->DrawLine(Offset, JointTransform.GetLocation(), JointTarget, FLinearColor::Red, 1.2f);
+						ExecuteContext.GetDrawInterface()->DrawLines(Offset, Positions, FLinearColor::Gray, 1.2f);
+						ExecuteContext.GetDrawInterface()->DrawLine(Offset, JointTransform.GetLocation(), JointTarget, FLinearColor::Red, 1.2f);
 					}
 				}
 			}
diff --git a/Engine/Plugins/Experimental/FullBodyIK/Source/PBIK/Private/RigUnit_PBIK.cpp b/Engine/Plugins/Experimental/FullBodyIK/Source/PBIK/Private/RigUnit_PBIK.cpp
--- a/Engine/Plugins/Experimental/FullBodyIK/Source/PBIK/Private/RigUnit_PBIK.cpp
+++ b/Engine/Plugins/Experimental/FullBodyIK/Source/PBIK/Private/RigUnit_PBIK.cpp
@@ -182,6 +182,6 @@ FRigUnit_PBIK_Execute()
 	}
 
 	// do all debug drawing
-	Debug.Draw(ExecuteContext.UnitContext.DrawInterface, &Solver);
+	Debug.Draw(ExecuteContext.GetDrawInterface(), &Solver);
 }
 
diff --git a/Engine/Plugins/Experimental/FullBodyIK/Source/PBIK/Public/RigUnit_PBIK.h b/Engine/Plugins/Experimental/FullBodyIK/Source/PBIK/Public/RigUnit_PBIK.h
--- a/Engine/Plugins/Experimental/FullBodyIK/Source/PBIK/Public/RigUnit_PBIK.h
+++ b/Engine/Plugins/Experimental/FullBodyIK/Source/PBIK/Public/RigUnit_PBIK.h
@@ -3,7 +3,6 @@
 #pragma once
 
 #include "Units/Highlevel/RigUnit_HighlevelBase.h"
-#include "Drawing/ControlRigDrawInterface.h"
 
 #include "Core/PBIKSolver.h"
 #include "Core/PBIKDebug.h"
@@ -27,7 +26,7 @@ struct FPBIKDebug
 	UPROPERTY(EditAnywhere, Category = Debug)
 	bool bDrawDebug = false;
 
-	void Draw(FControlRigDrawInterface* DrawInterface, FPBIKSolver* Solver) const
+	void Draw(FRigVMDrawInterface* DrawInterface, FPBIKSolver* Solver) const
 	{
 		if (!(DrawInterface && Solver && bDrawDebug))
 		{
diff --git a/Engine/Plugins/Experimental/LiveLinkControlRig/Source/LiveLinkControlRig/Private/LiveLinkRigUnits.cpp b/Engine/Plugins/Experimental/LiveLinkControlRig/Source/LiveLinkControlRig/Private/LiveLinkRigUnits.cpp
--- a/Engine/Plugins/Experimental/LiveLinkControlRig/Source/LiveLinkControlRig/Private/LiveLinkRigUnits.cpp
+++ b/Engine/Plugins/Experimental/LiveLinkControlRig/Source/LiveLinkControlRig/Private/LiveLinkRigUnits.cpp
@@ -40,7 +40,7 @@ FRigUnit_LiveLinkEvaluteFrameAnimation_Execute()
 	}
 
 	// Debug Drawing if required
-	if (bDrawDebug && ExecuteContext.UnitContext.DrawInterface != nullptr)
+	if (bDrawDebug && ExecuteContext.GetDrawInterface() != nullptr)
 	{
 		/*
 		FRigBoneHierarchy DebugHierarchy;
@@ -120,7 +120,7 @@ FRigUnit_LiveLinkEvaluteFrameTransform_Execute()
 	}
 
 	// Debug Drawing if required
-	if (bDrawDebug && ExecuteContext.UnitContext.DrawInterface != nullptr)
+	if (bDrawDebug && ExecuteContext.GetDrawInterface() != nullptr)
 	{
 		// Rather than drawing as an axis guide, draw as a hierarchy with a dummy root bone which makes it clearer what the origin of the transform
 		// space is
diff --git a/Engine/Source/Runtime/RigVM/Private/RigVMCore/RigVMDrawContainer.cpp b/Engine/Source/Runtime/RigVM/Private/RigVMCore/RigVMDrawContainer.cpp
new file mode 100644
--- /dev/null
+++ b/Engine/Source/Runtime/RigVM/Private/RigVMCore/RigVMDrawContainer.cpp
@@ -0,0 +1,23 @@
+// Copyright Epic Games, Inc. All Rights Reserved.
+
+#include "RigVMCore/RigVMDrawContainer.h"
+
+#include UE_INLINE_GENERATED_CPP_BY_NAME(RigVMDrawContainer)
+
+int32 FRigVMDrawContainer::GetIndex(const FName& InName) const
+{
+	for (int32 Index = 0; Index < Instructions.Num(); ++Index)
+	{
+		if (Instructions[Index].Name == InName)
+		{
+			return Index;
+		}
+	}
+	return INDEX_NONE;
+}
+
+void FRigVMDrawContainer::Reset()
+{
+	Instructions.Reset();
+}
+
diff --git a/Engine/Source/Runtime/RigVM/Private/RigVMCore/RigVMDrawInstruction.cpp b/Engine/Source/Runtime/RigVM/Private/RigVMCore/RigVMDrawInstruction.cpp
new file mode 100644
--- /dev/null
+++ b/Engine/Source/Runtime/RigVM/Private/RigVMCore/RigVMDrawInstruction.cpp
@@ -0,0 +1,6 @@
+// Copyright Epic Games, Inc. All Rights Reserved.
+
+#include "RigVMCore/RigVMDrawInstruction.h"
+
+#include UE_INLINE_GENERATED_CPP_BY_NAME(RigVMDrawInstruction)
+
diff --git a/Engine/Source/Runtime/RigVM/Private/RigVMCore/RigVMDrawInterface.cpp b/Engine/Source/Runtime/RigVM/Private/RigVMCore/RigVMDrawInterface.cpp
new file mode 100644
--- /dev/null
+++ b/Engine/Source/Runtime/RigVM/Private/RigVMCore/RigVMDrawInterface.cpp
@@ -0,0 +1,323 @@
+// Copyright Epic Games, Inc. All Rights Reserved.
+
+#include "RigVMCore/RigVMDrawInterface.h"
+#include "DynamicMeshBuilder.h"
+#include "SceneManagement.h" 
+
+#include UE_INLINE_GENERATED_CPP_BY_NAME(RigVMDrawInterface)
+
+TAutoConsoleVariable<int32> CVarEnableRigVMDrawInterfaceInShipping(TEXT("RigVM.EnableDrawInterfaceInShipping"), 0, TEXT("Set to 1 to enable control rig draw interface in shipping"));
+
+void FRigVMDrawInterface::DrawInstruction(const FRigVMDrawInstruction& InInstruction)
+{
+	if (!IsEnabled())
+	{
+		return;
+	}
+
+	Instructions.Add(InInstruction);
+}
+
+void FRigVMDrawInterface::DrawPoint(const FTransform& WorldOffset, const FVector& Position, float Size, const FLinearColor& Color)
+{
+	if (!IsEnabled())
+	{
+		return;
+	}
+
+	FRigVMDrawInstruction Instruction(ERigVMDrawSettings::Points, Color, Size, WorldOffset);
+	Instruction.Positions.Add(Position);
+	Instructions.Add(Instruction);
+}
+
+void FRigVMDrawInterface::DrawPoints(const FTransform& WorldOffset, const TArrayView<const FVector>& Points, float Size, const FLinearColor& Color)
+{
+	if (!IsEnabled())
+	{
+		return;
+	}
+
+	FRigVMDrawInstruction Instruction(ERigVMDrawSettings::Points, Color, Size, WorldOffset);
+	Instruction.Positions.Append(Points.GetData(), Points.Num());
+	Instructions.Add(Instruction);
+}
+
+void FRigVMDrawInterface::DrawLine(const FTransform& WorldOffset, const FVector& LineStart, const FVector& LineEnd, const FLinearColor& Color, float Thickness)
+{
+	if (!IsEnabled())
+	{
+		return;
+	}
+
+	FRigVMDrawInstruction Instruction(ERigVMDrawSettings::Lines, Color, Thickness, WorldOffset);
+	Instruction.Positions.Add(LineStart);
+	Instruction.Positions.Add(LineEnd);
+	Instructions.Add(Instruction);
+}
+
+void FRigVMDrawInterface::DrawLines(const FTransform& WorldOffset, const TArrayView<const FVector>& Positions, const FLinearColor& Color, float Thickness)
+{
+	if (!IsEnabled())
+	{
+		return;
+	}
+
+	FRigVMDrawInstruction Instruction(ERigVMDrawSettings::Lines, Color, Thickness, WorldOffset);
+	Instruction.Positions.Append(Positions.GetData(), Positions.Num());
+	Instructions.Add(Instruction);
+}
+
+void FRigVMDrawInterface::DrawLineStrip(const FTransform& WorldOffset, const TArrayView<const FVector>& Positions, const FLinearColor& Color, float Thickness)
+{
+	if (!IsEnabled())
+	{
+		return;
+	}
+
+	FRigVMDrawInstruction Instruction(ERigVMDrawSettings::LineStrip, Color, Thickness, WorldOffset);
+	Instruction.Positions.Append(Positions.GetData(), Positions.Num());
+	Instructions.Add(Instruction);
+}
+
+void FRigVMDrawInterface::DrawBox(const FTransform& WorldOffset, const FTransform& Transform, const FLinearColor& Color, float Thickness)
+{
+	if (!IsEnabled())
+	{
+		return;
+	}
+
+	FTransform DrawTransform = Transform * WorldOffset;
+
+	FRigVMDrawInstruction Instruction(ERigVMDrawSettings::Lines, Color, Thickness, DrawTransform);
+
+	Instruction.Positions.Add(FVector(0.5f, 0.5f, 0.5f));
+	Instruction.Positions.Add(FVector(0.5f, -0.5f, 0.5f));
+	Instruction.Positions.Add(FVector(0.5f, -0.5f, 0.5f));
+	Instruction.Positions.Add(FVector(-0.5f, -0.5f, 0.5f));
+	Instruction.Positions.Add(FVector(-0.5f, -0.5f, 0.5f));
+	Instruction.Positions.Add(FVector(-0.5f, 0.5f, 0.5f));
+	Instruction.Positions.Add(FVector(-0.5f, 0.5f, 0.5f));
+	Instruction.Positions.Add(FVector(0.5f, 0.5f, 0.5f));
+
+	Instruction.Positions.Add(FVector(0.5f, 0.5f, -0.5f));
+	Instruction.Positions.Add(FVector(0.5f, -0.5f, -0.5f));
+	Instruction.Positions.Add(FVector(0.5f, -0.5f, -0.5f));
+	Instruction.Positions.Add(FVector(-0.5f, -0.5f, -0.5f));
+	Instruction.Positions.Add(FVector(-0.5f, -0.5f, -0.5f));
+	Instruction.Positions.Add(FVector(-0.5f, 0.5f, -0.5f));
+	Instruction.Positions.Add(FVector(-0.5f, 0.5f, -0.5f));
+	Instruction.Positions.Add(FVector(0.5f, 0.5f, -0.5f));
+
+	Instruction.Positions.Add(FVector(0.5f, 0.5f, 0.5f));
+	Instruction.Positions.Add(FVector(0.5f, 0.5f, -0.5f));
+	Instruction.Positions.Add(FVector(0.5f, -0.5f, 0.5f));
+	Instruction.Positions.Add(FVector(0.5f, -0.5f, -0.5f));
+	Instruction.Positions.Add(FVector(-0.5f, -0.5f, 0.5f));
+	Instruction.Positions.Add(FVector(-0.5f, -0.5f, -0.5f));
+	Instruction.Positions.Add(FVector(-0.5f, 0.5f, 0.5f));
+	Instruction.Positions.Add(FVector(-0.5f, 0.5f, -0.5f));
+
+	Instructions.Add(Instruction);
+}
+
+void FRigVMDrawInterface::DrawAxes(const FTransform& WorldOffset, const FTransform& Transform, float Size, float Thickness)
+{
+	if (!IsEnabled())
+	{
+		return;
+	}
+
+	DrawLine(WorldOffset, Transform.GetLocation(), Transform.TransformPosition(FVector(Size, 0.f, 0.f)), FLinearColor::Red, Thickness);
+	DrawLine(WorldOffset, Transform.GetLocation(), Transform.TransformPosition(FVector(0.f, Size, 0.f)), FLinearColor::Green, Thickness);
+	DrawLine(WorldOffset, Transform.GetLocation(), Transform.TransformPosition(FVector(0.f, 0.f, Size)), FLinearColor::Blue, Thickness);
+}
+
+void FRigVMDrawInterface::DrawAxes(const FTransform& WorldOffset, const TArrayView<const FTransform>& Transforms, float Size, float Thickness)
+{
+	if (!IsEnabled())
+	{
+		return;
+	}
+
+	if (Transforms.Num() == 0)
+	{
+		return;
+	}
+
+	FRigVMDrawInstruction InstructionX(ERigVMDrawSettings::Lines, FLinearColor::Red, Thickness, WorldOffset);
+	FRigVMDrawInstruction InstructionY(ERigVMDrawSettings::Lines, FLinearColor::Green, Thickness, WorldOffset);
+	FRigVMDrawInstruction InstructionZ(ERigVMDrawSettings::Lines, FLinearColor::Blue, Thickness, WorldOffset);
+
+	InstructionX.Positions.Reserve(Transforms.Num() * 2);
+	InstructionY.Positions.Reserve(Transforms.Num() * 2);
+	InstructionZ.Positions.Reserve(Transforms.Num() * 2);
+
+	for (const FTransform& Transform : Transforms)
+	{
+		InstructionX.Positions.Add(Transform.GetLocation());
+		InstructionX.Positions.Add(Transform.TransformPosition(FVector(Size, 0.f, 0.f)));
+		InstructionY.Positions.Add(Transform.GetLocation());
+		InstructionY.Positions.Add(Transform.TransformPosition(FVector(0.f, Size, 0.f)));
+		InstructionZ.Positions.Add(Transform.GetLocation());
+		InstructionZ.Positions.Add(Transform.TransformPosition(FVector(0.f, 0.f, Size)));
+	}
+
+	Instructions.Add(InstructionX);
+	Instructions.Add(InstructionY);
+	Instructions.Add(InstructionZ);
+}
+
+void FRigVMDrawInterface::DrawAxes(const FTransform& WorldOffset, const FTransform& Transform,
+	const FLinearColor& InColor, float Size, float Thickness)
+{
+	if (!IsEnabled())
+	{
+		return;
+	}
+
+	FRigVMDrawInstruction Instruction(ERigVMDrawSettings::Lines, InColor, Thickness, WorldOffset);
+
+	Instruction.Positions.Reserve(6);
+	Instruction.Positions.Add(Transform.GetLocation());
+	Instruction.Positions.Add(Transform.TransformPosition(FVector(Size, 0.f, 0.f)));
+	Instruction.Positions.Add(Transform.GetLocation());
+	Instruction.Positions.Add(Transform.TransformPosition(FVector(0.f, Size, 0.f)));
+	Instruction.Positions.Add(Transform.GetLocation());
+	Instruction.Positions.Add(Transform.TransformPosition(FVector(0.f, 0.f, Size)));
+
+	Instructions.Add(Instruction);
+}
+
+void FRigVMDrawInterface::DrawAxes(const FTransform& WorldOffset, const TArrayView<const FTransform>& Transforms,
+	const FLinearColor& InColor, float Size, float Thickness)
+{
+	if (!IsEnabled())
+	{
+		return;
+	}
+
+	if (Transforms.Num() == 0)
+	{
+		return;
+	}
+
+	FRigVMDrawInstruction Instruction(ERigVMDrawSettings::Lines, InColor, Thickness, WorldOffset);
+
+	Instruction.Positions.Reserve(Transforms.Num() * 6);
+
+	for (const FTransform& Transform : Transforms)
+	{
+		Instruction.Positions.Add(Transform.GetLocation());
+		Instruction.Positions.Add(Transform.TransformPosition(FVector(Size, 0.f, 0.f)));
+		Instruction.Positions.Add(Transform.GetLocation());
+		Instruction.Positions.Add(Transform.TransformPosition(FVector(0.f, Size, 0.f)));
+		Instruction.Positions.Add(Transform.GetLocation());
+		Instruction.Positions.Add(Transform.TransformPosition(FVector(0.f, 0.f, Size)));
+	}
+
+	Instructions.Add(Instruction);
+}
+
+void FRigVMDrawInterface::DrawRectangle(const FTransform& WorldOffset, const FTransform& Transform, float Size, const FLinearColor& Color, float Thickness)
+{
+	if (!IsEnabled())
+	{
+		return;
+	}
+
+	FTransform DrawTransform = Transform * WorldOffset;
+
+	FRigVMDrawInstruction Instruction(ERigVMDrawSettings::LineStrip, Color, Thickness, DrawTransform);
+
+	float Extent = Size * 0.5f;
+	Instruction.Positions.Reserve(5);
+	Instruction.Positions.Add(FVector(-Extent, -Extent, 0.f));
+	Instruction.Positions.Add(FVector(-Extent, Extent, 0.f));
+	Instruction.Positions.Add(FVector(Extent, Extent, 0.f));
+	Instruction.Positions.Add(FVector(Extent, -Extent, 0.f));
+	Instruction.Positions.Add(FVector(-Extent, -Extent, 0.f));
+
+	Instructions.Add(Instruction);
+}
+
+void FRigVMDrawInterface::DrawArc(const FTransform& WorldOffset, const FTransform& Transform, float Radius, float MinimumAngle, float MaximumAngle, const FLinearColor& Color, float Thickness, int32 Detail)
+{
+	if (!IsEnabled())
+	{
+		return;
+	}
+
+	int32 Count = FMath::Clamp<int32>(Detail, 4, 32);
+	
+	FTransform DrawTransform = Transform * WorldOffset;
+
+	FRigVMDrawInstruction Instruction(ERigVMDrawSettings::LineStrip, Color, Thickness, DrawTransform);
+	Instruction.Positions.Reserve(Count);
+
+	FVector V = FVector(Radius, 0.f, 0.f);
+	FQuat Rotation(FVector(0.f, 0.f, 1.f), MinimumAngle);
+	V = Rotation.RotateVector(V);
+	Instruction.Positions.Add(V);
+	float StepAngle = (MaximumAngle - MinimumAngle) / float(Count);
+	if (FMath::Abs<float>(MaximumAngle - MinimumAngle) >= PI * 2.f - SMALL_NUMBER)
+	{
+		StepAngle = (PI * 2.f) / float(Count);
+		Count++;
+	}
+	Rotation = FQuat(FVector(0.f, 0.f, 1.f), StepAngle);
+	for(int32 Index=1;Index<Count;Index++)
+	{
+		V = Rotation.RotateVector(V);
+		Instruction.Positions.Add(V);
+	}
+
+	Instructions.Add(Instruction);
+}
+
+void FRigVMDrawInterface::DrawCone(const FTransform& WorldOffset, const FTransform& ConeOffset, float Angle1, float Angle2, uint32 NumSides, bool bDrawSideLines, const FLinearColor& SideLineColor, FMaterialRenderProxy* const MaterialRenderProxy)
+{
+	if (!IsEnabled())
+	{
+		return;
+	}
+
+	TArray<FDynamicMeshVertex> MeshVerts;
+	TArray<uint32> MeshIndices;
+	BuildConeVerts(Angle1, Angle2, 1.f, 0.f, NumSides, MeshVerts, MeshIndices);
+
+	FTransform ConeTransform = ConeOffset * WorldOffset;
+	FRigVMDrawInstruction MeshInstruction;
+	MeshInstruction.PrimitiveType = ERigVMDrawSettings::DynamicMesh;
+	MeshInstruction.Transform = ConeTransform;
+	MeshInstruction.MeshVerts = MeshVerts;
+	MeshInstruction.MeshIndices = MeshIndices;
+	MeshInstruction.MaterialRenderProxy = MaterialRenderProxy;
+	Instructions.Add(MeshInstruction);
+
+	if (bDrawSideLines)
+	{
+		TArray<FVector> ConeVerts;
+		ConeVerts.AddUninitialized(NumSides);
+
+		// Draw lines down major directions
+		for (int32 i = 0; i < 4; i++)
+		{
+			float Fraction = (float)i / (float)(4);
+			float Azi = 2.f * PI * Fraction;
+			FVector ConeVert = CalcConeVert(Angle1, Angle2, Azi);
+			DrawLine(WorldOffset, ConeOffset.GetLocation(), ConeOffset.TransformPosition(ConeVert), SideLineColor, 1.f);
+		}
+	}
+}
+
+bool FRigVMDrawInterface::IsEnabled() const
+{
+	bool bIsEnabled = true;
+
+#if UE_BUILD_SHIPPING || UE_BUILD_TEST 
+	bIsEnabled = CVarEnableRigVMDrawInterfaceInShipping.GetValueOnAnyThread() == 1; 
+#endif 
+
+	return bIsEnabled; 
+}
+
diff --git a/Engine/Source/Runtime/RigVM/Private/RigVMCore/RigVMLog.cpp b/Engine/Source/Runtime/RigVM/Private/RigVMCore/RigVMLog.cpp
new file mode 100644
--- /dev/null
+++ b/Engine/Source/Runtime/RigVM/Private/RigVMCore/RigVMLog.cpp
@@ -0,0 +1,42 @@
+// Copyright Epic Games, Inc. All Rights Reserved.
+
+#include "RigVMCore/RigVMLog.h"
+
+#if WITH_EDITOR
+TArray<FRigVMLog::FLogEntry> FRigVMLog::GetEntries(EMessageSeverity::Type InSeverity, bool bIncludeHigherSeverity)
+{
+	return Entries.FilterByPredicate([InSeverity, bIncludeHigherSeverity](const FRigVMLog::FLogEntry& Entry)
+	{
+		if (bIncludeHigherSeverity)
+		{
+			return Entry.Severity <= InSeverity;
+		}
+		return Entry.Severity == InSeverity;
+	});
+}
+#endif
+
+void FRigVMLog::Reset()
+{
+#if WITH_EDITOR
+	Entries.Reset();
+	KnownMessages.Reset();
+#endif
+}
+
+void FRigVMLog::Report(EMessageSeverity::Type InSeverity, const FName& InOperatorName, int32 InInstructionIndex, const FString& InMessage)
+{
+#if WITH_EDITOR
+
+	if (!InMessage.IsEmpty())
+	{
+		if (KnownMessages.Contains(InMessage))
+		{
+			return;
+		}
+	}
+
+	int32 EntryIndex = Entries.Add(FLogEntry(InSeverity, InOperatorName, InInstructionIndex, InMessage));
+	KnownMessages.Add(InMessage, true);
+#endif
+}
diff --git a/Engine/Source/Runtime/RigVM/Public/RigVMCore/RigVMDrawContainer.h b/Engine/Source/Runtime/RigVM/Public/RigVMCore/RigVMDrawContainer.h
new file mode 100644
--- /dev/null
+++ b/Engine/Source/Runtime/RigVM/Public/RigVMCore/RigVMDrawContainer.h
@@ -0,0 +1,33 @@
+// Copyright Epic Games, Inc. All Rights Reserved.
+
+#pragma once
+
+#include "CoreMinimal.h"
+#include "RigVMDrawInstruction.h"
+#include "RigVMDrawContainer.generated.h"
+
+USTRUCT()
+struct RIGVM_API FRigVMDrawContainer
+{
+	GENERATED_BODY()
+	virtual ~FRigVMDrawContainer() {}
+
+	int32 Num() const { return Instructions.Num(); }
+	int32 GetIndex(const FName& InName) const;
+	const FRigVMDrawInstruction& operator[](int32 InIndex) const { return Instructions[InIndex]; }
+	FRigVMDrawInstruction& operator[](int32 InIndex) { return Instructions[InIndex]; }
+	const FRigVMDrawInstruction& operator[](const FName& InName) const { return Instructions[GetIndex(InName)]; }
+	FRigVMDrawInstruction& operator[](const FName& InName) { return Instructions[GetIndex(InName)]; }
+
+	TArray<FRigVMDrawInstruction>::RangedForIteratorType      begin() { return Instructions.begin(); }
+	TArray<FRigVMDrawInstruction>::RangedForConstIteratorType begin() const { return Instructions.begin(); }
+	TArray<FRigVMDrawInstruction>::RangedForIteratorType      end() { return Instructions.end(); }
+	TArray<FRigVMDrawInstruction>::RangedForConstIteratorType end() const { return Instructions.end(); }
+
+	uint32 GetAllocatedSize(void) const { return Instructions.GetAllocatedSize(); }
+
+	void Reset();
+
+	UPROPERTY(EditAnywhere, Category = "DrawContainer")
+	TArray<FRigVMDrawInstruction> Instructions;
+};
diff --git a/Engine/Source/Runtime/RigVM/Public/RigVMCore/RigVMDrawInstruction.h b/Engine/Source/Runtime/RigVM/Public/RigVMCore/RigVMDrawInstruction.h
new file mode 100644
--- /dev/null
+++ b/Engine/Source/Runtime/RigVM/Public/RigVMCore/RigVMDrawInstruction.h
@@ -0,0 +1,78 @@
+// Copyright Epic Games, Inc. All Rights Reserved.
+
+#pragma once
+
+#include "CoreMinimal.h"
+#include "DynamicMeshBuilder.h"
+#include "RigVMDrawInstruction.generated.h"
+
+class FMaterialRenderProxy;
+
+UENUM()
+namespace ERigVMDrawSettings
+{
+	enum Primitive
+	{
+		Points,
+		Lines,
+		LineStrip, 
+		DynamicMesh
+	};
+}
+
+USTRUCT()
+struct RIGVM_API FRigVMDrawInstruction
+{
+	GENERATED_BODY()
+
+	UPROPERTY(EditAnywhere, Category = "DrawInstruction")
+	FName Name;
+
+	UPROPERTY(EditAnywhere, Category = "DrawInstruction")
+	TEnumAsByte<ERigVMDrawSettings::Primitive> PrimitiveType;
+
+	UPROPERTY(EditAnywhere, Category = "DrawInstruction")
+	TArray<FVector> Positions;
+
+	UPROPERTY(EditAnywhere, Category = "DrawInstruction")
+	FLinearColor Color;
+
+	UPROPERTY(EditAnywhere, Category = "DrawInstruction")
+	float Thickness;
+
+	UPROPERTY(EditAnywhere, Category = "DrawInstruction")
+	FTransform Transform;
+
+	// This is to draw cone, they're not UPROPERTY
+	TArray<FDynamicMeshVertex> MeshVerts;
+	TArray<uint32> MeshIndices;
+	FMaterialRenderProxy* MaterialRenderProxy = nullptr;
+
+	FRigVMDrawInstruction()
+		: Name(NAME_None)
+		, PrimitiveType(ERigVMDrawSettings::Points)
+		, Color(FLinearColor::Red)
+		, Thickness(0.f)
+		, Transform(FTransform::Identity)
+	{}
+
+	FRigVMDrawInstruction(ERigVMDrawSettings::Primitive InPrimitiveType, const FLinearColor& InColor, float InThickness = 0.f, FTransform InTransform = FTransform::Identity)
+		: Name(NAME_None)
+		, PrimitiveType(InPrimitiveType)
+		, Color(InColor)
+		, Thickness(InThickness)
+		, Transform(InTransform)
+	{}
+
+	bool IsValid() const
+	{
+		// if primitive type is dynamicmesh, we expect these data to be there. 
+		// otherwise, we can't draw
+		if (PrimitiveType == ERigVMDrawSettings::DynamicMesh)
+		{
+			return MeshVerts.Num() != 0 && MeshIndices.Num() != 0 && MaterialRenderProxy != nullptr;
+		}
+
+		return Positions.Num() != 0;
+	}
+};
diff --git a/Engine/Source/Runtime/RigVM/Public/RigVMCore/RigVMDrawInterface.h b/Engine/Source/Runtime/RigVM/Public/RigVMCore/RigVMDrawInterface.h
new file mode 100644
--- /dev/null
+++ b/Engine/Source/Runtime/RigVM/Public/RigVMCore/RigVMDrawInterface.h
@@ -0,0 +1,32 @@
+// Copyright Epic Games, Inc. All Rights Reserved.
+
+#pragma once
+
+#include "CoreMinimal.h"
+#include "RigVMDrawContainer.h"
+#include "RigVMDrawInterface.generated.h"
+
+USTRUCT()
+struct RIGVM_API FRigVMDrawInterface : public FRigVMDrawContainer
+{
+public:
+
+	GENERATED_BODY()
+
+	void DrawInstruction(const FRigVMDrawInstruction& InInstruction);
+	void DrawPoint(const FTransform& WorldOffset, const FVector& Position, float Size, const FLinearColor& Color);
+	void DrawPoints(const FTransform& WorldOffset, const TArrayView<const FVector>& Points, float Size, const FLinearColor& Color);
+	void DrawLine(const FTransform& WorldOffset, const FVector& LineStart, const FVector& LineEnd, const FLinearColor& Color, float Thickness = 0.f);
+	void DrawLines(const FTransform& WorldOffset, const TArrayView<const FVector>& Positions, const FLinearColor& Color, float Thickness = 0.f);
+	void DrawLineStrip(const FTransform& WorldOffset, const TArrayView<const FVector>& Positions, const FLinearColor& Color, float Thickness = 0.f);
+	void DrawBox(const FTransform& WorldOffset, const FTransform& Transform, const FLinearColor& Color, float Thickness = 0.f);
+	void DrawAxes(const FTransform& WorldOffset, const FTransform& Transform, float Size, float Thickness = 0.f);
+	void DrawAxes(const FTransform& WorldOffset, const TArrayView<const FTransform>& Transforms, float Size, float Thickness = 0.f);
+	void DrawAxes(const FTransform& WorldOffset, const FTransform& Transform, const FLinearColor& InColor, float Size, float Thickness = 0.f);
+	void DrawAxes(const FTransform& WorldOffset, const TArrayView<const FTransform>& Transforms, const FLinearColor& InColor, float Size, float Thickness = 0.f);
+	void DrawRectangle(const FTransform& WorldOffset, const FTransform& Transform, float Size, const FLinearColor& Color, float Thickness);
+	void DrawArc(const FTransform& WorldOffset, const FTransform& Transform, float Radius, float MinimumAngle, float MaximumAngle, const FLinearColor& Color, float Thickness, int32 Detail);
+	void DrawCone(const FTransform& WorldOffset, const FTransform& ConeOffset, float Angle1, float Angle2, uint32 NumSides, bool bDrawSideLines, const FLinearColor& SideLineColor, FMaterialRenderProxy* const MaterialRenderProxy);
+
+	bool IsEnabled() const;
+};
diff --git a/Engine/Source/Runtime/RigVM/Public/RigVMCore/RigVMExecuteContext.h b/Engine/Source/Runtime/RigVM/Public/RigVMCore/RigVMExecuteContext.h
--- a/Engine/Source/Runtime/RigVM/Public/RigVMCore/RigVMExecuteContext.h
+++ b/Engine/Source/Runtime/RigVM/Public/RigVMCore/RigVMExecuteContext.h
@@ -22,6 +22,9 @@
 #include "UObject/UnrealType.h"
 #include "RigVMCore/RigVMNameCache.h"
 #include "UObject/StructOnScope.h"
+#include "RigVMLog.h"
+#include "RigVMDrawInterface.h"
+#include "RigVMDrawContainer.h"
 
 #include "RigVMExecuteContext.generated.h"
 
@@ -226,8 +229,14 @@ struct RIGVM_API FRigVMExecuteContext
 		, InstructionIndex(0)
 		, DeltaTime(0.0)
 		, AbsoluteTime(0.0)
+		, FramesPerSecond(1.0 / 60.0)
 		, RuntimeSettings()
 		, NameCache(nullptr)
+#if WITH_EDITOR
+		, LogPtr(nullptr)
+#endif
+		, DrawInterfacePtr(nullptr)
+		, DrawContainerPtr(nullptr)
 	{
 	}
 
@@ -281,6 +290,18 @@ struct RIGVM_API FRigVMExecuteContext
 	
 	FRigVMNameCache* GetNameCache() const { return NameCache; }
 
+#if WITH_EDITOR
+	FRigVMLog* GetLog() const { return LogPtr; }
+	void SetLog(FRigVMLog* InLog) { LogPtr = InLog; }
+#endif
+
+	FRigVMDrawInterface* GetDrawInterface() const { return DrawInterfacePtr; }
+	void SetDrawInterface(FRigVMDrawInterface* InDrawInterface) { DrawInterfacePtr = InDrawInterface; }
+
+	const FRigVMDrawContainer* GetDrawContainer() const { return DrawContainerPtr; }
+	FRigVMDrawContainer* GetDrawContainer() { return DrawContainerPtr; }
+	void SetDrawContainer(FRigVMDrawContainer* InDrawContainer) { DrawContainerPtr = InDrawContainer; }
+
 	virtual void Initialize()
 	{
 		if(NameCache == nullptr)
@@ -322,6 +343,12 @@ protected:
 
 	mutable FRigVMNameCache* NameCache;
 
+#if WITH_EDITOR
+	FRigVMLog* LogPtr;
+#endif
+	FRigVMDrawInterface* DrawInterfacePtr;
+	FRigVMDrawContainer* DrawContainerPtr;
+
 #if UE_RIGVM_DEBUG_EXECUTION
 	FString DebugMemoryString;
 
diff --git a/Engine/Source/Runtime/RigVM/Public/RigVMCore/RigVMLog.h b/Engine/Source/Runtime/RigVM/Public/RigVMCore/RigVMLog.h
new file mode 100644
--- /dev/null
+++ b/Engine/Source/Runtime/RigVM/Public/RigVMCore/RigVMLog.h
@@ -0,0 +1,38 @@
+// Copyright Epic Games, Inc. All Rights Reserved.
+
+#pragma once
+
+#include "CoreMinimal.h"
+#include "Logging/TokenizedMessage.h"
+
+struct RIGVM_API FRigVMLog
+{
+public:
+
+	FRigVMLog() {}
+	virtual ~FRigVMLog() {}
+
+#if WITH_EDITOR
+	struct FLogEntry
+	{
+		FLogEntry(EMessageSeverity::Type InSeverity, const FName& InFunctionName, int32 InInstructionIndex, const FString& InMessage)
+		: Severity(InSeverity)
+		, FunctionName(InFunctionName)
+		, InstructionIndex(InInstructionIndex)
+		, Message(InMessage)
+		{}
+		
+		EMessageSeverity::Type Severity;
+		FName FunctionName;
+		int32 InstructionIndex;
+		FString Message;
+	};
+	TArray<FLogEntry> Entries;
+	TMap<FString, bool> KnownMessages;
+
+	TArray<FLogEntry> GetEntries(EMessageSeverity::Type InSeverity = EMessageSeverity::Info, bool bIncludeHigherSeverity = true);
+#endif
+
+	virtual void Reset();
+	virtual void Report(EMessageSeverity::Type InSeverity, const FName& InFunctionName, int32 InInstructionIndex, const FString& InMessage);
+};
