RigVM: Rely on redirector for resolved functions

diff --git a/Engine/Source/Runtime/RigVM/Private/RigVMCore/RigVMRegistry.cpp b/Engine/Source/Runtime/RigVM/Private/RigVMCore/RigVMRegistry.cpp
--- a/Engine/Source/Runtime/RigVM/Private/RigVMCore/RigVMRegistry.cpp
+++ b/Engine/Source/Runtime/RigVM/Private/RigVMCore/RigVMRegistry.cpp
@@ -1105,15 +1105,15 @@ const FRigVMFunction* FRigVMRegistry::FindFunction(const TCHAR* InName) const
 
 	// Otherwise ask the associated dispatch factory for a function matching this signature.
 	const FString NameString(InName);
-	FString FactoryName, ArgumentsString;
-	if(NameString.Split(TEXT("::"), &FactoryName, &ArgumentsString))
+	FString StructOrFactoryName, SuffixString;
+	if(NameString.Split(TEXT("::"), &StructOrFactoryName, &SuffixString))
 	{
 		// if the factory has never been registered - FindDispatchFactory will try to look it up and register
-		if(const FRigVMDispatchFactory* Factory = FindDispatchFactory(*FactoryName))
+		if(const FRigVMDispatchFactory* Factory = FindDispatchFactory(*StructOrFactoryName))
 		{
 			if(const FRigVMTemplate* Template = Factory->GetTemplate())
 			{
-				const FRigVMTemplateTypeMap ArgumentTypes = Template->GetArgumentTypesFromString(ArgumentsString);
+				const FRigVMTemplateTypeMap ArgumentTypes = Template->GetArgumentTypesFromString(SuffixString);
 				if(ArgumentTypes.Num() == Template->NumArguments())
 				{
 					const int32 PermutationIndex = Template->FindPermutation(ArgumentTypes);
@@ -1127,6 +1127,47 @@ const FRigVMFunction* FRigVMRegistry::FindFunction(const TCHAR* InName) const
 		}
 	}
 
+#if WITH_EDITOR
+	// if we haven't been able to find the function - try to see if we can get the dispatch or rigvmstruct
+	// from a core redirect
+	if(!StructOrFactoryName.IsEmpty())
+	{
+		static const FString StructPrefix = TEXT("F");
+		const bool bIsDispatchFactory = StructOrFactoryName.StartsWith(FRigVMDispatchFactory::DispatchPrefix, ESearchCase::CaseSensitive);
+		if(bIsDispatchFactory)
+		{
+			StructOrFactoryName = StructOrFactoryName.Mid(FRigVMDispatchFactory::DispatchPrefix.Len());
+		}
+
+		if(StructOrFactoryName.StartsWith(StructPrefix, ESearchCase::CaseSensitive))
+		{
+			StructOrFactoryName = StructOrFactoryName.Mid(StructPrefix.Len());
+		}
+		
+		const FCoreRedirectObjectName OldObjectName(StructOrFactoryName);
+		TArray<const FCoreRedirect*> Redirects;
+		if(FCoreRedirects::GetMatchingRedirects(ECoreRedirectFlags::Type_Struct, OldObjectName, Redirects, ECoreRedirectMatchFlags::AllowPartialMatch))
+		{
+			for(const FCoreRedirect* Redirect : Redirects)
+			{
+				FString NewStructOrFactoryName = Redirect->NewName.ObjectName.ToString();
+				NewStructOrFactoryName = StructPrefix + NewStructOrFactoryName;
+				if(bIsDispatchFactory)
+				{
+					NewStructOrFactoryName = FRigVMDispatchFactory::DispatchPrefix + NewStructOrFactoryName;
+				}
+				const FRigVMFunction* RedirectedFunction = FindFunction(*(NewStructOrFactoryName + TEXT("::") + SuffixString));
+				if(RedirectedFunction)
+				{
+					FRigVMRegistry& MutableRegistry = FRigVMRegistry::Get();
+					MutableRegistry.FunctionNameToIndex.Add(InName, RedirectedFunction->Index);
+					return RedirectedFunction;
+				}
+			}
+		}
+	}
+#endif
+	
 	return nullptr;
 }
 
