Control Rig: Move remaining nodes to RigVM

diff --git a/Engine/Plugins/Animation/ControlRig/Config/BaseControlRig.ini b/Engine/Plugins/Animation/ControlRig/Config/BaseControlRig.ini
--- a/Engine/Plugins/Animation/ControlRig/Config/BaseControlRig.ini
+++ b/Engine/Plugins/Animation/ControlRig/Config/BaseControlRig.ini
@@ -380,4 +380,74 @@
 +StructRedirects=(OldName="/Script/ControlRig.RigUnit_NoiseVector",NewName="/Script/RigVM.RigVMFunction_NoiseVector")
 +StructRedirects=(OldName="/Script/ControlRig.RigUnit_NoiseVector2",NewName="/Script/RigVM.RigVMFunction_NoiseVector2")
 +StructRedirects=(OldName="/Script/ControlRig.RigUnit_RandomFloat",NewName="/Script/RigVM.RigVMFunction_RandomFloat")
-+StructRedirects=(OldName="/Script/ControlRig.RigUnit_RandomVector",NewName="/Script/RigVM.RigVMFunction_RandomVector")
\ No newline at end of file
++StructRedirects=(OldName="/Script/ControlRig.RigUnit_RandomVector",NewName="/Script/RigVM.RigVMFunction_RandomVector")
+
+; moving remaining units to RigVM runtime
++EnumRedirects=(OldName="ECRSimPointIntegrateType",NewName="/Script/RigVM.ERigVMSimPointIntegrateType")
++StructRedirects=(OldName="/Script/ControlRig.CRSimPoint",NewName="/Script/RigVM.RigVMSimPoint")
++StructRedirects=(OldName="/Script/ControlRig.RigUnit_ForLoopCount",NewName="/Script/RigVM.RigVMFunction_ForLoopCount")
++StructRedirects=(OldName="/Script/ControlRig.RigUnit_UserDefinedEvent",NewName="/Script/RigVM.RigVMFunction_UserDefinedEvent")
++StructRedirects=(OldName="/Script/ControlRig.RigUnit_AccumulateBase",NewName="/Script/RigVM.RigVMFunction_AccumulateBase")
++StructRedirects=(OldName="/Script/ControlRig.RigUnit_AccumulateFloatAdd",NewName="/Script/RigVM.RigVMFunction_AccumulateFloatAdd")
++StructRedirects=(OldName="/Script/ControlRig.RigUnit_AccumulateVectorAdd",NewName="/Script/RigVM.RigVMFunction_AccumulateVectorAdd")
++StructRedirects=(OldName="/Script/ControlRig.RigUnit_AccumulateFloatMul",NewName="/Script/RigVM.RigVMFunction_AccumulateFloatMul")
++StructRedirects=(OldName="/Script/ControlRig.RigUnit_AccumulateVectorMul",NewName="/Script/RigVM.RigVMFunction_AccumulateVectorMul")
++StructRedirects=(OldName="/Script/ControlRig.RigUnit_AccumulateQuatMul",NewName="/Script/RigVM.RigVMFunction_AccumulateQuatMul")
++StructRedirects=(OldName="/Script/ControlRig.RigUnit_AccumulateTransformMul",NewName="/Script/RigVM.RigVMFunction_AccumulateTransformMul")
++StructRedirects=(OldName="/Script/ControlRig.RigUnit_AccumulateFloatLerp",NewName="/Script/RigVM.RigVMFunction_AccumulateFloatLerp")
++StructRedirects=(OldName="/Script/ControlRig.RigUnit_AccumulateVectorLerp",NewName="/Script/RigVM.RigVMFunction_AccumulateVectorLerp")
++StructRedirects=(OldName="/Script/ControlRig.RigUnit_AccumulateQuatLerp",NewName="/Script/RigVM.RigVMFunction_AccumulateQuatLerp")
++StructRedirects=(OldName="/Script/ControlRig.RigUnit_AccumulateTransformLerp",NewName="/Script/RigVM.RigVMFunction_AccumulateTransformLerp")
++StructRedirects=(OldName="/Script/ControlRig.RigUnit_AccumulateFloatRange",NewName="/Script/RigVM.RigVMFunction_AccumulateFloatRange")
++StructRedirects=(OldName="/Script/ControlRig.RigUnit_AccumulateVectorRange",NewName="/Script/RigVM.RigVMFunction_AccumulateVectorRange")
++StructRedirects=(OldName="/Script/ControlRig.RigUnit_AlphaInterp",NewName="/Script/RigVM.RigVMFunction_AlphaInterp")
++StructRedirects=(OldName="/Script/ControlRig.RigUnit_AlphaInterpVector",NewName="/Script/RigVM.RigVMFunction_AlphaInterpVector")
++StructRedirects=(OldName="/Script/ControlRig.RigUnit_AlphaInterpQuat",NewName="/Script/RigVM.RigVMFunction_AlphaInterpQuat")
++StructRedirects=(OldName="/Script/ControlRig.RigUnit_DeltaFromPreviousFloat",NewName="/Script/RigVM.RigVMFunction_DeltaFromPreviousFloat")
++StructRedirects=(OldName="/Script/ControlRig.RigUnit_DeltaFromPreviousVector",NewName="/Script/RigVM.RigVMFunction_DeltaFromPreviousVector")
++StructRedirects=(OldName="/Script/ControlRig.RigUnit_DeltaFromPreviousQuat",NewName="/Script/RigVM.RigVMFunction_DeltaFromPreviousQuat")
++StructRedirects=(OldName="/Script/ControlRig.RigUnit_DeltaFromPreviousTransform",NewName="/Script/RigVM.RigVMFunction_DeltaFromPreviousTransform")
++StructRedirects=(OldName="/Script/ControlRig.RigUnit_KalmanFloat",NewName="/Script/RigVM.RigVMFunction_KalmanFloat")
++StructRedirects=(OldName="/Script/ControlRig.RigUnit_KalmanVector",NewName="/Script/RigVM.RigVMFunction_KalmanVector")
++StructRedirects=(OldName="/Script/ControlRig.RigUnit_KalmanTransform",NewName="/Script/RigVM.RigVMFunction_KalmanTransform")
++StructRedirects=(OldName="/Script/ControlRig.RigUnit_SimBase",NewName="/Script/RigVM.RigVMFunction_SimBase")
++StructRedirects=(OldName="/Script/ControlRig.RigUnit_SimBaseMutable",NewName="/Script/RigVM.RigVMFunction_SimBaseMutable")
++StructRedirects=(OldName="/Script/ControlRig.RigUnit_Timeline",NewName="/Script/RigVM.RigVMFunction_Timeline")
++StructRedirects=(OldName="/Script/ControlRig.RigUnit_TimeLoop",NewName="/Script/RigVM.RigVMFunction_TimeLoop")
++StructRedirects=(OldName="/Script/ControlRig.RigUnit_TimeOffsetFloat",NewName="/Script/RigVM.RigVMFunction_TimeOffsetFloat")
++StructRedirects=(OldName="/Script/ControlRig.RigUnit_TimeOffsetVector",NewName="/Script/RigVM.RigVMFunction_TimeOffsetVector")
++StructRedirects=(OldName="/Script/ControlRig.RigUnit_TimeOffsetTransform",NewName="/Script/RigVM.RigVMFunction_TimeOffsetTransform")
++StructRedirects=(OldName="/Script/ControlRig.RigUnit_VerletIntegrateVector",NewName="/Script/RigVM.RigVMFunction_VerletIntegrateVector")
++StructRedirects=(OldName="/Script/ControlRig.RigUnit_DebugBase",NewName="/Script/RigVM.RigVMFunction_DebugBase")
++StructRedirects=(OldName="/Script/ControlRig.RigUnit_DebugBaseMutable",NewName="/Script/RigVM.RigVMFunction_DebugBaseMutable")
++StructRedirects=(OldName="/Script/ControlRig.RigUnit_DebugLine",NewName="/Script/RigVM.RigVMFunction_DebugLine")
++StructRedirects=(OldName="/Script/ControlRig.RigUnit_DebugLineItemSpace",NewName="/Script/RigVM.RigVMFunction_DebugLineNoSpace")
++StructRedirects=(OldName="/Script/ControlRig.RigUnit_DebugLineStrip",NewName="/Script/RigVM.RigVMFunction_DebugLineStrip")
++StructRedirects=(OldName="/Script/ControlRig.RigUnit_DebugLineStripItemSpace",NewName="/Script/RigVM.RigVMFunction_DebugLineStripNoSpace")
++StructRedirects=(OldName="/Script/ControlRig.RigUnit_DebugPoint",NewName="/Script/RigVM.RigVMFunction_DebugPoint")
++StructRedirects=(OldName="/Script/ControlRig.RigUnit_DebugPointMutable",NewName="/Script/RigVM.RigVMFunction_DebugPointMutable")
++StructRedirects=(OldName="/Script/ControlRig.RigUnit_DebugRectangle",NewName="/Script/RigVM.RigVMFunction_DebugRectangle")
++StructRedirects=(OldName="/Script/ControlRig.RigUnit_DebugRectangleItemSpace",NewName="/Script/RigVM.RigVMFunction_DebugRectangleNoSpace")
++StructRedirects=(OldName="/Script/ControlRig.RigUnit_DebugArc",NewName="/Script/RigVM.RigVMFunction_DebugArc")
++StructRedirects=(OldName="/Script/ControlRig.RigUnit_DebugArcItemSpace",NewName="/Script/RigVM.RigVMFunction_DebugArcNoSpace")
++StructRedirects=(OldName="/Script/ControlRig.RigUnit_DebugTransform",NewName="/Script/RigVM.RigVMFunction_DebugTransform")
++StructRedirects=(OldName="/Script/ControlRig.RigUnit_DebugTransformMutable",NewName="/Script/RigVM.RigVMFunction_DebugTransformMutable")
++StructRedirects=(OldName="/Script/ControlRig.RigUnit_DebugTransformMutableItemSpace",NewName="/Script/RigVM.RigVMFunction_DebugTransformMutableNoSpace")
++StructRedirects=(OldName="/Script/ControlRig.RigUnit_DebugTransformArrayMutable_WorkData",NewName="/Script/RigVM.RigVMFunction_DebugTransformArrayMutable_WorkData")
++StructRedirects=(OldName="/Script/ControlRig.RigUnit_DebugTransformArrayMutable",NewName="/Script/RigVM.RigVMFunction_DebugTransformArrayMutable")
++StructRedirects=(OldName="/Script/ControlRig.RigUnit_DebugTransformArrayMutableItemSpace",NewName="/Script/RigVM.RigVMFunction_DebugTransformArrayMutableNoSpace")
++StructRedirects=(OldName="/Script/ControlRig.RigUnit_VisualDebugVector",NewName="/Script/RigVM.RigVMFunction_VisualDebugVector")
++StructRedirects=(OldName="/Script/ControlRig.RigUnit_VisualDebugVectorItemSpace",NewName="/Script/RigVM.RigVMFunction_VisualDebugVectorNoSpace")
++StructRedirects=(OldName="/Script/ControlRig.RigUnit_VisualDebugQuat",NewName="/Script/RigVM.RigVMFunction_VisualDebugQuat")
++StructRedirects=(OldName="/Script/ControlRig.RigUnit_VisualDebugQuatItemSpace",NewName="/Script/RigVM.RigVMFunction_VisualDebugQuatNoSpace")
++StructRedirects=(OldName="/Script/ControlRig.RigUnit_VisualDebugTransform",NewName="/Script/RigVM.RigVMFunction_VisualDebugTransform")
++StructRedirects=(OldName="/Script/ControlRig.RigUnit_VisualDebugTransformItemSpace",NewName="/Script/RigVM.RigVMFunction_VisualDebugTransformNoSpace")
++StructRedirects=(OldName="/Script/ControlRig.RigUnit_AnimBase",NewName="/Script/RigVM.RigVMFunction_AnimBase")
++StructRedirects=(OldName="/Script/ControlRig.RigUnit_AnimEasingType",NewName="/Script/RigVM.RigVMFunction_AnimEasingType")
++StructRedirects=(OldName="/Script/ControlRig.RigUnit_AnimEasing",NewName="/Script/RigVM.RigVMFunction_AnimEasing")
++StructRedirects=(OldName="/Script/ControlRig.RigUnit_AnimEvalRichCurve",NewName="/Script/RigVM.RigVMFunction_AnimEvalRichCurve")
++StructRedirects=(OldName="/Script/ControlRig.RigUnit_AnimRichCurve",NewName="/Script/RigVM.RigVMFunction_AnimRichCurve")
++StructRedirects=(OldName="/Script/ControlRig.RigUnit_GetDeltaTime",NewName="/Script/RigVM.RigVMFunction_GetDeltaTime")
++StructRedirects=(OldName="/Script/ControlRig.RigUnit_GetWorldTime",NewName="/Script/RigVM.RigVMFunction_GetWorldTime")
++StructRedirects=(OldName="/Script/ControlRig.RigUnit_FramesToSeconds",NewName="/Script/RigVM.RigVMFunction_FramesToSeconds")
++StructRedirects=(OldName="/Script/ControlRig.RigUnit_SecondsToFrames",NewName="/Script/RigVM.RigVMFunction_SecondsToFrames")
diff --git a/Engine/Plugins/Animation/ControlRig/Source/ControlRig/Private/ControlRig.cpp b/Engine/Plugins/Animation/ControlRig/Source/ControlRig/Private/ControlRig.cpp
--- a/Engine/Plugins/Animation/ControlRig/Source/ControlRig/Private/ControlRig.cpp
+++ b/Engine/Plugins/Animation/ControlRig/Source/ControlRig/Private/ControlRig.cpp
@@ -248,11 +248,11 @@ bool UControlRig::InitializeVM(const FName& InEventName)
 			const FRigVMExecuteContext& PublicContext = ExtendedExecuteContext.GetPublicData<>(); 
 			if(ControlRigLog)
 			{
-				ControlRigLog->Report(InSeverity,PublicContext.FunctionName,PublicContext.InstructionIndex, Message);
+				ControlRigLog->Report(InSeverity,PublicContext.GetFunctionName(),PublicContext.GetInstructionIndex(), Message);
 			}
 			else
 			{
-				LogOnce(InSeverity,PublicContext.InstructionIndex, Message);
+				LogOnce(InSeverity,PublicContext.GetInstructionIndex(), Message);
 			}
 		};
 	}
@@ -1573,11 +1573,11 @@ void UControlRig::UpdateVMSettings()
 
 			if(ControlRigLog)
 			{
-				ControlRigLog->Report(InSeverity, InContext->FunctionName, InContext->InstructionIndex, Message);
+				ControlRigLog->Report(InSeverity, InContext->GetFunctionName(), InContext->GetInstructionIndex(), Message);
 			}
 			else
 			{
-				LogOnce(InSeverity, InContext->InstructionIndex, Message);
+				LogOnce(InSeverity, InContext->GetInstructionIndex(), Message);
 			}
 		});
 #endif
diff --git a/Engine/Plugins/Animation/ControlRig/Source/ControlRig/Private/Math/Simulation/CRSimLinearSpring.cpp b/Engine/Plugins/Animation/ControlRig/Source/ControlRig/Private/Math/Simulation/CRSimLinearSpring.cpp
--- a/Engine/Plugins/Animation/ControlRig/Source/ControlRig/Private/Math/Simulation/CRSimLinearSpring.cpp
+++ b/Engine/Plugins/Animation/ControlRig/Source/ControlRig/Private/Math/Simulation/CRSimLinearSpring.cpp
@@ -5,7 +5,7 @@
 
 #include UE_INLINE_GENERATED_CPP_BY_NAME(CRSimLinearSpring)
 
-void FCRSimLinearSpring::CalculateForPoints(const FCRSimPoint& InPointA, const FCRSimPoint& InPointB, FVector& ForceA, FVector& ForceB) const
+void FCRSimLinearSpring::CalculateForPoints(const FRigVMSimPoint& InPointA, const FRigVMSimPoint& InPointB, FVector& ForceA, FVector& ForceB) const
 {
 	ForceA = ForceB = FVector::ZeroVector;
 
diff --git a/Engine/Plugins/Animation/ControlRig/Source/ControlRig/Private/Math/Simulation/CRSimPoint.cpp b/Engine/Plugins/Animation/ControlRig/Source/ControlRig/Private/Math/Simulation/CRSimPoint.cpp
deleted file mode 100644
diff --git a/Engine/Plugins/Animation/ControlRig/Source/ControlRig/Private/Math/Simulation/CRSimPointConstraint.cpp b/Engine/Plugins/Animation/ControlRig/Source/ControlRig/Private/Math/Simulation/CRSimPointConstraint.cpp
--- a/Engine/Plugins/Animation/ControlRig/Source/ControlRig/Private/Math/Simulation/CRSimPointConstraint.cpp
+++ b/Engine/Plugins/Animation/ControlRig/Source/ControlRig/Private/Math/Simulation/CRSimPointConstraint.cpp
@@ -5,7 +5,7 @@
 
 #include UE_INLINE_GENERATED_CPP_BY_NAME(CRSimPointConstraint)
 
-void FCRSimPointConstraint::Apply(FCRSimPoint& OutPointA, FCRSimPoint& OutPointB) const
+void FCRSimPointConstraint::Apply(FRigVMSimPoint& OutPointA, FRigVMSimPoint& OutPointB) const
 {
 	switch (Type)
 	{
diff --git a/Engine/Plugins/Animation/ControlRig/Source/ControlRig/Private/Math/Simulation/CRSimPointContainer.cpp b/Engine/Plugins/Animation/ControlRig/Source/ControlRig/Private/Math/Simulation/CRSimPointContainer.cpp
--- a/Engine/Plugins/Animation/ControlRig/Source/ControlRig/Private/Math/Simulation/CRSimPointContainer.cpp
+++ b/Engine/Plugins/Animation/ControlRig/Source/ControlRig/Private/Math/Simulation/CRSimPointContainer.cpp
@@ -17,13 +17,13 @@ void FCRSimPointContainer::ResetTime()
 {
 	FCRSimContainer::ResetTime();
 	PreviousStep.Reset();
-	for (FCRSimPoint& Point : Points)
+	for (FRigVMSimPoint& Point : Points)
 	{
 		Point.LinearVelocity = FVector::ZeroVector;
 	}
 }
 
-FCRSimPoint FCRSimPointContainer::GetPointInterpolated(int32 InIndex) const
+FRigVMSimPoint FCRSimPointContainer::GetPointInterpolated(int32 InIndex) const
 {
 	if (TimeLeftForStep <= SMALL_NUMBER || PreviousStep.Num() != Points.Num())
 	{
@@ -31,8 +31,8 @@ FCRSimPoint FCRSimPointContainer::GetPointInterpolated(int32 InIndex) const
 	}
 
 	float T = 1.f - TimeLeftForStep / TimeStep;
-	const FCRSimPoint& PrevPoint = PreviousStep[InIndex];
-	FCRSimPoint Point = Points[InIndex];
+	const FRigVMSimPoint& PrevPoint = PreviousStep[InIndex];
+	FRigVMSimPoint Point = Points[InIndex];
 	Point.Position = FMath::Lerp<FVector>(PrevPoint.Position, Point.Position, T);
 	Point.Size = FMath::Lerp<float>(PrevPoint.Size, Point.Size, T);
 	Point.LinearVelocity = FMath::Lerp<FVector>(PrevPoint.LinearVelocity, Point.LinearVelocity, T);
@@ -42,7 +42,7 @@ FCRSimPoint FCRSimPointContainer::GetPointInterpolated(int32 InIndex) const
 void FCRSimPointContainer::CachePreviousStep()
 {
 	PreviousStep = Points;
-	for (FCRSimPoint& Point : Points)
+	for (FRigVMSimPoint& Point : Points)
 	{
 		Point.LinearVelocity = FVector::ZeroVector;
 	}
@@ -73,15 +73,15 @@ void FCRSimPointContainer::IntegrateSprings()
 		{
 			continue;
 		}
-		const FCRSimPoint& PrevPointA = PreviousStep[Spring.SubjectA];
-		const FCRSimPoint& PrevPointB = PreviousStep[Spring.SubjectB];
+		const FRigVMSimPoint& PrevPointA = PreviousStep[Spring.SubjectA];
+		const FRigVMSimPoint& PrevPointB = PreviousStep[Spring.SubjectB];
 
 		FVector ForceA = FVector::ZeroVector;
 		FVector ForceB = FVector::ZeroVector;
 		Spring.CalculateForPoints(PrevPointA, PrevPointB, ForceA, ForceB);
 
-		FCRSimPoint& PointA = Points[Spring.SubjectA];
-		FCRSimPoint& PointB = Points[Spring.SubjectB];
+		FRigVMSimPoint& PointA = Points[Spring.SubjectA];
+		FRigVMSimPoint& PointB = Points[Spring.SubjectB];
 		PointA.LinearVelocity += ForceA;
 		PointB.LinearVelocity += ForceB;
 	}
@@ -128,8 +128,8 @@ void FCRSimPointContainer::ApplyConstraints()
 			continue;
 		}
 
-		FCRSimPoint& PointA = Points[Constraint.SubjectA];
-		FCRSimPoint& PointB = Points[Constraint.SubjectB == INDEX_NONE ? Constraint.SubjectA : Constraint.SubjectB];
+		FRigVMSimPoint& PointA = Points[Constraint.SubjectA];
+		FRigVMSimPoint& PointB = Points[Constraint.SubjectB == INDEX_NONE ? Constraint.SubjectA : Constraint.SubjectB];
 		Constraint.Apply(PointA, PointB);
 	}
 }
diff --git a/Engine/Plugins/Animation/ControlRig/Source/ControlRig/Private/Math/Simulation/CRSimPointForce.cpp b/Engine/Plugins/Animation/ControlRig/Source/ControlRig/Private/Math/Simulation/CRSimPointForce.cpp
--- a/Engine/Plugins/Animation/ControlRig/Source/ControlRig/Private/Math/Simulation/CRSimPointForce.cpp
+++ b/Engine/Plugins/Animation/ControlRig/Source/ControlRig/Private/Math/Simulation/CRSimPointForce.cpp
@@ -4,7 +4,7 @@
 
 #include UE_INLINE_GENERATED_CPP_BY_NAME(CRSimPointForce)
 
-FVector FCRSimPointForce::Calculate(const FCRSimPoint& InPoint, float InDeltaTime) const
+FVector FCRSimPointForce::Calculate(const FRigVMSimPoint& InPoint, float InDeltaTime) const
 {
 	FVector Force = FVector::ZeroVector;
 	if(InPoint.Mass > SMALL_NUMBER)
diff --git a/Engine/Plugins/Animation/ControlRig/Source/ControlRig/Private/Math/Simulation/CRSimSoftCollision.cpp b/Engine/Plugins/Animation/ControlRig/Source/ControlRig/Private/Math/Simulation/CRSimSoftCollision.cpp
--- a/Engine/Plugins/Animation/ControlRig/Source/ControlRig/Private/Math/Simulation/CRSimSoftCollision.cpp
+++ b/Engine/Plugins/Animation/ControlRig/Source/ControlRig/Private/Math/Simulation/CRSimSoftCollision.cpp
@@ -84,7 +84,7 @@ float FCRSimSoftCollision::CalculateFalloff(const FCRSimSoftCollision& InCollisi
 	return FControlRigMathLibrary::EaseFloat(Ratio, InCollision.FalloffType);
 }
 
-FVector FCRSimSoftCollision::CalculateForPoint(const FCRSimPoint& InPoint, float InDeltaTime) const
+FVector FCRSimSoftCollision::CalculateForPoint(const FRigVMSimPoint& InPoint, float InDeltaTime) const
 {
 	FVector Force = FVector::ZeroVector;
 	if(InPoint.Mass > SMALL_NUMBER)
diff --git a/Engine/Plugins/Animation/ControlRig/Source/ControlRig/Private/Rigs/AdditiveControlRig.cpp b/Engine/Plugins/Animation/ControlRig/Source/ControlRig/Private/Rigs/AdditiveControlRig.cpp
--- a/Engine/Plugins/Animation/ControlRig/Source/ControlRig/Private/Rigs/AdditiveControlRig.cpp
+++ b/Engine/Plugins/Animation/ControlRig/Source/ControlRig/Private/Rigs/AdditiveControlRig.cpp
@@ -48,7 +48,7 @@ bool UAdditiveControlRig::ExecuteUnits(FRigUnitContext& InOutContext, const FNam
 	{
 		FControlRigExecuteContext ExecuteContext;
 		ExecuteContext.Hierarchy = GetHierarchy();
-		ExecuteContext.EventName = InEventName;
+		ExecuteContext.SetEventName(InEventName);
 		ExecuteContext.UnitContext = InOutContext;
 
 		for (FRigUnit_AddBoneTransform& Unit : AddBoneRigUnits)
diff --git a/Engine/Plugins/Animation/ControlRig/Source/ControlRig/Private/Units/Animation/RigUnit_AnimBase.h b/Engine/Plugins/Animation/ControlRig/Source/ControlRig/Private/Units/Animation/RigUnit_AnimBase.h
deleted file mode 100644
diff --git a/Engine/Plugins/Animation/ControlRig/Source/ControlRig/Private/Units/Animation/RigUnit_AnimEasing.cpp b/Engine/Plugins/Animation/ControlRig/Source/ControlRig/Private/Units/Animation/RigUnit_AnimEasing.cpp
deleted file mode 100644
diff --git a/Engine/Plugins/Animation/ControlRig/Source/ControlRig/Private/Units/Animation/RigUnit_AnimEasing.h b/Engine/Plugins/Animation/ControlRig/Source/ControlRig/Private/Units/Animation/RigUnit_AnimEasing.h
deleted file mode 100644
diff --git a/Engine/Plugins/Animation/ControlRig/Source/ControlRig/Private/Units/Animation/RigUnit_AnimEvalRichCurve.cpp b/Engine/Plugins/Animation/ControlRig/Source/ControlRig/Private/Units/Animation/RigUnit_AnimEvalRichCurve.cpp
deleted file mode 100644
diff --git a/Engine/Plugins/Animation/ControlRig/Source/ControlRig/Private/Units/Animation/RigUnit_AnimEvalRichCurve.h b/Engine/Plugins/Animation/ControlRig/Source/ControlRig/Private/Units/Animation/RigUnit_AnimEvalRichCurve.h
deleted file mode 100644
diff --git a/Engine/Plugins/Animation/ControlRig/Source/ControlRig/Private/Units/Animation/RigUnit_AnimRichCurve.cpp b/Engine/Plugins/Animation/ControlRig/Source/ControlRig/Private/Units/Animation/RigUnit_AnimRichCurve.cpp
deleted file mode 100644
diff --git a/Engine/Plugins/Animation/ControlRig/Source/ControlRig/Private/Units/Animation/RigUnit_AnimRichCurve.h b/Engine/Plugins/Animation/ControlRig/Source/ControlRig/Private/Units/Animation/RigUnit_AnimRichCurve.h
deleted file mode 100644
diff --git a/Engine/Plugins/Animation/ControlRig/Source/ControlRig/Private/Units/Animation/RigUnit_GetDeltaTime.cpp b/Engine/Plugins/Animation/ControlRig/Source/ControlRig/Private/Units/Animation/RigUnit_GetDeltaTime.cpp
deleted file mode 100644
diff --git a/Engine/Plugins/Animation/ControlRig/Source/ControlRig/Private/Units/Animation/RigUnit_GetDeltaTime.h b/Engine/Plugins/Animation/ControlRig/Source/ControlRig/Private/Units/Animation/RigUnit_GetDeltaTime.h
deleted file mode 100644
diff --git a/Engine/Plugins/Animation/ControlRig/Source/ControlRig/Private/Units/Animation/RigUnit_GetWorldTime.cpp b/Engine/Plugins/Animation/ControlRig/Source/ControlRig/Private/Units/Animation/RigUnit_GetWorldTime.cpp
deleted file mode 100644
diff --git a/Engine/Plugins/Animation/ControlRig/Source/ControlRig/Private/Units/Animation/RigUnit_GetWorldTime.h b/Engine/Plugins/Animation/ControlRig/Source/ControlRig/Private/Units/Animation/RigUnit_GetWorldTime.h
deleted file mode 100644
diff --git a/Engine/Plugins/Animation/ControlRig/Source/ControlRig/Private/Units/Animation/RigUnit_TimeConversion.cpp b/Engine/Plugins/Animation/ControlRig/Source/ControlRig/Private/Units/Animation/RigUnit_TimeConversion.cpp
deleted file mode 100644
diff --git a/Engine/Plugins/Animation/ControlRig/Source/ControlRig/Private/Units/Animation/RigUnit_TimeConversion.h b/Engine/Plugins/Animation/ControlRig/Source/ControlRig/Private/Units/Animation/RigUnit_TimeConversion.h
deleted file mode 100644
diff --git a/Engine/Plugins/Animation/ControlRig/Source/ControlRig/Private/Units/Debug/RigUnit_DebugBase.h b/Engine/Plugins/Animation/ControlRig/Source/ControlRig/Private/Units/Debug/RigUnit_DebugBase.h
deleted file mode 100644
diff --git a/Engine/Plugins/Animation/ControlRig/Source/ControlRig/Private/Units/Debug/RigUnit_DebugBezier.cpp b/Engine/Plugins/Animation/ControlRig/Source/ControlRig/Private/Units/Debug/RigUnit_DebugBezier.cpp
--- a/Engine/Plugins/Animation/ControlRig/Source/ControlRig/Private/Units/Debug/RigUnit_DebugBezier.cpp
+++ b/Engine/Plugins/Animation/ControlRig/Source/ControlRig/Private/Units/Debug/RigUnit_DebugBezier.cpp
@@ -34,7 +34,7 @@ FRigUnit_DebugBezierItemSpace_Execute()
 	FTransform Transform = WorldOffset;
 	if (Space.IsValid())
 	{
-		Transform = Transform * ExecuteContext.Hierarchy->GetGlobalTransform(Space);
+		//Transform = Transform * ExecuteContext.Hierarchy->GetGlobalTransform(Space);
 	}
 
 	DrawBezier(ExecuteContext, Transform, Bezier, MinimumU, MaximumU, Color, Thickness, Detail);
@@ -46,7 +46,7 @@ FRigVMStructUpgradeInfo FRigUnit_DebugBezierItemSpace::GetUpgradeInfo() const
 	return FRigVMStructUpgradeInfo();
 }
 
-void FRigUnit_DebugBezierItemSpace::DrawBezier(const FControlRigExecuteContext& InContext, const FTransform& WorldOffset, const FRigVMFourPointBezier& InBezier, float MinimumU, float MaximumU, const FLinearColor& Color, float Thickness, int32 Detail)
+void FRigUnit_DebugBezierItemSpace::DrawBezier(const FRigVMExecuteContext& InContext, const FTransform& WorldOffset, const FRigVMFourPointBezier& InBezier, float MinimumU, float MaximumU, const FLinearColor& Color, float Thickness, int32 Detail)
 {
 	FRigVMDrawInterface* DrawInterface = InContext.GetDrawInterface();
 	if(DrawInterface == nullptr)
diff --git a/Engine/Plugins/Animation/ControlRig/Source/ControlRig/Private/Units/Debug/RigUnit_DebugBezier.h b/Engine/Plugins/Animation/ControlRig/Source/ControlRig/Private/Units/Debug/RigUnit_DebugBezier.h
--- a/Engine/Plugins/Animation/ControlRig/Source/ControlRig/Private/Units/Debug/RigUnit_DebugBezier.h
+++ b/Engine/Plugins/Animation/ControlRig/Source/ControlRig/Private/Units/Debug/RigUnit_DebugBezier.h
@@ -2,12 +2,12 @@
 
 #pragma once
 
-#include "RigUnit_DebugBase.h"
+#include "RigVMFunctions/Debug/RigVMFunction_DebugBase.h"
 #include "Math/ControlRigMathLibrary.h"
 #include "RigUnit_DebugBezier.generated.h"
 
 USTRUCT(meta=(DisplayName="Draw Bezier", Deprecated = "4.25"))
-struct CONTROLRIG_API FRigUnit_DebugBezier : public FRigUnit_DebugBaseMutable
+struct CONTROLRIG_API FRigUnit_DebugBezier : public FRigVMFunction_DebugBaseMutable
 {
 	GENERATED_BODY()
 
@@ -58,7 +58,7 @@ struct CONTROLRIG_API FRigUnit_DebugBezier : public FRigUnit_DebugBaseMutable
 };
 
 USTRUCT(meta=(DisplayName="Draw Bezier", Deprecated = "4.25"))
-struct CONTROLRIG_API FRigUnit_DebugBezierItemSpace : public FRigUnit_DebugBaseMutable
+struct CONTROLRIG_API FRigUnit_DebugBezierItemSpace : public FRigVMFunction_DebugBaseMutable
 {
 	GENERATED_BODY()
 
@@ -108,5 +108,5 @@ struct CONTROLRIG_API FRigUnit_DebugBezierItemSpace : public FRigUnit_DebugBaseM
 	RIGVM_METHOD()
 	virtual FRigVMStructUpgradeInfo GetUpgradeInfo() const override;
 
-	static void DrawBezier(const FControlRigExecuteContext& InContext, const FTransform& WorldOffset, const FRigVMFourPointBezier& InBezier, float MinimumU, float MaximumU, const FLinearColor& Color, float Thickness, int32 Detail);
+	static void DrawBezier(const FRigVMExecuteContext& InContext, const FTransform& WorldOffset, const FRigVMFourPointBezier& InBezier, float MinimumU, float MaximumU, const FLinearColor& Color, float Thickness, int32 Detail);
 };
diff --git a/Engine/Plugins/Animation/ControlRig/Source/ControlRig/Private/Units/Debug/RigUnit_DebugHierarchy.cpp b/Engine/Plugins/Animation/ControlRig/Source/ControlRig/Private/Units/Debug/RigUnit_DebugHierarchy.cpp
--- a/Engine/Plugins/Animation/ControlRig/Source/ControlRig/Private/Units/Debug/RigUnit_DebugHierarchy.cpp
+++ b/Engine/Plugins/Animation/ControlRig/Source/ControlRig/Private/Units/Debug/RigUnit_DebugHierarchy.cpp
@@ -2,6 +2,7 @@
 
 #include "Units/Debug/RigUnit_DebugHierarchy.h"
 #include "Units/RigUnitContext.h"
+#include "ControlRigDefines.h"
 
 #include UE_INLINE_GENERATED_CPP_BY_NAME(RigUnit_DebugHierarchy)
 
@@ -27,7 +28,7 @@ FRigUnit_DebugPose_Execute()
 	}
 }
 
-void FRigUnit_DebugHierarchy::DrawHierarchy(const FControlRigExecuteContext& InContext, const FTransform& WorldOffset, URigHierarchy* Hierarchy, EControlRigDrawHierarchyMode::Type Mode, float Scale, const FLinearColor& Color, float Thickness, const FRigPose* InPose)
+void FRigUnit_DebugHierarchy::DrawHierarchy(const FRigVMExecuteContext& InContext, const FTransform& WorldOffset, URigHierarchy* Hierarchy, EControlRigDrawHierarchyMode::Type Mode, float Scale, const FLinearColor& Color, float Thickness, const FRigPose* InPose)
 {
 	FRigVMDrawInterface* DrawInterface = InContext.GetDrawInterface();
 	if(DrawInterface == nullptr)
diff --git a/Engine/Plugins/Animation/ControlRig/Source/ControlRig/Private/Units/Debug/RigUnit_DebugHierarchy.h b/Engine/Plugins/Animation/ControlRig/Source/ControlRig/Private/Units/Debug/RigUnit_DebugHierarchy.h
--- a/Engine/Plugins/Animation/ControlRig/Source/ControlRig/Private/Units/Debug/RigUnit_DebugHierarchy.h
+++ b/Engine/Plugins/Animation/ControlRig/Source/ControlRig/Private/Units/Debug/RigUnit_DebugHierarchy.h
@@ -2,7 +2,9 @@
 
 #pragma once
 
-#include "RigUnit_DebugBase.h"
+#include "RigVMFunctions/Debug/RigVMFunction_DebugBase.h"
+#include "Rigs/RigHierarchy.h"
+#include "Units/RigUnitContext.h"
 #include "RigUnit_DebugHierarchy.generated.h"
 
 UENUM()
@@ -22,7 +24,7 @@ namespace EControlRigDrawHierarchyMode
  * Draws vectors on each bone in the viewport across the entire hierarchy
  */
 USTRUCT(meta=(DisplayName="Draw Hierarchy"))
-struct CONTROLRIG_API FRigUnit_DebugHierarchy : public FRigUnit_DebugBaseMutable
+struct CONTROLRIG_API FRigUnit_DebugHierarchy : public FRigVMFunction_DebugBase
 {
 	GENERATED_BODY()
 
@@ -38,6 +40,9 @@ struct CONTROLRIG_API FRigUnit_DebugHierarchy : public FRigUnit_DebugBaseMutable
 	RIGVM_METHOD()
 	virtual void Execute() override;
 
+	UPROPERTY(DisplayName = "Execute", Transient, meta = (Input, Output))
+	FControlRigExecuteContext ExecuteContext;
+
 	UPROPERTY(meta = (Input))
 	float Scale;
 
@@ -53,14 +58,14 @@ struct CONTROLRIG_API FRigUnit_DebugHierarchy : public FRigUnit_DebugBaseMutable
 	UPROPERTY(meta = (Input))
 	bool bEnabled;
 
-	static void DrawHierarchy(const FControlRigExecuteContext& InContext, const FTransform& WorldOffset, URigHierarchy* Hierarchy, EControlRigDrawHierarchyMode::Type Mode, float Scale, const FLinearColor& Color, float Thickness, const FRigPose* InPose);
+	static void DrawHierarchy(const FRigVMExecuteContext& InContext, const FTransform& WorldOffset, URigHierarchy* Hierarchy, EControlRigDrawHierarchyMode::Type Mode, float Scale, const FLinearColor& Color, float Thickness, const FRigPose* InPose);
 };
 
 /**
 * Draws vectors on each bone in the viewport across the entire pose
 */
 USTRUCT(meta=(DisplayName="Draw Pose Cache"))
-struct CONTROLRIG_API FRigUnit_DebugPose : public FRigUnit_DebugBaseMutable
+struct CONTROLRIG_API FRigUnit_DebugPose : public FRigVMFunction_DebugBase
 {
 	GENERATED_BODY()
 
@@ -76,6 +81,9 @@ struct CONTROLRIG_API FRigUnit_DebugPose : public FRigUnit_DebugBaseMutable
 	RIGVM_METHOD()
 	virtual void Execute() override;
 
+	UPROPERTY(DisplayName = "Execute", Transient, meta = (Input, Output))
+	FControlRigExecuteContext ExecuteContext;
+
 	UPROPERTY(meta = (Input))
 	FRigPose Pose;
 
diff --git a/Engine/Plugins/Animation/ControlRig/Source/ControlRig/Private/Units/Debug/RigUnit_DebugLine.cpp b/Engine/Plugins/Animation/ControlRig/Source/ControlRig/Private/Units/Debug/RigUnit_DebugLine.cpp
deleted file mode 100644
diff --git a/Engine/Plugins/Animation/ControlRig/Source/ControlRig/Private/Units/Debug/RigUnit_DebugLine.h b/Engine/Plugins/Animation/ControlRig/Source/ControlRig/Private/Units/Debug/RigUnit_DebugLine.h
deleted file mode 100644
diff --git a/Engine/Plugins/Animation/ControlRig/Source/ControlRig/Private/Units/Debug/RigUnit_DebugLineStrip.cpp b/Engine/Plugins/Animation/ControlRig/Source/ControlRig/Private/Units/Debug/RigUnit_DebugLineStrip.cpp
deleted file mode 100644
diff --git a/Engine/Plugins/Animation/ControlRig/Source/ControlRig/Private/Units/Debug/RigUnit_DebugLineStrip.h b/Engine/Plugins/Animation/ControlRig/Source/ControlRig/Private/Units/Debug/RigUnit_DebugLineStrip.h
deleted file mode 100644
diff --git a/Engine/Plugins/Animation/ControlRig/Source/ControlRig/Private/Units/Debug/RigUnit_DebugPoint.cpp b/Engine/Plugins/Animation/ControlRig/Source/ControlRig/Private/Units/Debug/RigUnit_DebugPoint.cpp
deleted file mode 100644
diff --git a/Engine/Plugins/Animation/ControlRig/Source/ControlRig/Private/Units/Debug/RigUnit_DebugPoint.h b/Engine/Plugins/Animation/ControlRig/Source/ControlRig/Private/Units/Debug/RigUnit_DebugPoint.h
deleted file mode 100644
diff --git a/Engine/Plugins/Animation/ControlRig/Source/ControlRig/Private/Units/Debug/RigUnit_DebugPrimitives.cpp b/Engine/Plugins/Animation/ControlRig/Source/ControlRig/Private/Units/Debug/RigUnit_DebugPrimitives.cpp
deleted file mode 100644
diff --git a/Engine/Plugins/Animation/ControlRig/Source/ControlRig/Private/Units/Debug/RigUnit_DebugPrimitives.h b/Engine/Plugins/Animation/ControlRig/Source/ControlRig/Private/Units/Debug/RigUnit_DebugPrimitives.h
deleted file mode 100644
diff --git a/Engine/Plugins/Animation/ControlRig/Source/ControlRig/Private/Units/Debug/RigUnit_DebugTransform.cpp b/Engine/Plugins/Animation/ControlRig/Source/ControlRig/Private/Units/Debug/RigUnit_DebugTransform.cpp
deleted file mode 100644
diff --git a/Engine/Plugins/Animation/ControlRig/Source/ControlRig/Private/Units/Debug/RigUnit_DebugTransform.h b/Engine/Plugins/Animation/ControlRig/Source/ControlRig/Private/Units/Debug/RigUnit_DebugTransform.h
deleted file mode 100644
diff --git a/Engine/Plugins/Animation/ControlRig/Source/ControlRig/Private/Units/Debug/RigUnit_ProfilingBracket.cpp b/Engine/Plugins/Animation/ControlRig/Source/ControlRig/Private/Units/Debug/RigUnit_ProfilingBracket.cpp
--- a/Engine/Plugins/Animation/ControlRig/Source/ControlRig/Private/Units/Debug/RigUnit_ProfilingBracket.cpp
+++ b/Engine/Plugins/Animation/ControlRig/Source/ControlRig/Private/Units/Debug/RigUnit_ProfilingBracket.cpp
@@ -2,6 +2,7 @@
 
 #include "Units/Debug/RigUnit_ProfilingBracket.h"
 #include "Units/RigUnitContext.h"
+#include "ControlRigDefines.h"
 #include "KismetAnimationLibrary.h"
 
 #include UE_INLINE_GENERATED_CPP_BY_NAME(RigUnit_ProfilingBracket)
diff --git a/Engine/Plugins/Animation/ControlRig/Source/ControlRig/Private/Units/Debug/RigUnit_ProfilingBracket.h b/Engine/Plugins/Animation/ControlRig/Source/ControlRig/Private/Units/Debug/RigUnit_ProfilingBracket.h
--- a/Engine/Plugins/Animation/ControlRig/Source/ControlRig/Private/Units/Debug/RigUnit_ProfilingBracket.h
+++ b/Engine/Plugins/Animation/ControlRig/Source/ControlRig/Private/Units/Debug/RigUnit_ProfilingBracket.h
@@ -2,14 +2,14 @@
 
 #pragma once
 
-#include "RigUnit_DebugBase.h"
+#include "RigVMFunctions/Debug/RigVMFunction_DebugBase.h"
 #include "RigUnit_ProfilingBracket.generated.h"
 
 /**
  * Starts a profiling timer for debugging, used in conjunction with End Profiling Timer
  */
 USTRUCT(meta=(DisplayName="Start Profiling Timer", Keywords="Measure,BeginProfiling,Profile", NodeColor="0.25, 0.25, 0.05000000074505806"))
-struct CONTROLRIG_API FRigUnit_StartProfilingTimer : public FRigUnit_DebugBaseMutable
+struct CONTROLRIG_API FRigUnit_StartProfilingTimer : public FRigVMFunction_DebugBaseMutable
 {
 	GENERATED_BODY()
 
@@ -25,7 +25,7 @@ struct CONTROLRIG_API FRigUnit_StartProfilingTimer : public FRigUnit_DebugBaseMu
  * Ends an existing profiling timer for debugging, used in conjunction with Start Profiling Timer
  */
 USTRUCT(meta = (DisplayName = "End Profiling Timer", Keywords = "Measure,StopProfiling,Meter,Profile", NodeColor="0.25, 0.25, 0.05000000074505806"))
-struct CONTROLRIG_API FRigUnit_EndProfilingTimer : public FRigUnit_DebugBaseMutable
+struct CONTROLRIG_API FRigUnit_EndProfilingTimer : public FRigVMFunction_DebugBaseMutable
 {
 	GENERATED_BODY()
 
diff --git a/Engine/Plugins/Animation/ControlRig/Source/ControlRig/Private/Units/Debug/RigUnit_VisualDebug.cpp b/Engine/Plugins/Animation/ControlRig/Source/ControlRig/Private/Units/Debug/RigUnit_VisualDebug.cpp
deleted file mode 100644
diff --git a/Engine/Plugins/Animation/ControlRig/Source/ControlRig/Private/Units/Debug/RigUnit_VisualDebug.h b/Engine/Plugins/Animation/ControlRig/Source/ControlRig/Private/Units/Debug/RigUnit_VisualDebug.h
deleted file mode 100644
diff --git a/Engine/Plugins/Animation/ControlRig/Source/ControlRig/Private/Units/Execution/RigUnit_BeginExecution.cpp b/Engine/Plugins/Animation/ControlRig/Source/ControlRig/Private/Units/Execution/RigUnit_BeginExecution.cpp
--- a/Engine/Plugins/Animation/ControlRig/Source/ControlRig/Private/Units/Execution/RigUnit_BeginExecution.cpp
+++ b/Engine/Plugins/Animation/ControlRig/Source/ControlRig/Private/Units/Execution/RigUnit_BeginExecution.cpp
@@ -10,6 +10,6 @@ FName FRigUnit_BeginExecution::EventName = TEXT("Forwards Solve");
 FRigUnit_BeginExecution_Execute()
 {
     DECLARE_SCOPE_HIERARCHICAL_COUNTER_RIGUNIT()
-	ExecuteContext.EventName = FRigUnit_BeginExecution::EventName;
+	ExecuteContext.SetEventName(FRigUnit_BeginExecution::EventName);
 }
 
diff --git a/Engine/Plugins/Animation/ControlRig/Source/ControlRig/Private/Units/Execution/RigUnit_ForLoop.cpp b/Engine/Plugins/Animation/ControlRig/Source/ControlRig/Private/Units/Execution/RigUnit_ForLoop.cpp
deleted file mode 100644
diff --git a/Engine/Plugins/Animation/ControlRig/Source/ControlRig/Private/Units/Execution/RigUnit_ForLoop.h b/Engine/Plugins/Animation/ControlRig/Source/ControlRig/Private/Units/Execution/RigUnit_ForLoop.h
deleted file mode 100644
diff --git a/Engine/Plugins/Animation/ControlRig/Source/ControlRig/Private/Units/Execution/RigUnit_InteractionExecution.cpp b/Engine/Plugins/Animation/ControlRig/Source/ControlRig/Private/Units/Execution/RigUnit_InteractionExecution.cpp
--- a/Engine/Plugins/Animation/ControlRig/Source/ControlRig/Private/Units/Execution/RigUnit_InteractionExecution.cpp
+++ b/Engine/Plugins/Animation/ControlRig/Source/ControlRig/Private/Units/Execution/RigUnit_InteractionExecution.cpp
@@ -10,6 +10,6 @@ FName FRigUnit_InteractionExecution::EventName = TEXT("Interaction");
 FRigUnit_InteractionExecution_Execute()
 {
     DECLARE_SCOPE_HIERARCHICAL_COUNTER_RIGUNIT()
-	ExecuteContext.EventName = FRigUnit_InteractionExecution::EventName;
+	ExecuteContext.SetEventName(FRigUnit_InteractionExecution::EventName);
 }
 
diff --git a/Engine/Plugins/Animation/ControlRig/Source/ControlRig/Private/Units/Execution/RigUnit_InverseExecution.cpp b/Engine/Plugins/Animation/ControlRig/Source/ControlRig/Private/Units/Execution/RigUnit_InverseExecution.cpp
--- a/Engine/Plugins/Animation/ControlRig/Source/ControlRig/Private/Units/Execution/RigUnit_InverseExecution.cpp
+++ b/Engine/Plugins/Animation/ControlRig/Source/ControlRig/Private/Units/Execution/RigUnit_InverseExecution.cpp
@@ -10,6 +10,6 @@ FName FRigUnit_InverseExecution::EventName = TEXT("Backwards Solve");
 FRigUnit_InverseExecution_Execute()
 {
     DECLARE_SCOPE_HIERARCHICAL_COUNTER_RIGUNIT()
-	ExecuteContext.EventName = FRigUnit_InverseExecution::EventName;
+	ExecuteContext.SetEventName(FRigUnit_InverseExecution::EventName);
 }
 
diff --git a/Engine/Plugins/Animation/ControlRig/Source/ControlRig/Private/Units/Execution/RigUnit_PrepareForExecution.cpp b/Engine/Plugins/Animation/ControlRig/Source/ControlRig/Private/Units/Execution/RigUnit_PrepareForExecution.cpp
--- a/Engine/Plugins/Animation/ControlRig/Source/ControlRig/Private/Units/Execution/RigUnit_PrepareForExecution.cpp
+++ b/Engine/Plugins/Animation/ControlRig/Source/ControlRig/Private/Units/Execution/RigUnit_PrepareForExecution.cpp
@@ -10,6 +10,6 @@ FName FRigUnit_PrepareForExecution::EventName = TEXT("Construction");
 FRigUnit_PrepareForExecution_Execute()
 {
     DECLARE_SCOPE_HIERARCHICAL_COUNTER_RIGUNIT()
-	ExecuteContext.EventName = FRigUnit_PrepareForExecution::EventName;
+	ExecuteContext.SetEventName(FRigUnit_PrepareForExecution::EventName);
 }
 
diff --git a/Engine/Plugins/Animation/ControlRig/Source/ControlRig/Private/Units/Execution/RigUnit_UserDefinedEvent.cpp b/Engine/Plugins/Animation/ControlRig/Source/ControlRig/Private/Units/Execution/RigUnit_UserDefinedEvent.cpp
deleted file mode 100644
diff --git a/Engine/Plugins/Animation/ControlRig/Source/ControlRig/Private/Units/Execution/RigUnit_UserDefinedEvent.h b/Engine/Plugins/Animation/ControlRig/Source/ControlRig/Private/Units/Execution/RigUnit_UserDefinedEvent.h
deleted file mode 100644
diff --git a/Engine/Plugins/Animation/ControlRig/Source/ControlRig/Private/Units/Highlevel/Hierarchy/RigUnit_SpringIK.cpp b/Engine/Plugins/Animation/ControlRig/Source/ControlRig/Private/Units/Highlevel/Hierarchy/RigUnit_SpringIK.cpp
--- a/Engine/Plugins/Animation/ControlRig/Source/ControlRig/Private/Units/Highlevel/Hierarchy/RigUnit_SpringIK.cpp
+++ b/Engine/Plugins/Animation/ControlRig/Source/ControlRig/Private/Units/Highlevel/Hierarchy/RigUnit_SpringIK.cpp
@@ -55,7 +55,7 @@ FRigUnit_SpringIK_Execute()
 
 		for (int32 PointIndex = 0; PointIndex < CachedBones.Num() - 1; PointIndex++)
 		{
-			Simulation.Points.Add(FCRSimPoint());
+			Simulation.Points.Add(FRigVMSimPoint());
 
 			FTransform A = Hierarchy->GetInitialGlobalTransform(CachedBones[PointIndex]);
 			FTransform B = Hierarchy->GetInitialGlobalTransform(CachedBones[PointIndex + 1]);
@@ -89,7 +89,7 @@ FRigUnit_SpringIK_Execute()
 				Simulation.Springs.Add(Spring);
 			}
 		}
-		Simulation.Points.Add(FCRSimPoint());
+		Simulation.Points.Add(FRigVMSimPoint());
 
 		Simulation.Points[0].Mass = 0.f;
 		Simulation.Points.Last().Mass = 0.f;
diff --git a/Engine/Plugins/Animation/ControlRig/Source/ControlRig/Private/Units/Simulation/RigUnit_Accumulate.cpp b/Engine/Plugins/Animation/ControlRig/Source/ControlRig/Private/Units/Simulation/RigUnit_Accumulate.cpp
deleted file mode 100644
diff --git a/Engine/Plugins/Animation/ControlRig/Source/ControlRig/Private/Units/Simulation/RigUnit_Accumulate.h b/Engine/Plugins/Animation/ControlRig/Source/ControlRig/Private/Units/Simulation/RigUnit_Accumulate.h
deleted file mode 100644
diff --git a/Engine/Plugins/Animation/ControlRig/Source/ControlRig/Private/Units/Simulation/RigUnit_AlphaInterp.cpp b/Engine/Plugins/Animation/ControlRig/Source/ControlRig/Private/Units/Simulation/RigUnit_AlphaInterp.cpp
deleted file mode 100644
diff --git a/Engine/Plugins/Animation/ControlRig/Source/ControlRig/Private/Units/Simulation/RigUnit_AlphaInterp.h b/Engine/Plugins/Animation/ControlRig/Source/ControlRig/Private/Units/Simulation/RigUnit_AlphaInterp.h
deleted file mode 100644
diff --git a/Engine/Plugins/Animation/ControlRig/Source/ControlRig/Private/Units/Simulation/RigUnit_DeltaFromPrevious.cpp b/Engine/Plugins/Animation/ControlRig/Source/ControlRig/Private/Units/Simulation/RigUnit_DeltaFromPrevious.cpp
deleted file mode 100644
diff --git a/Engine/Plugins/Animation/ControlRig/Source/ControlRig/Private/Units/Simulation/RigUnit_DeltaFromPrevious.h b/Engine/Plugins/Animation/ControlRig/Source/ControlRig/Private/Units/Simulation/RigUnit_DeltaFromPrevious.h
deleted file mode 100644
diff --git a/Engine/Plugins/Animation/ControlRig/Source/ControlRig/Private/Units/Simulation/RigUnit_Kalman.cpp b/Engine/Plugins/Animation/ControlRig/Source/ControlRig/Private/Units/Simulation/RigUnit_Kalman.cpp
deleted file mode 100644
diff --git a/Engine/Plugins/Animation/ControlRig/Source/ControlRig/Private/Units/Simulation/RigUnit_Kalman.h b/Engine/Plugins/Animation/ControlRig/Source/ControlRig/Private/Units/Simulation/RigUnit_Kalman.h
deleted file mode 100644
diff --git a/Engine/Plugins/Animation/ControlRig/Source/ControlRig/Private/Units/Simulation/RigUnit_PointSimulation.cpp b/Engine/Plugins/Animation/ControlRig/Source/ControlRig/Private/Units/Simulation/RigUnit_PointSimulation.cpp
--- a/Engine/Plugins/Animation/ControlRig/Source/ControlRig/Private/Units/Simulation/RigUnit_PointSimulation.cpp
+++ b/Engine/Plugins/Animation/ControlRig/Source/ControlRig/Private/Units/Simulation/RigUnit_PointSimulation.cpp
@@ -8,247 +8,6 @@
 FRigUnit_PointSimulation_Execute()
 {
     DECLARE_SCOPE_HIERARCHICAL_COUNTER_RIGUNIT()
-	URigHierarchy* Hierarchy = ExecuteContext.Hierarchy;
-
-	float DeltaTime = ExecuteContext.GetDeltaTime();
-
-	FCRSimPointContainer& Simulation = WorkData.Simulation;
-	TArray<FCachedRigElement>& BoneIndices = WorkData.BoneIndices;
-
-	if (Simulation.Points.Num() != Points.Num() ||
-		Simulation.Springs.Num() != Links.Num() ||
-		Simulation.Forces.Num() != Forces.Num() ||
-		Simulation.CollisionVolumes.Num() != CollisionVolumes.Num() ||
-		BoneIndices.Num() != BoneTargets.Num())
-	{
-		Simulation.Reset();
-		BoneIndices.Reset();
-		return;
-	}
-
-	if(Simulation.AccumulatedTime < SMALL_NUMBER)
-	{
-		Simulation.TimeStep = 1.0f / FMath::Clamp<float>(SimulatedStepsPerSecond, 1.f, 120.f);
-
-		bool bFoundErrors = false;
-		for (const FCRSimLinearSpring& Link : Links)
-		{
-			if (Link.SubjectA < 0 || Link.SubjectA >= Points.Num() ||
-				Link.SubjectB < 0 || Link.SubjectB >= Points.Num() ||
-				Link.SubjectA == Link.SubjectB)
-			{
-				UE_CONTROLRIG_RIGUNIT_REPORT_WARNING(TEXT("Link has invalid subjects: %d and %d."), Link.SubjectA, Link.SubjectB);
-				bFoundErrors = true;
-			}
-		}
-
-		for (const FRigUnit_PointSimulation_BoneTarget& BoneTarget : BoneTargets)
-		{
-			if (BoneTarget.TranslationPoint >= Points.Num())
-			{
-				UE_CONTROLRIG_RIGUNIT_REPORT_WARNING(TEXT("Bone Target has invalid translation point: %d."), BoneTarget.TranslationPoint);
-				bFoundErrors = true;
-			}
-			if (BoneTarget.PrimaryAimPoint >= Points.Num())
-			{
-				UE_CONTROLRIG_RIGUNIT_REPORT_WARNING(TEXT("Bone Target has invalid primary aim point: %d."), BoneTarget.PrimaryAimPoint);
-				bFoundErrors = true;
-			}
-			if (BoneTarget.SecondaryAimPoint >= Points.Num())
-			{
-				UE_CONTROLRIG_RIGUNIT_REPORT_WARNING(TEXT("Bone Target has invalid secondary aim point: %d."), BoneTarget.SecondaryAimPoint);
-				bFoundErrors = true;
-			}
-		}
-
-		if (bFoundErrors)
-		{
-			return;
-		}
-
-		Simulation.Points.SetNum(Points.Num());
-		Simulation.Springs.SetNum(Links.Num());
-		Simulation.Forces.SetNum(Forces.Num());
-		Simulation.CollisionVolumes.SetNum(CollisionVolumes.Num());
-
-		if (Hierarchy)
-		{
-			for (const FRigUnit_PointSimulation_BoneTarget& BoneTarget : BoneTargets)
-			{
-				BoneIndices.Add(FCachedRigElement());
-			}
-		}
-
-		DeltaTime = 0.f;
-	}
-
-	if (Simulation.Points.Num() == 0)
-	{
-		return;
-	}
-
-	for (int32 PointIndex = 0; PointIndex < Points.Num(); PointIndex++)
-	{
-		Simulation.Points[PointIndex].Mass = Points[PointIndex].Mass;
-		Simulation.Points[PointIndex].Size = Points[PointIndex].Size;
-		if (Simulation.AccumulatedTime < SMALL_NUMBER || Points[PointIndex].Mass < SMALL_NUMBER)
-		{
-			Simulation.Points[PointIndex] = Points[PointIndex];
-		}
-		else if (Points[PointIndex].InheritMotion > SMALL_NUMBER)
-		{
-			Simulation.Points[PointIndex].InheritMotion = FMath::Clamp<float>(Points[PointIndex].InheritMotion * Simulation.TimeStep, 0.f, 1.f);
-			Simulation.Points[PointIndex].Position = FMath::Lerp<FVector>(Simulation.Points[PointIndex].Position, Points[PointIndex].Position, Simulation.Points[PointIndex].InheritMotion);
-		}
-	}
-	for (int32 LinkIndex = 0; LinkIndex < Links.Num(); LinkIndex++)
-	{
-		if (Simulation.AccumulatedTime < SMALL_NUMBER)
-		{
-			Simulation.Springs[LinkIndex] = Links[LinkIndex];
-			if (Simulation.Springs[LinkIndex].Equilibrium < -SMALL_NUMBER)
-			{
-				const FCRSimPoint SubjectA = Simulation.Points[Simulation.Springs[LinkIndex].SubjectA];
-				const FCRSimPoint SubjectB = Simulation.Points[Simulation.Springs[LinkIndex].SubjectB];
-				Simulation.Springs[LinkIndex].Equilibrium = (SubjectA.Position - SubjectB.Position).Size();
-			}
-		}
-		else
-		{
-			Simulation.Springs[LinkIndex].Coefficient = Links[LinkIndex].Coefficient;
-			if (Links[LinkIndex].Equilibrium >= 0.f)
-			{
-				Simulation.Springs[LinkIndex].Equilibrium = Links[LinkIndex].Equilibrium;
-			}
-		}
-	}
-	for (int32 ForceIndex = 0; ForceIndex < Forces.Num(); ForceIndex++)
-	{
-		Simulation.Forces[ForceIndex] = Forces[ForceIndex];
-	}
-	for (int32 VolumeIndex = 0; VolumeIndex < CollisionVolumes.Num(); VolumeIndex++)
-	{
-		Simulation.CollisionVolumes[VolumeIndex] = CollisionVolumes[VolumeIndex];
-	}
-
-	if (DeltaTime > SMALL_NUMBER)
-	{
-		switch (IntegratorType)
-		{
-		case ECRSimPointIntegrateType::Verlet:
-		{
-			Simulation.StepVerlet(DeltaTime, VerletBlend);
-			break;
-		}
-		case ECRSimPointIntegrateType::SemiExplicitEuler:
-		{
-			Simulation.StepSemiExplicitEuler(DeltaTime);
-			break;
-		}
-		}
-	}
-
-	if (Hierarchy)
-	{
-		TArray<FVector> Positions;
-		Positions.SetNumUninitialized(Simulation.Points.Num());
-		for (int32 PointIndex = 0; PointIndex < Positions.Num(); PointIndex++)
-		{
-			Positions[PointIndex] = Simulation.GetPointInterpolated(PointIndex).Position;
-		}
-
-		for (int32 TargetIndex = 0; TargetIndex < BoneTargets.Num(); TargetIndex++)
-		{
-			const FRigElementKey Key(BoneTargets[TargetIndex].Bone, ERigElementType::Bone);
-			if (BoneIndices[TargetIndex].UpdateCache(Key, Hierarchy))
-			{
-				const FRigUnit_PointSimulation_BoneTarget& BoneTarget = BoneTargets[TargetIndex];
-				FTransform Transform = Hierarchy->GetGlobalTransform(BoneIndices[TargetIndex]);
-
-				if (BoneTarget.TranslationPoint >= 0 && BoneTarget.TranslationPoint < Simulation.Points.Num())
-				{
-					Transform.SetLocation(Positions[BoneTarget.TranslationPoint]);
-				}
-
-				if (bLimitLocalPosition)
-				{
-					FRigElementKey ParentKey = Hierarchy->GetFirstParent(BoneIndices[TargetIndex].GetKey());
-					if (ParentKey.IsValid())
-					{
-						FTransform InitialTransform = Hierarchy->GetInitialGlobalTransform(BoneIndices[TargetIndex]);
-						FTransform ParentInitialTransform = Hierarchy->GetGlobalTransform(ParentKey, true);
-						FTransform ParentTransform = Hierarchy->GetGlobalTransform(ParentKey);
-						float ExpectedDistance = (InitialTransform.GetLocation() - ParentInitialTransform.GetLocation()).Size();
-						if (ExpectedDistance > SMALL_NUMBER)
-						{
-							FVector Direction = Transform.GetLocation() - ParentTransform.GetLocation();
-							if (!Direction.IsNearlyZero())
-							{
-								Transform.SetLocation(ParentTransform.GetLocation() + Direction.GetSafeNormal() * ExpectedDistance);
-								if (BoneTarget.TranslationPoint >= 0 && BoneTarget.TranslationPoint < Simulation.Points.Num())
-								{
-									Positions[BoneTarget.TranslationPoint] = Transform.GetLocation();
-								}
-							}
-						}
-					}
-				}
-
-				Hierarchy->SetGlobalTransform(BoneIndices[TargetIndex], Transform, false);
-			}
-		}
-
-		for (int32 TargetIndex = 0; TargetIndex < BoneTargets.Num(); TargetIndex++)
-		{
-			if (BoneIndices[TargetIndex] != INDEX_NONE)
-			{
-				const FRigUnit_PointSimulation_BoneTarget& BoneTarget = BoneTargets[TargetIndex];
-				FTransform Transform = Hierarchy->GetGlobalTransform(BoneIndices[TargetIndex]);
-
-				FVector PrimaryAxis = Transform.TransformVectorNoScale(PrimaryAimAxis);
-
-				if (BoneTarget.PrimaryAimPoint >= 0 && BoneTarget.PrimaryAimPoint < Simulation.Points.Num())
-				{
-					FVector PrimaryAimPosition = Positions[BoneTarget.PrimaryAimPoint];
-					FVector Target = PrimaryAimPosition - Transform.GetLocation();
-
-					if (!Target.IsNearlyZero() && !PrimaryAxis.IsNearlyZero())
-					{
-						PrimaryAxis = PrimaryAxis.GetSafeNormal();
-						Target = Target.GetSafeNormal();
-						FQuat Rotation = FQuat::FindBetweenNormals(PrimaryAxis, Target);
-						Transform.SetRotation((Rotation * Transform.GetRotation()).GetNormalized());
-					}
-				}
-
-				if (BoneTarget.SecondaryAimPoint >= 0 && BoneTarget.SecondaryAimPoint < Simulation.Points.Num())
-				{
-					FVector SecondaryAimPosition = Positions[BoneTarget.SecondaryAimPoint];
-					FVector Target = SecondaryAimPosition - Transform.GetLocation();
-					Target = Target - FVector::DotProduct(Target, PrimaryAxis) * PrimaryAxis;
-					FVector SecondaryAxis = Transform.TransformVectorNoScale(SecondaryAimAxis);
-
-					if (!Target.IsNearlyZero() && !SecondaryAxis.IsNearlyZero())
-					{
-						SecondaryAxis = SecondaryAxis.GetSafeNormal();
-						Target = Target.GetSafeNormal();
-						FQuat Rotation = FQuat::FindBetweenNormals(SecondaryAxis, Target);
-						Transform.SetRotation((Rotation * Transform.GetRotation()).GetNormalized());
-					}
-				}
-
-				Hierarchy->SetGlobalTransform(BoneIndices[TargetIndex], Transform, bPropagateToChildren);
-			}
-		}
-	}
-
-	if (Simulation.Points.Num() > 3)
-	{
-		Bezier.A = Simulation.GetPointInterpolated(0).Position;
-		Bezier.B = Simulation.GetPointInterpolated(1).Position;
-		Bezier.C = Simulation.GetPointInterpolated(2).Position;
-		Bezier.D = Simulation.GetPointInterpolated(3).Position;
-	}
 }
 
 FRigVMStructUpgradeInfo FRigUnit_PointSimulation::GetUpgradeInfo() const
diff --git a/Engine/Plugins/Animation/ControlRig/Source/ControlRig/Private/Units/Simulation/RigUnit_PointSimulation.h b/Engine/Plugins/Animation/ControlRig/Source/ControlRig/Private/Units/Simulation/RigUnit_PointSimulation.h
--- a/Engine/Plugins/Animation/ControlRig/Source/ControlRig/Private/Units/Simulation/RigUnit_PointSimulation.h
+++ b/Engine/Plugins/Animation/ControlRig/Source/ControlRig/Private/Units/Simulation/RigUnit_PointSimulation.h
@@ -2,9 +2,9 @@
 
 #pragma once
 
-#include "RigUnit_SimBase.h"
+#include "RigVMFunctions/Simulation/RigVMFunction_SimBase.h"
 #include "Math/ControlRigMathLibrary.h"
-#include "Math/Simulation/CRSimPoint.h"
+#include "RigVMFunctions/Math/RigVMMathLibrary.h"
 #include "Math/Simulation/CRSimLinearSpring.h"
 #include "Math/Simulation/CRSimPointForce.h"
 #include "Math/Simulation/CRSimSoftCollision.h"
@@ -118,14 +118,14 @@ struct CONTROLRIG_API FRigUnit_PointSimulation_WorkData
  * Note: Disabled for now.
  */
 USTRUCT(meta=(DisplayName="Point Simulation", Keywords="Simulate,Verlet,Springs", Deprecated="4.25"))
-struct CONTROLRIG_API FRigUnit_PointSimulation : public FRigUnit_SimBaseMutable
+struct CONTROLRIG_API FRigUnit_PointSimulation : public FRigVMFunction_SimBaseMutable
 {
 	GENERATED_BODY()
 	
 	FRigUnit_PointSimulation()
 	{
 		SimulatedStepsPerSecond = 60.f;
-		IntegratorType = ECRSimPointIntegrateType::Verlet;
+		IntegratorType = ERigVMSimPointIntegrateType::Verlet;
 		VerletBlend = 4.f;
 		bLimitLocalPosition = true;
 		bPropagateToChildren = true;
@@ -140,7 +140,7 @@ struct CONTROLRIG_API FRigUnit_PointSimulation : public FRigUnit_SimBaseMutable
 
 	/** The points to simulate */
 	UPROPERTY(meta = (Input))
-	TArray<FCRSimPoint> Points;
+	TArray<FRigVMSimPoint> Points;
 
 	/** The links to connect the points with */
 	UPROPERTY(meta = (Input))
@@ -160,7 +160,7 @@ struct CONTROLRIG_API FRigUnit_PointSimulation : public FRigUnit_SimBaseMutable
 
 	/** The type of integrator to use */
 	UPROPERTY(meta = (Input, Constant))
-	ECRSimPointIntegrateType IntegratorType;
+	ERigVMSimPointIntegrateType IntegratorType;
 
 	/** The amount of blending to apply per second ( only for verlet integrations )*/
 	UPROPERTY(meta = (Input))
diff --git a/Engine/Plugins/Animation/ControlRig/Source/ControlRig/Private/Units/Simulation/RigUnit_SimBase.h b/Engine/Plugins/Animation/ControlRig/Source/ControlRig/Private/Units/Simulation/RigUnit_SimBase.h
deleted file mode 100644
diff --git a/Engine/Plugins/Animation/ControlRig/Source/ControlRig/Private/Units/Simulation/RigUnit_SpringInterp.cpp b/Engine/Plugins/Animation/ControlRig/Source/ControlRig/Private/Units/Simulation/RigUnit_SpringInterp.cpp
--- a/Engine/Plugins/Animation/ControlRig/Source/ControlRig/Private/Units/Simulation/RigUnit_SpringInterp.cpp
+++ b/Engine/Plugins/Animation/ControlRig/Source/ControlRig/Private/Units/Simulation/RigUnit_SpringInterp.cpp
@@ -1,6 +1,7 @@
 // Copyright Epic Games, Inc. All Rights Reserved.
 
 #include "Units/Simulation/RigUnit_SpringInterp.h"
+#include "ControlRigDefines.h"
 #include "Units/RigUnitContext.h"
 
 #include UE_INLINE_GENERATED_CPP_BY_NAME(RigUnit_SpringInterp)
diff --git a/Engine/Plugins/Animation/ControlRig/Source/ControlRig/Private/Units/Simulation/RigUnit_SpringInterp.h b/Engine/Plugins/Animation/ControlRig/Source/ControlRig/Private/Units/Simulation/RigUnit_SpringInterp.h
--- a/Engine/Plugins/Animation/ControlRig/Source/ControlRig/Private/Units/Simulation/RigUnit_SpringInterp.h
+++ b/Engine/Plugins/Animation/ControlRig/Source/ControlRig/Private/Units/Simulation/RigUnit_SpringInterp.h
@@ -2,7 +2,7 @@
 
 #pragma once
 
-#include "RigUnit_SimBase.h"
+#include "RigVMFunctions/Simulation/RigVMFunction_SimBase.h"
 #include "Kismet/KismetMathLibrary.h"
 #include "RigUnit_SpringInterp.generated.h"
 
@@ -10,7 +10,7 @@
  * Uses a simple spring model to interpolate a float from Current to Target.
  */
 USTRUCT(meta=(DisplayName="Spring Interpolate", Keywords="Alpha,SpringInterpolate", TemplateName="SpringInterp", Deprecated="5.0"))
-struct CONTROLRIG_API FRigUnit_SpringInterp : public FRigUnit_SimBase
+struct CONTROLRIG_API FRigUnit_SpringInterp : public FRigVMFunction_SimBase
 {
 	GENERATED_BODY()
  
@@ -55,7 +55,7 @@ struct CONTROLRIG_API FRigUnit_SpringInterp : public FRigUnit_SimBase
  * Uses a simple spring model to interpolate a vector from Current to Target.
  */
 USTRUCT(meta=(DisplayName="Spring Interpolate", Keywords="Alpha,SpringInterpolate", TemplateName="SpringInterp", MenuDescSuffix = "(Vector)", Deprecated="5.0"))
-struct CONTROLRIG_API FRigUnit_SpringInterpVector : public FRigUnit_SimBase
+struct CONTROLRIG_API FRigUnit_SpringInterpVector : public FRigVMFunction_SimBase
 {
 	GENERATED_BODY()
 	
@@ -100,7 +100,7 @@ struct CONTROLRIG_API FRigUnit_SpringInterpVector : public FRigUnit_SimBase
  * Uses a simple spring model to interpolate a float from Current to Target.
  */
 USTRUCT(meta=(DisplayName="Spring Interpolate", Keywords="Alpha,SpringInterpolate,Verlet", Category = "Simulation|Springs", TemplateName="SpringInterp"))
-struct CONTROLRIG_API FRigUnit_SpringInterpV2 : public FRigUnit_SimBase
+struct CONTROLRIG_API FRigUnit_SpringInterpV2 : public FRigVMFunction_SimBase
 {
 	GENERATED_BODY()
 
@@ -190,7 +190,7 @@ struct CONTROLRIG_API FRigUnit_SpringInterpV2 : public FRigUnit_SimBase
  * Uses a simple spring model to interpolate a vector from Current to Target.
  */
 USTRUCT(meta=(DisplayName="Spring Interpolate", Keywords="Alpha,SpringInterpolate,Verlet", Category = "Simulation|Springs", TemplateName="SpringInterp", MenuDescSuffix = "(Vector)"))
-struct CONTROLRIG_API FRigUnit_SpringInterpVectorV2 : public FRigUnit_SimBase
+struct CONTROLRIG_API FRigUnit_SpringInterpVectorV2 : public FRigVMFunction_SimBase
 {
 	GENERATED_BODY()
 	
@@ -280,7 +280,7 @@ struct CONTROLRIG_API FRigUnit_SpringInterpVectorV2 : public FRigUnit_SimBase
  * Uses a simple spring model to interpolate a quaternion from Current to Target.
  */
 USTRUCT(meta=(DisplayName="Spring Interpolate", Keywords="Alpha,SpringInterpolate,Verlet", Category = "Simulation|Springs", TemplateName="SpringInterp", MenuDescSuffix = "(Quaternion)"))
-struct CONTROLRIG_API FRigUnit_SpringInterpQuaternionV2 : public FRigUnit_SimBase
+struct CONTROLRIG_API FRigUnit_SpringInterpQuaternionV2 : public FRigVMFunction_SimBase
 {
 	GENERATED_BODY()
 	
diff --git a/Engine/Plugins/Animation/ControlRig/Source/ControlRig/Private/Units/Simulation/RigUnit_TimeOffset.cpp b/Engine/Plugins/Animation/ControlRig/Source/ControlRig/Private/Units/Simulation/RigUnit_TimeOffset.cpp
deleted file mode 100644
diff --git a/Engine/Plugins/Animation/ControlRig/Source/ControlRig/Private/Units/Simulation/RigUnit_TimeOffset.h b/Engine/Plugins/Animation/ControlRig/Source/ControlRig/Private/Units/Simulation/RigUnit_TimeOffset.h
deleted file mode 100644
diff --git a/Engine/Plugins/Animation/ControlRig/Source/ControlRig/Private/Units/Simulation/RigUnit_Timeline.cpp b/Engine/Plugins/Animation/ControlRig/Source/ControlRig/Private/Units/Simulation/RigUnit_Timeline.cpp
deleted file mode 100644
diff --git a/Engine/Plugins/Animation/ControlRig/Source/ControlRig/Private/Units/Simulation/RigUnit_Timeline.h b/Engine/Plugins/Animation/ControlRig/Source/ControlRig/Private/Units/Simulation/RigUnit_Timeline.h
deleted file mode 100644
diff --git a/Engine/Plugins/Animation/ControlRig/Source/ControlRig/Private/Units/Simulation/RigUnit_Verlet.cpp b/Engine/Plugins/Animation/ControlRig/Source/ControlRig/Private/Units/Simulation/RigUnit_Verlet.cpp
deleted file mode 100644
diff --git a/Engine/Plugins/Animation/ControlRig/Source/ControlRig/Private/Units/Simulation/RigUnit_Verlet.h b/Engine/Plugins/Animation/ControlRig/Source/ControlRig/Private/Units/Simulation/RigUnit_Verlet.h
deleted file mode 100644
diff --git a/Engine/Plugins/Animation/ControlRig/Source/ControlRig/Public/Math/Simulation/CRSimLinearSpring.h b/Engine/Plugins/Animation/ControlRig/Source/ControlRig/Public/Math/Simulation/CRSimLinearSpring.h
--- a/Engine/Plugins/Animation/ControlRig/Source/ControlRig/Public/Math/Simulation/CRSimLinearSpring.h
+++ b/Engine/Plugins/Animation/ControlRig/Source/ControlRig/Public/Math/Simulation/CRSimLinearSpring.h
@@ -3,7 +3,7 @@
 #pragma once
 
 #include "CoreMinimal.h"
-#include "CRSimPoint.h"
+#include "RigVMFunctions/Math/RigVMMathLibrary.h"
 #include "CRSimLinearSpring.generated.h"
 
 USTRUCT(BlueprintType)
@@ -44,5 +44,5 @@ struct FCRSimLinearSpring
 	UPROPERTY(EditAnywhere, Category=Simulation)
 	float Equilibrium;
 
-	void CalculateForPoints(const FCRSimPoint& InPointA, const FCRSimPoint& InPointB, FVector& ForceA, FVector& ForceB) const;
+	void CalculateForPoints(const FRigVMSimPoint& InPointA, const FRigVMSimPoint& InPointB, FVector& ForceA, FVector& ForceB) const;
 };
diff --git a/Engine/Plugins/Animation/ControlRig/Source/ControlRig/Public/Math/Simulation/CRSimPoint.h b/Engine/Plugins/Animation/ControlRig/Source/ControlRig/Public/Math/Simulation/CRSimPoint.h
deleted file mode 100644
diff --git a/Engine/Plugins/Animation/ControlRig/Source/ControlRig/Public/Math/Simulation/CRSimPointConstraint.h b/Engine/Plugins/Animation/ControlRig/Source/ControlRig/Public/Math/Simulation/CRSimPointConstraint.h
--- a/Engine/Plugins/Animation/ControlRig/Source/ControlRig/Public/Math/Simulation/CRSimPointConstraint.h
+++ b/Engine/Plugins/Animation/ControlRig/Source/ControlRig/Public/Math/Simulation/CRSimPointConstraint.h
@@ -3,7 +3,7 @@
 #pragma once
 
 #include "CoreMinimal.h"
-#include "CRSimPoint.h"
+#include "RigVMFunctions/Math/RigVMMathLibrary.h"
 #include "CRSimPointConstraint.generated.h"
 
 UENUM()
@@ -58,5 +58,5 @@ struct FCRSimPointConstraint
 	UPROPERTY()
 	FVector DataB;
 
-	void Apply(FCRSimPoint& OutPointA, FCRSimPoint& OutPointB) const;
+	void Apply(FRigVMSimPoint& OutPointA, FRigVMSimPoint& OutPointB) const;
 };
diff --git a/Engine/Plugins/Animation/ControlRig/Source/ControlRig/Public/Math/Simulation/CRSimPointContainer.h b/Engine/Plugins/Animation/ControlRig/Source/ControlRig/Public/Math/Simulation/CRSimPointContainer.h
--- a/Engine/Plugins/Animation/ControlRig/Source/ControlRig/Public/Math/Simulation/CRSimPointContainer.h
+++ b/Engine/Plugins/Animation/ControlRig/Source/ControlRig/Public/Math/Simulation/CRSimPointContainer.h
@@ -3,7 +3,7 @@
 #pragma once
 
 #include "CoreMinimal.h"
-#include "CRSimPoint.h"
+#include "RigVMFunctions/Math/RigVMMathLibrary.h"
 #include "CRSimLinearSpring.h"
 #include "CRSimPointForce.h"
 #include "CRSimPointConstraint.h"
@@ -24,7 +24,7 @@ struct FCRSimPointContainer : public FCRSimContainer
 	 * The points within the simulation
 	 */
 	UPROPERTY()
-	TArray<FCRSimPoint> Points;
+	TArray<FRigVMSimPoint> Points;
 
 	/**
 	 * The springs within the simulation
@@ -50,7 +50,7 @@ struct FCRSimPointContainer : public FCRSimContainer
 	UPROPERTY()
 	TArray<FCRSimPointConstraint> Constraints;
 
-	FCRSimPoint GetPointInterpolated(int32 InIndex) const;
+	FRigVMSimPoint GetPointInterpolated(int32 InIndex) const;
 
 	virtual void Reset() override;
 	virtual void ResetTime() override;
@@ -58,7 +58,7 @@ struct FCRSimPointContainer : public FCRSimContainer
 protected:
 
 	UPROPERTY()
-	TArray<FCRSimPoint> PreviousStep;
+	TArray<FRigVMSimPoint> PreviousStep;
 
 	virtual void CachePreviousStep() override;
 	virtual void IntegrateVerlet(float InBlend) override;
diff --git a/Engine/Plugins/Animation/ControlRig/Source/ControlRig/Public/Math/Simulation/CRSimPointForce.h b/Engine/Plugins/Animation/ControlRig/Source/ControlRig/Public/Math/Simulation/CRSimPointForce.h
--- a/Engine/Plugins/Animation/ControlRig/Source/ControlRig/Public/Math/Simulation/CRSimPointForce.h
+++ b/Engine/Plugins/Animation/ControlRig/Source/ControlRig/Public/Math/Simulation/CRSimPointForce.h
@@ -3,7 +3,7 @@
 #pragma once
 
 #include "CoreMinimal.h"
-#include "CRSimPoint.h"
+#include "RigVMFunctions/Math/RigVMMathLibrary.h"
 #include "CRSimPointForce.generated.h"
 
 UENUM()
@@ -51,5 +51,5 @@ struct FCRSimPointForce
 	UPROPERTY(EditAnywhere, Category=Simulation)
 	bool bNormalize;
 
-	FVector Calculate(const FCRSimPoint& InPoint, float InDeltaTime) const;
+	FVector Calculate(const FRigVMSimPoint& InPoint, float InDeltaTime) const;
 };
diff --git a/Engine/Plugins/Animation/ControlRig/Source/ControlRig/Public/Math/Simulation/CRSimSoftCollision.h b/Engine/Plugins/Animation/ControlRig/Source/ControlRig/Public/Math/Simulation/CRSimSoftCollision.h
--- a/Engine/Plugins/Animation/ControlRig/Source/ControlRig/Public/Math/Simulation/CRSimSoftCollision.h
+++ b/Engine/Plugins/Animation/ControlRig/Source/ControlRig/Public/Math/Simulation/CRSimSoftCollision.h
@@ -4,7 +4,7 @@
 
 #include "CoreMinimal.h"
 #include "Math/ControlRigMathLibrary.h"
-#include "CRSimPoint.h"
+#include "RigVMFunctions/Math/RigVMMathLibrary.h"
 #include "CRSimSoftCollision.generated.h"
 
 UENUM()
@@ -78,5 +78,5 @@ struct FCRSimSoftCollision
 	bool bInverted;
 
 	static float CalculateFalloff(const FCRSimSoftCollision& InCollision, const FVector& InPosition, float InSize, FVector& OutDirection);
-	FVector CalculateForPoint(const FCRSimPoint& InPoint, float InDeltaTime) const;
+	FVector CalculateForPoint(const FRigVMSimPoint& InPoint, float InDeltaTime) const;
 };
diff --git a/Engine/Plugins/Animation/ControlRig/Source/ControlRigEditor/Private/ControlRigEditorModule.cpp b/Engine/Plugins/Animation/ControlRig/Source/ControlRigEditor/Private/ControlRigEditorModule.cpp
--- a/Engine/Plugins/Animation/ControlRig/Source/ControlRigEditor/Private/ControlRigEditorModule.cpp
+++ b/Engine/Plugins/Animation/ControlRig/Source/ControlRigEditor/Private/ControlRigEditorModule.cpp
@@ -82,8 +82,8 @@
 #include "RigVMModel/Nodes/RigVMLibraryNode.h"
 #include "RigVMModel/RigVMVariableDescription.h"
 #include "ControlRigBlueprint.h"
-#include "Units/Simulation/RigUnit_AlphaInterp.h"
-#include "Units/Debug/RigUnit_VisualDebug.h"
+#include "RigVMFunctions/Simulation/RigVMFunction_AlphaInterp.h"
+#include "RigVMFunctions/Debug/RigVMFunction_VisualDebug.h"
 #include "SKismetInspector.h"
 #include "Dialogs/Dialogs.h"
 #include "Settings/ControlRigSettings.h"
@@ -1653,11 +1653,11 @@ void FControlRigEditorModule::GetContextMenuActions(const UControlRigGraphSchema
 
 									if ((ModelPin->GetCPPType() == TEXT("float")) || (ModelPin->GetCPPType() == TEXT("double")))
 									{
-										ScriptStruct = FRigUnit_AlphaInterp::StaticStruct();
+										ScriptStruct = FRigVMFunction_AlphaInterp::StaticStruct();
 									}
 									else if (ModelPin->GetCPPType() == TEXT("FVector"))
 									{
-										ScriptStruct = FRigUnit_AlphaInterpVector::StaticStruct();
+										ScriptStruct = FRigVMFunction_AlphaInterpVector::StaticStruct();
 									}
 									else
 									{
@@ -1745,15 +1745,15 @@ void FControlRigEditorModule::GetContextMenuActions(const UControlRigGraphSchema
 
 									if (ModelPin->GetCPPType() == TEXT("FVector"))
 									{
-										ScriptStruct = FRigUnit_VisualDebugVectorItemSpace::StaticStruct();
+										ScriptStruct = FRigVMFunction_VisualDebugVectorNoSpace::StaticStruct();
 									}
 									else if (ModelPin->GetCPPType() == TEXT("FQuat"))
 									{
-										ScriptStruct = FRigUnit_VisualDebugQuatItemSpace::StaticStruct();
+										ScriptStruct = FRigVMFunction_VisualDebugQuatNoSpace::StaticStruct();
 									}
 									else if (ModelPin->GetCPPType() == TEXT("FTransform"))
 									{
-										ScriptStruct = FRigUnit_VisualDebugTransformItemSpace::StaticStruct();
+										ScriptStruct = FRigVMFunction_VisualDebugTransformNoSpace::StaticStruct();
 									}
 									else
 									{
diff --git a/Engine/Plugins/Animation/ControlRig/Source/ControlRigEditor/Private/Graph/NodeSpawners/ControlRigUnitNodeSpawner.cpp b/Engine/Plugins/Animation/ControlRig/Source/ControlRigEditor/Private/Graph/NodeSpawners/ControlRigUnitNodeSpawner.cpp
--- a/Engine/Plugins/Animation/ControlRig/Source/ControlRigEditor/Private/Graph/NodeSpawners/ControlRigUnitNodeSpawner.cpp
+++ b/Engine/Plugins/Animation/ControlRig/Source/ControlRigEditor/Private/Graph/NodeSpawners/ControlRigUnitNodeSpawner.cpp
@@ -21,7 +21,7 @@
 #include "Units/Execution/RigUnit_InverseExecution.h"
 #include "Units/Execution/RigUnit_PrepareForExecution.h"
 #include "Units/Execution/RigUnit_InteractionExecution.h"
-#include "Units/Execution/RigUnit_UserDefinedEvent.h"
+#include "RigVMFunctions/Execution/RigVMFunction_UserDefinedEvent.h"
 
 #include UE_INLINE_GENERATED_CPP_BY_NAME(ControlRigUnitNodeSpawner)
 
@@ -184,7 +184,7 @@ UControlRigGraphNode* UControlRigUnitNodeSpawner::SpawnNode(UEdGraph* ParentGrap
 			if(!EventName.IsNone())
 			{
 				if(StructInstance->CanOnlyExistOnce() &&
-					(StructTemplate != FRigUnit_UserDefinedEvent::StaticStruct()))
+					(StructTemplate != FRigVMFunction_UserDefinedEvent::StaticStruct()))
 				{
 					const TArray<URigVMGraph*> Models = RigBlueprint->GetAllModels();
 					for(URigVMGraph* Model : Models)
@@ -226,7 +226,7 @@ UControlRigGraphNode* UControlRigUnitNodeSpawner::SpawnNode(UEdGraph* ParentGrap
 
 			if (NewNode && bIsUserFacingNode)
 			{
-				if(StructTemplate == FRigUnit_UserDefinedEvent::StaticStruct())
+				if(StructTemplate == FRigVMFunction_UserDefinedEvent::StaticStruct())
 				{
 					// ensure uniqueness for the event name
 					TArray<FName> ExistingEventNames = {
diff --git a/Engine/Source/Runtime/RigVM/Private/RigVMFunctions/Animation/RigVMFunction_AnimEasing.cpp b/Engine/Source/Runtime/RigVM/Private/RigVMFunctions/Animation/RigVMFunction_AnimEasing.cpp
new file mode 100644
--- /dev/null
+++ b/Engine/Source/Runtime/RigVM/Private/RigVMFunctions/Animation/RigVMFunction_AnimEasing.cpp
@@ -0,0 +1,23 @@
+// Copyright Epic Games, Inc. All Rights Reserved.
+
+#include "RigVMFunctions/Animation/RigVMFunction_AnimEasing.h"
+
+#include UE_INLINE_GENERATED_CPP_BY_NAME(RigVMFunction_AnimEasing)
+
+FRigVMFunction_AnimEasingType_Execute()
+{
+}
+
+FRigVMFunction_AnimEasing_Execute()
+{
+	if (FMath::IsNearlyEqual(SourceMinimum, SourceMaximum))
+	{
+		UE_RIGVMSTRUCT_REPORT_WARNING(TEXT("The source minimum and maximum are the same."));
+	}
+
+	Result = FMath::Clamp<float>((Value - SourceMinimum) / (SourceMaximum - SourceMinimum), 0.f, 1.f);
+	Result = FRigVMMathLibrary::EaseFloat(Result, Type);
+	Result = FMath::Lerp<float>(TargetMinimum, TargetMaximum, Result);
+}
+
+
diff --git a/Engine/Source/Runtime/RigVM/Private/RigVMFunctions/Animation/RigVMFunction_AnimEvalRichCurve.cpp b/Engine/Source/Runtime/RigVM/Private/RigVMFunctions/Animation/RigVMFunction_AnimEvalRichCurve.cpp
new file mode 100644
--- /dev/null
+++ b/Engine/Source/Runtime/RigVM/Private/RigVMFunctions/Animation/RigVMFunction_AnimEvalRichCurve.cpp
@@ -0,0 +1,44 @@
+// Copyright Epic Games, Inc. All Rights Reserved.
+
+#include "RigVMFunctions/Animation/RigVMFunction_AnimEvalRichCurve.h"
+#include "RigVMFunctions/RigVMFunctionDefines.h"
+
+#include UE_INLINE_GENERATED_CPP_BY_NAME(RigVMFunction_AnimEvalRichCurve)
+
+FRigVMFunction_AnimEvalRichCurve_Execute()
+{
+	if (FMath::IsNearlyEqual(SourceMinimum, SourceMaximum))
+	{
+		UE_RIGVMSTRUCT_REPORT_WARNING(TEXT("The source minimum and maximum are the same."));
+	}
+
+	Result = FMath::Clamp<float>((Value - SourceMinimum) / (SourceMaximum - SourceMinimum), 0.f, 1.f);
+
+	if (Curve.GetRichCurveConst()->GetNumKeys() > 0)
+	{
+		Result = Curve.GetRichCurveConst()->Eval(Result, 0.f);
+	}
+
+	Result = FMath::Lerp<float>(TargetMinimum, TargetMaximum, Result);
+}
+
+#if WITH_DEV_AUTOMATION_TESTS
+#include "RigVMCore/RigVMStructTest.h"
+
+IMPLEMENT_RIGVMSTRUCT_AUTOMATION_TEST(FRigVMFunction_AnimEvalRichCurve)
+{
+	Unit.Value = 0.f;
+	Execute();
+	AddErrorIfFalse(FMath::IsNearlyEqual(Unit.Result, 0.f), TEXT("unexpected curve result"));
+	Unit.Value = 0.5f;
+	Execute();
+	AddErrorIfFalse(FMath::IsNearlyEqual(Unit.Result, 0.5f), TEXT("unexpected curve result"));
+	Unit.Value = 1.f;
+	Execute();
+	AddErrorIfFalse(FMath::IsNearlyEqual(Unit.Result, 1.f), TEXT("unexpected curve result"));
+	Unit.Value = 1.5f;
+	Execute();
+	AddErrorIfFalse(FMath::IsNearlyEqual(Unit.Result, 1.f), TEXT("unexpected curve result"));
+	return true;
+}
+#endif
diff --git a/Engine/Source/Runtime/RigVM/Private/RigVMFunctions/Animation/RigVMFunction_AnimRichCurve.cpp b/Engine/Source/Runtime/RigVM/Private/RigVMFunctions/Animation/RigVMFunction_AnimRichCurve.cpp
new file mode 100644
--- /dev/null
+++ b/Engine/Source/Runtime/RigVM/Private/RigVMFunctions/Animation/RigVMFunction_AnimRichCurve.cpp
@@ -0,0 +1,9 @@
+// Copyright Epic Games, Inc. All Rights Reserved.
+
+#include "RigVMFunctions/Animation/RigVMFunction_AnimRichCurve.h"
+
+#include UE_INLINE_GENERATED_CPP_BY_NAME(RigVMFunction_AnimRichCurve)
+
+FRigVMFunction_AnimRichCurve_Execute()
+{
+}
diff --git a/Engine/Source/Runtime/RigVM/Private/RigVMFunctions/Animation/RigVMFunction_GetDeltaTime.cpp b/Engine/Source/Runtime/RigVM/Private/RigVMFunctions/Animation/RigVMFunction_GetDeltaTime.cpp
new file mode 100644
--- /dev/null
+++ b/Engine/Source/Runtime/RigVM/Private/RigVMFunctions/Animation/RigVMFunction_GetDeltaTime.cpp
@@ -0,0 +1,22 @@
+// Copyright Epic Games, Inc. All Rights Reserved.
+
+#include "RigVMFunctions/Animation/RigVMFunction_GetDeltaTime.h"
+
+#include UE_INLINE_GENERATED_CPP_BY_NAME(RigVMFunction_GetDeltaTime)
+
+FRigVMFunction_GetDeltaTime_Execute()
+{
+	Result = ExecuteContext.GetDeltaTime();
+}
+
+#if WITH_DEV_AUTOMATION_TESTS
+#include "RigVMCore/RigVMStructTest.h"
+
+IMPLEMENT_RIGVMSTRUCT_AUTOMATION_TEST(FRigVMFunction_GetDeltaTime)
+{
+	ExecuteContext.SetDeltaTime(0.2f);
+	Execute();
+	AddErrorIfFalse(FMath::IsNearlyEqual(Unit.Result, 0.2f), TEXT("unexpected delta time"));
+	return true;
+}
+#endif
diff --git a/Engine/Source/Runtime/RigVM/Private/RigVMFunctions/Animation/RigVMFunction_GetWorldTime.cpp b/Engine/Source/Runtime/RigVM/Private/RigVMFunctions/Animation/RigVMFunction_GetWorldTime.cpp
new file mode 100644
--- /dev/null
+++ b/Engine/Source/Runtime/RigVM/Private/RigVMFunctions/Animation/RigVMFunction_GetWorldTime.cpp
@@ -0,0 +1,20 @@
+// Copyright Epic Games, Inc. All Rights Reserved.
+
+#include "RigVMFunctions/Animation/RigVMFunction_GetWorldTime.h"
+
+#include UE_INLINE_GENERATED_CPP_BY_NAME(RigVMFunction_GetWorldTime)
+
+FRigVMFunction_GetWorldTime_Execute()
+{
+	FDateTime Now = FDateTime::Now();
+	Year = (float)Now.GetYear();
+	Month = (float)Now.GetMonth();
+	Day = (float)Now.GetDay();
+	WeekDay = (float)Now.GetDayOfWeek();
+	Hours = (float)Now.GetHour();
+	Minutes = (float)Now.GetMinute();
+	Seconds = float(Now.GetSecond()) + float(Now.GetMillisecond()) * 0.001f;
+	OverallSeconds = (float)FPlatformTime::Seconds();
+}
+
+
diff --git a/Engine/Source/Runtime/RigVM/Private/RigVMFunctions/Animation/RigVMFunction_TimeConversion.cpp b/Engine/Source/Runtime/RigVM/Private/RigVMFunctions/Animation/RigVMFunction_TimeConversion.cpp
new file mode 100644
--- /dev/null
+++ b/Engine/Source/Runtime/RigVM/Private/RigVMFunctions/Animation/RigVMFunction_TimeConversion.cpp
@@ -0,0 +1,30 @@
+// Copyright Epic Games, Inc. All Rights Reserved.
+
+#include "RigVMFunctions/Animation/RigVMFunction_TimeConversion.h"
+
+#include UE_INLINE_GENERATED_CPP_BY_NAME(RigVMFunction_TimeConversion)
+
+FRigVMFunction_FramesToSeconds_Execute()
+{
+    if(ExecuteContext.GetFramesPerSecond() > SMALL_NUMBER)
+    {
+		Seconds = Frames / ExecuteContext.GetFramesPerSecond();
+	}
+	else
+	{
+		Seconds = 0.f;
+	}
+}
+
+FRigVMFunction_SecondsToFrames_Execute()
+{
+    if(ExecuteContext.GetFramesPerSecond() > SMALL_NUMBER)
+    {
+		Frames = Seconds * ExecuteContext.GetFramesPerSecond();
+	}
+	else
+	{
+		Frames = 0.f;
+	}
+}
+
diff --git a/Engine/Source/Runtime/RigVM/Private/RigVMFunctions/Debug/RigVMFunction_DebugLine.cpp b/Engine/Source/Runtime/RigVM/Private/RigVMFunctions/Debug/RigVMFunction_DebugLine.cpp
new file mode 100644
--- /dev/null
+++ b/Engine/Source/Runtime/RigVM/Private/RigVMFunctions/Debug/RigVMFunction_DebugLine.cpp
@@ -0,0 +1,42 @@
+// Copyright Epic Games, Inc. All Rights Reserved.
+
+#include "RigVMFunctions/Debug/RigVMFunction_DebugLine.h"
+
+#include UE_INLINE_GENERATED_CPP_BY_NAME(RigVMFunction_DebugLine)
+
+FRigVMFunction_DebugLine_Execute()
+{
+	FRigVMFunction_DebugLineNoSpace::StaticExecute(
+		ExecuteContext, 
+		A,
+		B,
+		Color,
+		Thickness,
+		WorldOffset, 
+		bEnabled);
+}
+
+FRigVMStructUpgradeInfo FRigVMFunction_DebugLine::GetUpgradeInfo() const
+{
+	FRigVMFunction_DebugLineNoSpace NewNode;
+	NewNode.A = A;
+	NewNode.B = B;
+	NewNode.Color = Color;
+	NewNode.Thickness = Thickness;
+	NewNode.WorldOffset = WorldOffset;
+	NewNode.bEnabled = bEnabled;
+
+	FRigVMStructUpgradeInfo Info(*this, NewNode);
+	Info.AddRemappedPin(TEXT("Space"), TEXT("Space.Name"));
+	return Info;
+}
+
+FRigVMFunction_DebugLineNoSpace_Execute()
+{
+	if (ExecuteContext.GetDrawInterface() == nullptr || !bEnabled)
+	{
+		return;
+	}
+
+	ExecuteContext.GetDrawInterface()->DrawLine(WorldOffset, A, B, Color, Thickness);
+}
diff --git a/Engine/Source/Runtime/RigVM/Private/RigVMFunctions/Debug/RigVMFunction_DebugLineStrip.cpp b/Engine/Source/Runtime/RigVM/Private/RigVMFunctions/Debug/RigVMFunction_DebugLineStrip.cpp
new file mode 100644
--- /dev/null
+++ b/Engine/Source/Runtime/RigVM/Private/RigVMFunctions/Debug/RigVMFunction_DebugLineStrip.cpp
@@ -0,0 +1,41 @@
+// Copyright Epic Games, Inc. All Rights Reserved.
+
+#include "RigVMFunctions/Debug/RigVMFunction_DebugLineStrip.h"
+
+#include UE_INLINE_GENERATED_CPP_BY_NAME(RigVMFunction_DebugLineStrip)
+
+FRigVMFunction_DebugLineStrip_Execute()
+{
+	FRigVMFunction_DebugLineStripNoSpace::StaticExecute(
+		ExecuteContext, 
+		Points,
+		Color,
+		Thickness,
+		WorldOffset, 
+		bEnabled);
+}
+
+FRigVMStructUpgradeInfo FRigVMFunction_DebugLineStrip::GetUpgradeInfo() const
+{
+	FRigVMFunction_DebugLineStripNoSpace NewNode;
+	NewNode.Points = Points;
+	NewNode.Color = Color;
+	NewNode.Thickness = Thickness;
+	NewNode.WorldOffset = WorldOffset;
+	NewNode.bEnabled = bEnabled;
+
+	FRigVMStructUpgradeInfo Info(*this, NewNode);
+	Info.AddRemappedPin(TEXT("Space"), TEXT("Space.Name"));
+	return Info;
+}
+
+FRigVMFunction_DebugLineStripNoSpace_Execute()
+{
+	if (ExecuteContext.GetDrawInterface() == nullptr || !bEnabled)
+	{
+		return;
+	}
+
+	ExecuteContext.GetDrawInterface()->DrawLineStrip(WorldOffset, TArrayView<const FVector>(Points.GetData(), Points.Num()), Color, Thickness);
+}
+
diff --git a/Engine/Source/Runtime/RigVM/Private/RigVMFunctions/Debug/RigVMFunction_DebugPoint.cpp b/Engine/Source/Runtime/RigVM/Private/RigVMFunctions/Debug/RigVMFunction_DebugPoint.cpp
new file mode 100644
--- /dev/null
+++ b/Engine/Source/Runtime/RigVM/Private/RigVMFunctions/Debug/RigVMFunction_DebugPoint.cpp
@@ -0,0 +1,79 @@
+// Copyright Epic Games, Inc. All Rights Reserved.
+
+#include "RigVMFunctions/Debug/RigVMFunction_DebugPoint.h"
+#include "RigVMFunctions/Debug/RigVMFunction_VisualDebug.h"
+
+#include UE_INLINE_GENERATED_CPP_BY_NAME(RigVMFunction_DebugPoint)
+
+FRigVMFunction_DebugPoint_Execute()
+{
+	if (ExecuteContext.GetDrawInterface() == nullptr || !bEnabled)
+	{
+		return;
+	}
+
+	static const FVector Center = FVector::ZeroVector;
+	const FVector DrawVector = Vector;
+
+	switch (Mode)
+	{
+		case ERigUnitDebugPointMode::Point:
+		{
+			ExecuteContext.GetDrawInterface()->DrawPoint(WorldOffset, DrawVector, Scale, Color);
+			break;
+		}
+		case ERigUnitDebugPointMode::Vector:
+		{
+			ExecuteContext.GetDrawInterface()->DrawLine(WorldOffset, Center, DrawVector, Color, Thickness);
+			break;
+		}
+	}
+}
+
+FRigVMStructUpgradeInfo FRigVMFunction_DebugPoint::GetUpgradeInfo() const
+{
+	FRigVMFunction_VisualDebugVector NewNode;
+	NewNode.Value = Vector;
+	NewNode.Color = Color;
+	NewNode.Thickness = Thickness;
+	NewNode.BoneSpace = Space;
+	NewNode.Scale = Scale;
+	NewNode.bEnabled = bEnabled;
+
+	FRigVMStructUpgradeInfo Info(*this, NewNode);
+	Info.AddRemappedPin(TEXT("Vector"), TEXT("Value"));
+	Info.AddRemappedPin(TEXT("Space"), TEXT("BoneSpace"));
+	return Info;
+}
+
+FRigVMFunction_DebugPointMutable_Execute()
+{
+	if (ExecuteContext.GetDrawInterface() == nullptr || !bEnabled)
+	{
+		return;
+	}
+
+	static const FVector Center = FVector::ZeroVector;
+	const FVector DrawVector = Vector;
+
+	switch (Mode)
+	{
+		case ERigUnitDebugPointMode::Point:
+		{
+			ExecuteContext.GetDrawInterface()->DrawPoint(WorldOffset, DrawVector, Scale, Color);
+			break;
+		}
+		case ERigUnitDebugPointMode::Vector:
+		{
+			ExecuteContext.GetDrawInterface()->DrawLine(WorldOffset, Center, DrawVector, Color, Thickness);
+			break;
+		}
+	}
+}
+
+FRigVMStructUpgradeInfo FRigVMFunction_DebugPointMutable::GetUpgradeInfo() const
+{
+	// this node is no longer supported
+	return FRigVMStructUpgradeInfo();
+}
+
diff --git a/Engine/Source/Runtime/RigVM/Private/RigVMFunctions/Debug/RigVMFunction_DebugPrimitives.cpp b/Engine/Source/Runtime/RigVM/Private/RigVMFunctions/Debug/RigVMFunction_DebugPrimitives.cpp
new file mode 100644
--- /dev/null
+++ b/Engine/Source/Runtime/RigVM/Private/RigVMFunctions/Debug/RigVMFunction_DebugPrimitives.cpp
@@ -0,0 +1,85 @@
+// Copyright Epic Games, Inc. All Rights Reserved.
+
+#include "RigVMFunctions/Debug/RigVMFunction_DebugPrimitives.h"
+
+#include UE_INLINE_GENERATED_CPP_BY_NAME(RigVMFunction_DebugPrimitives)
+
+FRigVMFunction_DebugRectangle_Execute()
+{
+	FRigVMFunction_DebugRectangleNoSpace::StaticExecute(
+		ExecuteContext, 
+		Transform,
+		Color,
+		Scale,
+		Thickness,
+		WorldOffset, 
+		bEnabled);
+}
+
+FRigVMStructUpgradeInfo FRigVMFunction_DebugRectangle::GetUpgradeInfo() const
+{
+	FRigVMFunction_DebugRectangleNoSpace NewNode;
+	NewNode.Transform = Transform;
+	NewNode.Color = Color;
+	NewNode.Scale = Scale;
+	NewNode.Thickness = Thickness;
+	NewNode.WorldOffset = WorldOffset;
+	NewNode.bEnabled = bEnabled;
+
+	FRigVMStructUpgradeInfo Info(*this, NewNode);
+	Info.AddRemappedPin(TEXT("Space"), TEXT("Space.Name"));
+	return Info;
+}
+
+FRigVMFunction_DebugRectangleNoSpace_Execute()
+{
+	if (ExecuteContext.GetDrawInterface() == nullptr || !bEnabled)
+	{
+		return;
+	}
+
+	ExecuteContext.GetDrawInterface()->DrawRectangle(WorldOffset, Transform, Scale, Color, Thickness);
+}
+
+FRigVMFunction_DebugArc_Execute()
+{
+	FRigVMFunction_DebugArcNoSpace::StaticExecute(
+		ExecuteContext, 
+		Transform,
+		Color,
+		Radius,
+		MinimumDegrees,
+		MaximumDegrees,
+		Thickness,
+		Detail,
+		WorldOffset, 
+		bEnabled);
+}
+
+FRigVMStructUpgradeInfo FRigVMFunction_DebugArc::GetUpgradeInfo() const
+{
+	FRigVMFunction_DebugArcNoSpace NewNode;
+	NewNode.Transform = Transform;
+	NewNode.Color = Color;
+	NewNode.Radius = Radius;
+	NewNode.MinimumDegrees = MinimumDegrees;
+	NewNode.MaximumDegrees = MaximumDegrees;
+	NewNode.Thickness = Thickness;
+	NewNode.Detail = Detail;
+	NewNode.WorldOffset = WorldOffset;
+	NewNode.bEnabled = bEnabled;
+
+	FRigVMStructUpgradeInfo Info(*this, NewNode);
+	Info.AddRemappedPin(TEXT("Space"), TEXT("Space.Name"));
+	return Info;
+}
+
+FRigVMFunction_DebugArcNoSpace_Execute()
+{
+	if (ExecuteContext.GetDrawInterface() == nullptr || !bEnabled)
+	{
+		return;
+	}
+
+	ExecuteContext.GetDrawInterface()->DrawArc(WorldOffset, Transform, Radius, FMath::DegreesToRadians(MinimumDegrees), FMath::DegreesToRadians(MaximumDegrees), Color, Thickness, Detail);
+}
diff --git a/Engine/Source/Runtime/RigVM/Private/RigVMFunctions/Debug/RigVMFunction_DebugTransform.cpp b/Engine/Source/Runtime/RigVM/Private/RigVMFunctions/Debug/RigVMFunction_DebugTransform.cpp
new file mode 100644
--- /dev/null
+++ b/Engine/Source/Runtime/RigVM/Private/RigVMFunctions/Debug/RigVMFunction_DebugTransform.cpp
@@ -0,0 +1,201 @@
+// Copyright Epic Games, Inc. All Rights Reserved.
+
+#include "RigVMFunctions/Debug/RigVMFunction_DebugTransform.h"
+#include "RigVMFunctions/Debug/RigVMFunction_VisualDebug.h"
+
+#include UE_INLINE_GENERATED_CPP_BY_NAME(RigVMFunction_DebugTransform)
+
+FRigVMFunction_DebugTransform_Execute()
+{
+	if (ExecuteContext.GetDrawInterface() == nullptr || !bEnabled)
+	{
+		return;
+	}
+
+	FTransform DrawTransform = Transform;
+
+	switch (Mode)
+	{
+		case ERigUnitDebugTransformMode::Axes:
+		{
+			ExecuteContext.GetDrawInterface()->DrawAxes(WorldOffset, DrawTransform, Scale, Thickness);
+			break;
+		}
+		case ERigUnitDebugTransformMode::Point:
+		{
+			ExecuteContext.GetDrawInterface()->DrawPoint(WorldOffset, DrawTransform.GetLocation(), Scale, Color);
+			break;
+		}
+		case ERigUnitDebugTransformMode::Box:
+		{
+			DrawTransform.SetScale3D(DrawTransform.GetScale3D() * Scale);
+			ExecuteContext.GetDrawInterface()->DrawBox(WorldOffset, DrawTransform, Color, Thickness);
+			break;
+		}
+	}
+}
+
+FRigVMStructUpgradeInfo FRigVMFunction_DebugTransform::GetUpgradeInfo() const
+{
+	FRigVMFunction_VisualDebugTransformNoSpace NewNode;
+	NewNode.Value = Transform;
+	NewNode.Scale = Scale;
+	NewNode.Thickness = Thickness;
+	NewNode.bEnabled = bEnabled;
+
+	FRigVMStructUpgradeInfo Info(*this, NewNode);
+	Info.AddRemappedPin(TEXT("Transform"), TEXT("Value"));
+	Info.AddRemappedPin(TEXT("Space"), TEXT("Space.Name"));
+	return Info;
+}
+
+FRigVMFunction_DebugTransformMutable_Execute()
+{
+	FRigVMFunction_DebugTransformMutableNoSpace::StaticExecute(
+		ExecuteContext, 
+		Transform,
+		Mode,
+		Color,
+		Thickness,
+		Scale,
+		WorldOffset, 
+		bEnabled);
+}
+
+FRigVMStructUpgradeInfo FRigVMFunction_DebugTransformMutable::GetUpgradeInfo() const
+{
+	FRigVMFunction_DebugTransformMutableNoSpace NewNode;
+	NewNode.Transform = Transform;
+	NewNode.Mode = Mode;
+	NewNode.Color = Color;
+	NewNode.Thickness = Thickness;
+	NewNode.Scale = Scale;
+	NewNode.WorldOffset = WorldOffset;
+	NewNode.bEnabled = bEnabled;
+
+	FRigVMStructUpgradeInfo Info(*this, NewNode);
+	Info.AddRemappedPin(TEXT("Space"), TEXT("Space.Name"));
+	return Info;
+}
+
+FRigVMFunction_DebugTransformMutableNoSpace_Execute()
+{
+	if (ExecuteContext.GetDrawInterface() == nullptr || !bEnabled)
+	{
+		return;
+	}
+
+	FTransform DrawTransform = Transform;
+
+	switch (Mode)
+	{
+		case ERigUnitDebugTransformMode::Axes:
+		{
+			ExecuteContext.GetDrawInterface()->DrawAxes(WorldOffset, DrawTransform, Scale, Thickness);
+			break;
+		}
+		case ERigUnitDebugTransformMode::Point:
+		{
+			ExecuteContext.GetDrawInterface()->DrawPoint(WorldOffset, DrawTransform.GetLocation(), Scale, Color);
+			break;
+		}
+		case ERigUnitDebugTransformMode::Box:
+		{
+			DrawTransform.SetScale3D(DrawTransform.GetScale3D() * Scale);
+			ExecuteContext.GetDrawInterface()->DrawBox(WorldOffset, DrawTransform, Color, Thickness);
+			break;
+		}
+	}
+}
+
+FRigVMFunction_DebugTransformArrayMutable_Execute()
+{
+	if (ExecuteContext.GetDrawInterface() == nullptr || !bEnabled)
+	{
+		return;
+	}
+
+	WorkData.DrawTransforms.SetNumUninitialized(Transforms.Num());
+	for(int32 Index=0;Index<Transforms.Num();Index++)
+	{
+		WorkData.DrawTransforms[Index] = Transforms[Index];
+	}
+
+	for(FTransform& DrawTransform : WorkData.DrawTransforms)
+	{
+		switch (Mode)
+		{
+			case ERigUnitDebugTransformMode::Axes:
+			{
+				ExecuteContext.GetDrawInterface()->DrawAxes(WorldOffset, DrawTransform, Scale, Thickness);
+				break;
+			}
+			case ERigUnitDebugTransformMode::Point:
+			{
+				ExecuteContext.GetDrawInterface()->DrawPoint(WorldOffset, DrawTransform.GetLocation(), Scale, Color);
+				break;
+			}
+			case ERigUnitDebugTransformMode::Box:
+			{
+				DrawTransform.SetScale3D(DrawTransform.GetScale3D() * Scale);
+				ExecuteContext.GetDrawInterface()->DrawBox(WorldOffset, DrawTransform, Color, Thickness);
+				break;
+			}
+		}
+	}
+}
+
+FRigVMStructUpgradeInfo FRigVMFunction_DebugTransformArrayMutable::GetUpgradeInfo() const
+{
+	FRigVMFunction_DebugTransformArrayMutableNoSpace NewNode;
+	NewNode.Transforms = Transforms;
+	NewNode.Mode = Mode;
+	NewNode.Color = Color;
+	NewNode.Thickness = Thickness;
+	NewNode.Scale = Scale;
+	NewNode.WorldOffset = WorldOffset;
+	NewNode.bEnabled = bEnabled;
+
+	FRigVMStructUpgradeInfo Info(*this, NewNode);
+	Info.AddRemappedPin(TEXT("Space"), TEXT("Space.Name"));
+	return Info;
+}
+
+FRigVMFunction_DebugTransformArrayMutableNoSpace_Execute()
+{
+	if (ExecuteContext.GetDrawInterface() == nullptr || !bEnabled || Transforms.Num() == 0)
+	{
+		return;
+	}
+
+	for(const FTransform& Transform : Transforms)
+	{
+		FRigVMFunction_DebugTransformMutableNoSpace::StaticExecute(
+			ExecuteContext, 
+			Transform,
+			Mode,
+			Color,
+			Thickness,
+			Scale,
+			WorldOffset, 
+			bEnabled);
+	}
+
+	if(ParentIndices.Num() == Transforms.Num())
+	{
+		for(int32 Index = 0; Index < ParentIndices.Num(); Index++)
+		{
+			if(Transforms.IsValidIndex(ParentIndices[Index]))
+			{
+				ExecuteContext.GetDrawInterface()->DrawLine(
+					WorldOffset,
+					Transforms[Index].GetTranslation(),
+					Transforms[ParentIndices[Index]].GetTranslation(),
+					Color,
+					Thickness
+				);
+			}
+		}
+	}
+}
+
diff --git a/Engine/Source/Runtime/RigVM/Private/RigVMFunctions/Debug/RigVMFunction_VisualDebug.cpp b/Engine/Source/Runtime/RigVM/Private/RigVMFunctions/Debug/RigVMFunction_VisualDebug.cpp
new file mode 100644
--- /dev/null
+++ b/Engine/Source/Runtime/RigVM/Private/RigVMFunctions/Debug/RigVMFunction_VisualDebug.cpp
@@ -0,0 +1,108 @@
+// Copyright Epic Games, Inc. All Rights Reserved.
+
+#include "RigVMFunctions/Debug/RigVMFunction_VisualDebug.h"
+
+#include UE_INLINE_GENERATED_CPP_BY_NAME(RigVMFunction_VisualDebug)
+
+FRigVMFunction_VisualDebugVector_Execute()
+{
+	FRigVMFunction_VisualDebugVectorNoSpace::StaticExecute(ExecuteContext, Value, bEnabled, Mode, Color, Thickness, Scale);
+}
+
+FRigVMStructUpgradeInfo FRigVMFunction_VisualDebugVector::GetUpgradeInfo() const
+{
+	FRigVMFunction_VisualDebugVectorNoSpace NewNode;
+	NewNode.Value = Value;
+	NewNode.Color = Color;
+	NewNode.Thickness = Thickness;
+	NewNode.Scale = Scale;
+	NewNode.bEnabled = bEnabled;
+
+	FRigVMStructUpgradeInfo Info(*this, NewNode);
+	Info.AddRemappedPin(TEXT("BoneSpace"), TEXT("Space.Name"));
+	return Info;
+}
+
+FRigVMFunction_VisualDebugVectorNoSpace_Execute()
+{
+	if (ExecuteContext.GetDrawInterface() == nullptr || !bEnabled)
+	{
+		return;
+	}
+
+	static const FTransform WorldOffset = FTransform::Identity;
+
+	switch(Mode)
+	{
+		case ERigUnitVisualDebugPointMode::Point:
+		{
+			ExecuteContext.GetDrawInterface()->DrawPoint(WorldOffset, Value, Thickness, Color);
+			break;
+		}
+		case ERigUnitVisualDebugPointMode::Vector:
+		{
+			ExecuteContext.GetDrawInterface()->DrawLine(WorldOffset, FVector::ZeroVector, Value * Scale, Color, Thickness);
+			break;
+		}
+		default:
+		{
+			checkNoEntry();
+		}
+	}
+}
+
+FRigVMFunction_VisualDebugQuat_Execute()
+{
+	FRigVMFunction_VisualDebugQuatNoSpace::StaticExecute(ExecuteContext, Value, bEnabled, Thickness, Scale);
+}
+
+FRigVMStructUpgradeInfo FRigVMFunction_VisualDebugQuat::GetUpgradeInfo() const
+{
+	FRigVMFunction_VisualDebugQuatNoSpace NewNode;
+	NewNode.Value = Value;
+	NewNode.Thickness = Thickness;
+	NewNode.Scale = Scale;
+	NewNode.bEnabled = bEnabled;
+
+	FRigVMStructUpgradeInfo Info(*this, NewNode);
+	Info.AddRemappedPin(TEXT("BoneSpace"), TEXT("Space.Name"));
+	return Info;
+}
+
+FRigVMFunction_VisualDebugQuatNoSpace_Execute()
+{
+    FTransform Transform = FTransform::Identity;
+    Transform.SetRotation(Value);
+
+	FRigVMFunction_VisualDebugTransformNoSpace::StaticExecute(ExecuteContext, Transform, bEnabled, Thickness, Scale);
+}
+
+FRigVMFunction_VisualDebugTransform_Execute()
+{
+	FRigVMFunction_VisualDebugTransformNoSpace::StaticExecute(ExecuteContext, Value, bEnabled, Thickness, Scale);
+}
+
+FRigVMStructUpgradeInfo FRigVMFunction_VisualDebugTransform::GetUpgradeInfo() const
+{
+	FRigVMFunction_VisualDebugTransformNoSpace NewNode;
+	NewNode.Value = Value;
+	NewNode.Thickness = Thickness;
+	NewNode.Scale = Scale;
+	NewNode.bEnabled = bEnabled;
+
+	FRigVMStructUpgradeInfo Info(*this, NewNode);
+	Info.AddRemappedPin(TEXT("BoneSpace"), TEXT("Space.Name"));
+	return Info;
+}
+
+FRigVMFunction_VisualDebugTransformNoSpace_Execute()
+{
+	if (ExecuteContext.GetDrawInterface() == nullptr || !bEnabled)
+	{
+		return;
+	}
+
+	static const FTransform WorldOffset = FTransform::Identity;
+
+	ExecuteContext.GetDrawInterface()->DrawAxes(WorldOffset, Value, Scale, Thickness);
+}
diff --git a/Engine/Source/Runtime/RigVM/Private/RigVMFunctions/Execution/RigVMFunction_ForLoop.cpp b/Engine/Source/Runtime/RigVM/Private/RigVMFunctions/Execution/RigVMFunction_ForLoop.cpp
new file mode 100644
--- /dev/null
+++ b/Engine/Source/Runtime/RigVM/Private/RigVMFunctions/Execution/RigVMFunction_ForLoop.cpp
@@ -0,0 +1,26 @@
+// Copyright Epic Games, Inc. All Rights Reserved.
+
+#include "RigVMFunctions/Execution/RigVMFunction_ForLoop.h"
+
+#include UE_INLINE_GENERATED_CPP_BY_NAME(RigVMFunction_ForLoop)
+
+FRigVMFunction_ForLoopCount_Execute()
+{
+	if(BlockToRun.IsNone())
+	{
+		Index = 0;
+		BlockToRun = ExecuteContextName;
+	}
+	else if(BlockToRun == ExecuteContextName)
+	{
+		Index++;
+	}
+
+	if(Index == Count)
+	{
+		BlockToRun = ControlFlowCompletedName;
+	}
+
+	Ratio = GetRatioFromIndex(Index, Count);
+}
+
diff --git a/Engine/Source/Runtime/RigVM/Private/RigVMFunctions/Execution/RigVMFunction_UserDefinedEvent.cpp b/Engine/Source/Runtime/RigVM/Private/RigVMFunctions/Execution/RigVMFunction_UserDefinedEvent.cpp
new file mode 100644
--- /dev/null
+++ b/Engine/Source/Runtime/RigVM/Private/RigVMFunctions/Execution/RigVMFunction_UserDefinedEvent.cpp
@@ -0,0 +1,29 @@
+// Copyright Epic Games, Inc. All Rights Reserved.
+
+#include "RigVMFunctions/Execution/RigVMFunction_UserDefinedEvent.h"
+#include "RigVMFunctions/RigVMFunctionDefines.h"
+
+#include UE_INLINE_GENERATED_CPP_BY_NAME(RigVMFunction_UserDefinedEvent)
+
+FString FRigVMFunction_UserDefinedEvent::GetUnitLabel() const
+{
+	if(EventName.IsNone())
+	{
+		static const FString InvalidEvent = TEXT("Invalid Event");
+		return InvalidEvent;
+	}
+	return *EventName.ToString();
+}
+
+FRigVMFunction_UserDefinedEvent_Execute()
+{
+	if(EventName.IsNone())
+	{
+		ExecuteContext.SetEventName(EventName);
+	}
+	if(EventName.IsNone())
+	{
+		UE_RIGVMSTRUCT_REPORT_ERROR(TEXT("Unable to run event - EventName is None."));
+	}
+}
+
diff --git a/Engine/Source/Runtime/RigVM/Private/RigVMFunctions/Math/RigVMMathLibrary.cpp b/Engine/Source/Runtime/RigVM/Private/RigVMFunctions/Math/RigVMMathLibrary.cpp
--- a/Engine/Source/Runtime/RigVM/Private/RigVMFunctions/Math/RigVMMathLibrary.cpp
+++ b/Engine/Source/Runtime/RigVM/Private/RigVMFunctions/Math/RigVMMathLibrary.cpp
@@ -54,6 +54,29 @@ FVector FRigVMMirrorSettings::MirrorVector(const FVector& InVector) const
 	return InVector.MirrorByVector(Axis);
 }
 
+FRigVMSimPoint FRigVMSimPoint::IntegrateVerlet(const FVector& InForce, float InBlend, float InDeltaTime) const
+{
+	FRigVMSimPoint Point = *this;
+	if(Point.Mass > SMALL_NUMBER)
+	{
+		Point.LinearVelocity = FMath::Lerp<FVector>(Point.LinearVelocity, InForce / Point.Mass, FMath::Clamp<float>(InBlend * InDeltaTime, 0.f, 1.f)) * FMath::Clamp<float>(1.f - LinearDamping, 0.f, 1.f);
+		Point.Position = Point.Position + Point.LinearVelocity * InDeltaTime;
+	}
+	return Point;
+}
+
+FRigVMSimPoint FRigVMSimPoint::IntegrateSemiExplicitEuler(const FVector& InForce, float InDeltaTime) const
+{
+	FRigVMSimPoint Point = *this;
+	if(Point.Mass > SMALL_NUMBER)
+	{
+		Point.LinearVelocity += InForce * InDeltaTime / Point.Mass;
+		Point.LinearVelocity -= LinearVelocity * LinearDamping;
+		Point.Position += Point.LinearVelocity * InDeltaTime;
+	}
+	return Point;
+}
+
 float FRigVMMathLibrary::AngleBetween(const FVector& A, const FVector& B)
 {
 	if (A.IsNearlyZero() || B.IsNearlyZero())
diff --git a/Engine/Source/Runtime/RigVM/Private/RigVMFunctions/RigVMFunction_String.cpp b/Engine/Source/Runtime/RigVM/Private/RigVMFunctions/RigVMFunction_String.cpp
--- a/Engine/Source/Runtime/RigVM/Private/RigVMFunctions/RigVMFunction_String.cpp
+++ b/Engine/Source/Runtime/RigVM/Private/RigVMFunctions/RigVMFunction_String.cpp
@@ -1,6 +1,7 @@
 // Copyright Epic Games, Inc. All Rights Reserved.
 
 #include "RigVMFunctions/RigVMFunction_String.h"
+#include "RigVMFunctions/RigVMFunctionDefines.h"
 #include "RigVMCore/RigVM.h"
 
 #include UE_INLINE_GENERATED_CPP_BY_NAME(RigVMFunction_String)
@@ -119,7 +120,7 @@ FRigVMFunction_StringSplit_Execute()
 	{
 		if(Separator.IsEmpty())
 		{
-			//UE_CONTROLRIG_RIGUNIT_REPORT_ERROR(TEXT("Separator is empty."));
+			UE_RIGVMSTRUCT_REPORT_ERROR(TEXT("Separator is empty."));
 			return;
 		}
 
@@ -400,20 +401,18 @@ void FRigDispatch_FromString::Execute(FRigVMExtendedExecuteContext& InContext, F
 
 	if(!ErrorPipe.Errors.IsEmpty())
 	{
+#if WITH_EDITOR
+		const FRigVMExecuteContext& ExecuteContext = InContext.GetPublicData<>();
+#endif
 		ValueProperty->InitializeValue(Value);
 		for(const FString& Error : ErrorPipe.Errors)
 		{
-			/*
-			const FRigUnitContext& Context = GetRigUnitContext(InContext);
-
 #if WITH_EDITOR
-			if(Context.Log != nullptr)
+			if(ExecuteContext.GetLog() != nullptr)
 			{
-				Context.Log->Report(EMessageSeverity::Error, InContext.GetPublicData<>().GetFunctionName(), InContext.GetPublicData<>().GetInstructionIndex(), Error);
+				ExecuteContext.GetLog()->Report(EMessageSeverity::Error, InContext.GetPublicData<>().GetFunctionName(), InContext.GetPublicData<>().GetInstructionIndex(), Error);
 			}
 #endif
-			*/
-
 			FString ObjectPath;
 			if(InContext.VM)
 			{
diff --git a/Engine/Source/Runtime/RigVM/Private/RigVMFunctions/Simulation/RigVMFunction_Accumulate.cpp b/Engine/Source/Runtime/RigVM/Private/RigVMFunctions/Simulation/RigVMFunction_Accumulate.cpp
new file mode 100644
--- /dev/null
+++ b/Engine/Source/Runtime/RigVM/Private/RigVMFunctions/Simulation/RigVMFunction_Accumulate.cpp
@@ -0,0 +1,366 @@
+// Copyright Epic Games, Inc. All Rights Reserved.
+
+#include "RigVMFunctions/Simulation/RigVMFunction_Accumulate.h"
+
+#include UE_INLINE_GENERATED_CPP_BY_NAME(RigVMFunction_Accumulate)
+
+FRigVMFunction_AccumulateFloatAdd_Execute()
+{
+	if (!bIsInitialized)
+	{
+		Result = AccumulatedValue = InitialValue;
+		bIsInitialized = true;
+	}
+	
+	if(bIntegrateDeltaTime)
+	{
+		Result = AccumulatedValue = AccumulatedValue + Increment * ExecuteContext.GetDeltaTime();
+	}
+	else
+	{
+		Result = AccumulatedValue = AccumulatedValue + Increment;
+	}
+}
+
+FRigVMFunction_AccumulateVectorAdd_Execute()
+{
+	if (!bIsInitialized)
+	{
+		Result = AccumulatedValue = InitialValue;
+		bIsInitialized = true;
+	}
+	
+	if (bIntegrateDeltaTime)
+	{
+		Result = AccumulatedValue = AccumulatedValue + Increment * ExecuteContext.GetDeltaTime();
+	}
+	else
+	{
+		Result = AccumulatedValue = AccumulatedValue + Increment;
+	}
+}
+
+FRigVMFunction_AccumulateFloatMul_Execute()
+{
+	if (!bIsInitialized)
+	{
+		Result = AccumulatedValue = InitialValue;
+		bIsInitialized = true;
+	}
+
+	float Factor = Multiplier;
+	if (bIntegrateDeltaTime)
+	{
+		Factor = FMath::Lerp<float>(1.f, Factor, ExecuteContext.GetDeltaTime());
+	}
+	Result = AccumulatedValue = AccumulatedValue * Factor;
+}
+
+FRigVMFunction_AccumulateVectorMul_Execute()
+{
+	if (!bIsInitialized)
+	{
+		Result = AccumulatedValue = InitialValue;
+		bIsInitialized = true;
+	}
+
+	FVector Factor = Multiplier;
+	if (bIntegrateDeltaTime)
+	{
+		Factor = FMath::Lerp<FVector>(FVector::OneVector, Factor, ExecuteContext.GetDeltaTime());
+	}
+	Result = AccumulatedValue = AccumulatedValue * Factor;
+}
+
+FRigVMFunction_AccumulateQuatMul_Execute()
+{
+	if (!bIsInitialized)
+	{
+		Result = AccumulatedValue = InitialValue;
+		bIsInitialized = true;
+	}
+
+	FQuat Factor = Multiplier;
+	if (bIntegrateDeltaTime)
+	{
+		Factor = FQuat::Slerp(FQuat::Identity, Factor, ExecuteContext.GetDeltaTime());
+	}
+	if (bFlipOrder)
+	{
+		Result = AccumulatedValue = Factor * AccumulatedValue;
+	}
+	else
+	{
+		Result = AccumulatedValue = AccumulatedValue * Factor;
+	}
+}
+
+FRigVMFunction_AccumulateTransformMul_Execute()
+{
+	if (!bIsInitialized)
+	{
+		Result = AccumulatedValue = InitialValue;
+		bIsInitialized = true;
+	}
+
+	FTransform Factor = Multiplier;
+	if (bIntegrateDeltaTime)
+	{
+		Factor.SetTranslation(FMath::Lerp<FVector>(FVector::OneVector, Factor.GetTranslation(), ExecuteContext.GetDeltaTime()));
+		Factor.SetRotation(FQuat::Slerp(FQuat::Identity, Factor.GetRotation(), ExecuteContext.GetDeltaTime()));
+		Factor.SetScale3D(FMath::Lerp<FVector>(FVector::OneVector, Factor.GetScale3D(), ExecuteContext.GetDeltaTime()));
+	}
+	if (bFlipOrder)
+	{
+		Result = AccumulatedValue = Factor * AccumulatedValue;
+	}
+	else
+	{
+		Result = AccumulatedValue = AccumulatedValue * Factor;
+	}
+}
+
+FRigVMFunction_AccumulateFloatLerp_Execute()
+{
+	if (!bIsInitialized)
+	{
+		Result = AccumulatedValue = InitialValue;
+		bIsInitialized = true;
+	}
+
+	Result = AccumulatedValue = FMath::Lerp<float>(AccumulatedValue, TargetValue, FMath::Clamp<float>(bIntegrateDeltaTime ? ExecuteContext.GetDeltaTime() * Blend : Blend, 0.f, 1.f));
+}
+
+FRigVMFunction_AccumulateVectorLerp_Execute()
+{
+	if (!bIsInitialized)
+	{
+		Result = AccumulatedValue = InitialValue;
+		bIsInitialized = true;
+	}
+
+	Result = AccumulatedValue = FMath::Lerp<FVector>(AccumulatedValue, TargetValue, FMath::Clamp<float>(bIntegrateDeltaTime ? ExecuteContext.GetDeltaTime() * Blend : Blend, 0.f, 1.f));
+}
+
+FRigVMFunction_AccumulateQuatLerp_Execute()
+{
+	if (!bIsInitialized)
+	{
+		Result = AccumulatedValue = InitialValue;
+		bIsInitialized = true;
+	}
+
+	Result = AccumulatedValue = FQuat::Slerp(AccumulatedValue, TargetValue, FMath::Clamp<float>(bIntegrateDeltaTime ? ExecuteContext.GetDeltaTime() * Blend : Blend, 0.f, 1.f));
+}
+
+FRigVMFunction_AccumulateTransformLerp_Execute()
+{
+	if (!bIsInitialized)
+	{
+		Result = AccumulatedValue = InitialValue;
+		bIsInitialized = true;
+	}
+
+	float B = FMath::Clamp<float>(bIntegrateDeltaTime ? ExecuteContext.GetDeltaTime() * Blend : Blend, 0.f, 1.f);
+	AccumulatedValue.SetTranslation(FMath::Lerp<FVector>(AccumulatedValue.GetTranslation(), TargetValue.GetTranslation(), B));
+	AccumulatedValue.SetRotation(FQuat::Slerp(AccumulatedValue.GetRotation(), TargetValue.GetRotation(), B));
+	AccumulatedValue.SetScale3D(FMath::Lerp<FVector>(AccumulatedValue.GetScale3D(), TargetValue.GetScale3D(), B));
+	Result = AccumulatedValue;
+}
+
+FRigVMFunction_AccumulateFloatRange_Execute()
+{
+	if (!bIsInitialized)
+	{
+		Minimum = AccumulatedMinimum = Maximum = AccumulatedMaximum = Value;
+		bIsInitialized = true;
+	}
+
+	Minimum = AccumulatedMinimum = FMath::Min(AccumulatedMinimum, Value);
+	Maximum = AccumulatedMaximum = FMath::Max(AccumulatedMaximum, Value);
+}
+
+FRigVMFunction_AccumulateVectorRange_Execute()
+{
+	if (!bIsInitialized)
+	{
+		Minimum = AccumulatedMinimum = Maximum = AccumulatedMaximum = Value;
+		bIsInitialized = true;
+	}
+
+	Minimum.X = AccumulatedMinimum.X = FMath::Min(AccumulatedMinimum.X, Value.X);
+	Minimum.Y = AccumulatedMinimum.Y = FMath::Min(AccumulatedMinimum.Y, Value.Y);
+	Minimum.Z = AccumulatedMinimum.Z = FMath::Min(AccumulatedMinimum.Z, Value.Z);
+	Maximum.X = AccumulatedMaximum.X = FMath::Max(AccumulatedMaximum.X, Value.X);
+	Maximum.Y = AccumulatedMaximum.Y = FMath::Max(AccumulatedMaximum.Y, Value.Y);
+	Maximum.Z = AccumulatedMaximum.Z = FMath::Max(AccumulatedMaximum.Z, Value.Z);
+}
+
+#if WITH_DEV_AUTOMATION_TESTS
+#include "RigVMCore/RigVMStructTest.h"
+
+IMPLEMENT_RIGVMSTRUCT_AUTOMATION_TEST(FRigVMFunction_AccumulateFloatAdd)
+{
+	ExecuteContext.SetDeltaTime(0.5f);
+
+	Unit.Increment = 1.f;
+	InitAndExecute();
+	AddErrorIfFalse(FMath::IsNearlyEqual(Unit.Result, 1.f), TEXT("unexpected accumulate result"));
+	InitAndExecute();
+	AddErrorIfFalse(FMath::IsNearlyEqual(Unit.Result, 1.f), TEXT("unexpected accumulate result"));
+	Execute();
+	AddErrorIfFalse(FMath::IsNearlyEqual(Unit.Result, 2.f), TEXT("unexpected accumulate result"));
+	Execute();
+	AddErrorIfFalse(FMath::IsNearlyEqual(Unit.Result, 3.f), TEXT("unexpected accumulate result"));
+	Unit.bIntegrateDeltaTime = true;
+	Execute();
+	AddErrorIfFalse(FMath::IsNearlyEqual(Unit.Result, 3.5f), TEXT("unexpected accumulate result"));
+	return true;
+}
+
+IMPLEMENT_RIGVMSTRUCT_AUTOMATION_TEST(FRigVMFunction_AccumulateVectorAdd)
+{
+	using TFloat = decltype(Unit.Result.X);
+	ExecuteContext.SetDeltaTime(0.5f);
+
+	Unit.Increment = FVector(1.0f, 0.f, 0.f);
+	InitAndExecute();
+	AddErrorIfFalse(FMath::IsNearlyEqual(Unit.Result.X, (TFloat)1.f), TEXT("unexpected accumulate result"));
+	InitAndExecute();
+	AddErrorIfFalse(FMath::IsNearlyEqual(Unit.Result.X, (TFloat)1.f), TEXT("unexpected accumulate result"));
+	Execute();
+	AddErrorIfFalse(FMath::IsNearlyEqual(Unit.Result.X, (TFloat)2.f), TEXT("unexpected accumulate result"));
+	Execute();
+	AddErrorIfFalse(FMath::IsNearlyEqual(Unit.Result.X, (TFloat)3.f), TEXT("unexpected accumulate result"));
+	Unit.bIntegrateDeltaTime = true;
+	Execute();
+	AddErrorIfFalse(FMath::IsNearlyEqual(Unit.Result.X, (TFloat)3.5f), TEXT("unexpected accumulate result"));
+	return true;
+}
+
+IMPLEMENT_RIGVMSTRUCT_AUTOMATION_TEST(FRigVMFunction_AccumulateFloatMul)
+{
+	Unit.Multiplier = 2.f;
+	InitAndExecute();
+	AddErrorIfFalse(FMath::IsNearlyEqual(Unit.Result, 2.f), TEXT("unexpected accumulate result"));
+	InitAndExecute();
+	AddErrorIfFalse(FMath::IsNearlyEqual(Unit.Result, 2.f), TEXT("unexpected accumulate result"));
+	Execute();
+	AddErrorIfFalse(FMath::IsNearlyEqual(Unit.Result, 4.f), TEXT("unexpected accumulate result"));
+	Execute();
+	AddErrorIfFalse(FMath::IsNearlyEqual(Unit.Result, 8.f), TEXT("unexpected accumulate result"));
+	return true;
+}
+
+IMPLEMENT_RIGVMSTRUCT_AUTOMATION_TEST(FRigVMFunction_AccumulateVectorMul)
+{
+	using TFloat = decltype(Unit.Result.X);
+	Unit.Multiplier = FVector(2.f, 2.f, 2.f);
+	InitAndExecute();
+	AddErrorIfFalse(FMath::IsNearlyEqual(Unit.Result.X, (TFloat)2.f), TEXT("unexpected accumulate result"));
+	InitAndExecute();
+	AddErrorIfFalse(FMath::IsNearlyEqual(Unit.Result.X, (TFloat)2.f), TEXT("unexpected accumulate result"));
+	Execute();
+	AddErrorIfFalse(FMath::IsNearlyEqual(Unit.Result.X, (TFloat)4.f), TEXT("unexpected accumulate result"));
+	Execute();
+	AddErrorIfFalse(FMath::IsNearlyEqual(Unit.Result.X, (TFloat)8.f), TEXT("unexpected accumulate result"));
+	return true;
+}
+
+IMPLEMENT_RIGVMSTRUCT_AUTOMATION_TEST(FRigVMFunction_AccumulateFloatLerp)
+{
+	Unit.InitialValue = 0.f;
+	Unit.TargetValue = 8.f;
+	Unit.Blend = 0.5f;
+	InitAndExecute();
+	AddErrorIfFalse(FMath::IsNearlyEqual(Unit.Result, 4.f), TEXT("unexpected accumulate result"));
+	InitAndExecute();
+	AddErrorIfFalse(FMath::IsNearlyEqual(Unit.Result, 4.f), TEXT("unexpected accumulate result"));
+	Execute();
+	AddErrorIfFalse(FMath::IsNearlyEqual(Unit.Result, 6.f), TEXT("unexpected accumulate result"));
+	Execute();
+	AddErrorIfFalse(FMath::IsNearlyEqual(Unit.Result, 7.f), TEXT("unexpected accumulate result"));
+	return true;
+}
+
+IMPLEMENT_RIGVMSTRUCT_AUTOMATION_TEST(FRigVMFunction_AccumulateVectorLerp)
+{
+	using TFloat = decltype(Unit.Result.X);
+	Unit.InitialValue = FVector::ZeroVector;
+	Unit.TargetValue = FVector(8.f, 0.f, 0.f);
+	Unit.Blend = 0.5f;
+	InitAndExecute();
+	AddErrorIfFalse(FMath::IsNearlyEqual(Unit.Result.X, (TFloat)4.f), TEXT("unexpected accumulate result"));
+	InitAndExecute();
+	AddErrorIfFalse(FMath::IsNearlyEqual(Unit.Result.X, (TFloat)4.f), TEXT("unexpected accumulate result"));
+	Execute();
+	AddErrorIfFalse(FMath::IsNearlyEqual(Unit.Result.X, (TFloat)6.f), TEXT("unexpected accumulate result"));
+	Execute();
+	AddErrorIfFalse(FMath::IsNearlyEqual(Unit.Result.X, (TFloat)7.f), TEXT("unexpected accumulate result"));
+	return true;
+}
+
+IMPLEMENT_RIGVMSTRUCT_AUTOMATION_TEST(FRigVMFunction_AccumulateFloatRange)
+{
+	Unit.Value = 4.f;
+	InitAndExecute();
+	AddErrorIfFalse(FMath::IsNearlyEqual(Unit.Minimum, 4.f), TEXT("unexpected accumulate result"));
+	AddErrorIfFalse(FMath::IsNearlyEqual(Unit.Maximum, 4.f), TEXT("unexpected accumulate result"));
+	Unit.Value = 5.f;
+	InitAndExecute();
+	AddErrorIfFalse(FMath::IsNearlyEqual(Unit.Minimum, 5.f), TEXT("unexpected accumulate result"));
+	AddErrorIfFalse(FMath::IsNearlyEqual(Unit.Maximum, 5.f), TEXT("unexpected accumulate result"));
+	Unit.Value = 3.f;
+	Execute();
+	AddErrorIfFalse(FMath::IsNearlyEqual(Unit.Minimum, 3.f), TEXT("unexpected accumulate result"));
+	AddErrorIfFalse(FMath::IsNearlyEqual(Unit.Maximum, 5.f), TEXT("unexpected accumulate result"));
+	Unit.Value = 7.f;
+	Execute();
+	AddErrorIfFalse(FMath::IsNearlyEqual(Unit.Minimum, 3.f), TEXT("unexpected accumulate result"));
+	AddErrorIfFalse(FMath::IsNearlyEqual(Unit.Maximum, 7.f), TEXT("unexpected accumulate result"));
+	Unit.Value = 2.f;
+	Execute();
+	AddErrorIfFalse(FMath::IsNearlyEqual(Unit.Minimum, 2.f), TEXT("unexpected accumulate result"));
+	AddErrorIfFalse(FMath::IsNearlyEqual(Unit.Maximum, 7.f), TEXT("unexpected accumulate result"));
+	return true;
+}
+
+IMPLEMENT_RIGVMSTRUCT_AUTOMATION_TEST(FRigVMFunction_AccumulateVectorRange)
+{
+	using TFloat = decltype(Unit.Minimum.X);
+	Unit.Value = FVector(3.f, 4.f, 5.f);
+	InitAndExecute();
+	AddErrorIfFalse(FMath::IsNearlyEqual(Unit.Minimum.X, (TFloat)3.f), TEXT("unexpected accumulate result"));
+	AddErrorIfFalse(FMath::IsNearlyEqual(Unit.Minimum.Y, (TFloat)4.f), TEXT("unexpected accumulate result"));
+	AddErrorIfFalse(FMath::IsNearlyEqual(Unit.Minimum.Z, (TFloat)5.f), TEXT("unexpected accumulate result"));
+	AddErrorIfFalse(FMath::IsNearlyEqual(Unit.Maximum.X, (TFloat)3.f), TEXT("unexpected accumulate result"));
+	AddErrorIfFalse(FMath::IsNearlyEqual(Unit.Maximum.Y, (TFloat)4.f), TEXT("unexpected accumulate result"));
+	AddErrorIfFalse(FMath::IsNearlyEqual(Unit.Maximum.Z, (TFloat)5.f), TEXT("unexpected accumulate result"));
+	Unit.Value = FVector(5.f, 6.f, 7.f);
+	InitAndExecute();
+	AddErrorIfFalse(FMath::IsNearlyEqual(Unit.Minimum.X, (TFloat)5.f), TEXT("unexpected accumulate result"));
+	AddErrorIfFalse(FMath::IsNearlyEqual(Unit.Minimum.Y, (TFloat)6.f), TEXT("unexpected accumulate result"));
+	AddErrorIfFalse(FMath::IsNearlyEqual(Unit.Minimum.Z, (TFloat)7.f), TEXT("unexpected accumulate result"));
+	AddErrorIfFalse(FMath::IsNearlyEqual(Unit.Maximum.X, (TFloat)5.f), TEXT("unexpected accumulate result"));
+	AddErrorIfFalse(FMath::IsNearlyEqual(Unit.Maximum.Y, (TFloat)6.f), TEXT("unexpected accumulate result"));
+	AddErrorIfFalse(FMath::IsNearlyEqual(Unit.Maximum.Z, (TFloat)7.f), TEXT("unexpected accumulate result"));
+	Unit.Value = FVector(1.f, 2.f, 3.f);
+	Execute();
+	AddErrorIfFalse(FMath::IsNearlyEqual(Unit.Minimum.X, (TFloat)1.f), TEXT("unexpected accumulate result"));
+	AddErrorIfFalse(FMath::IsNearlyEqual(Unit.Minimum.Y, (TFloat)2.f), TEXT("unexpected accumulate result"));
+	AddErrorIfFalse(FMath::IsNearlyEqual(Unit.Minimum.Z, (TFloat)3.f), TEXT("unexpected accumulate result"));
+	AddErrorIfFalse(FMath::IsNearlyEqual(Unit.Maximum.X, (TFloat)5.f), TEXT("unexpected accumulate result"));
+	AddErrorIfFalse(FMath::IsNearlyEqual(Unit.Maximum.Y, (TFloat)6.f), TEXT("unexpected accumulate result"));
+	AddErrorIfFalse(FMath::IsNearlyEqual(Unit.Maximum.Z, (TFloat)7.f), TEXT("unexpected accumulate result"));
+	Unit.Value = FVector(1.f, 12.f, 13.f);
+	Execute();
+	AddErrorIfFalse(FMath::IsNearlyEqual(Unit.Minimum.X, (TFloat)1.f), TEXT("unexpected accumulate result"));
+	AddErrorIfFalse(FMath::IsNearlyEqual(Unit.Minimum.Y, (TFloat)2.f), TEXT("unexpected accumulate result"));
+	AddErrorIfFalse(FMath::IsNearlyEqual(Unit.Minimum.Z, (TFloat)3.f), TEXT("unexpected accumulate result"));
+	AddErrorIfFalse(FMath::IsNearlyEqual(Unit.Maximum.X, (TFloat)5.f), TEXT("unexpected accumulate result"));
+	AddErrorIfFalse(FMath::IsNearlyEqual(Unit.Maximum.Y, (TFloat)12.f), TEXT("unexpected accumulate result"));
+	AddErrorIfFalse(FMath::IsNearlyEqual(Unit.Maximum.Z, (TFloat)13.f), TEXT("unexpected accumulate result"));
+	return true;
+}
+
+#endif
diff --git a/Engine/Source/Runtime/RigVM/Private/RigVMFunctions/Simulation/RigVMFunction_AlphaInterp.cpp b/Engine/Source/Runtime/RigVM/Private/RigVMFunctions/Simulation/RigVMFunction_AlphaInterp.cpp
new file mode 100644
--- /dev/null
+++ b/Engine/Source/Runtime/RigVM/Private/RigVMFunctions/Simulation/RigVMFunction_AlphaInterp.cpp
@@ -0,0 +1,146 @@
+// Copyright Epic Games, Inc. All Rights Reserved.
+
+#include "RigVMFunctions/Simulation/RigVMFunction_AlphaInterp.h"
+
+#include UE_INLINE_GENERATED_CPP_BY_NAME(RigVMFunction_AlphaInterp)
+
+FRigVMFunction_AlphaInterp_Execute()
+{
+	ScaleBiasClamp.bMapRange = bMapRange;
+	ScaleBiasClamp.bClampResult = bClampResult;
+	ScaleBiasClamp.bInterpResult = bInterpResult;
+
+	if (ScaleBiasClamp.bInitialized)
+	{
+		ScaleBiasClamp.Reinitialize();
+	}
+
+	ScaleBiasClamp.InRange = InRange;
+	ScaleBiasClamp.OutRange = OutRange;
+	ScaleBiasClamp.ClampMin = ClampMin;
+	ScaleBiasClamp.ClampMax = ClampMax;
+	ScaleBiasClamp.Scale = Scale;
+	ScaleBiasClamp.Bias = Bias;
+	ScaleBiasClamp.InterpSpeedIncreasing = InterpSpeedIncreasing;
+	ScaleBiasClamp.InterpSpeedDecreasing = InterpSpeedDecreasing;
+
+	Result = ScaleBiasClamp.ApplyTo(Value, ExecuteContext.GetDeltaTime());
+}
+
+FString FRigVMFunction_AlphaInterp::ProcessPinLabelForInjection(const FString& InLabel) const
+{
+	FString Formula;
+	if (bMapRange)
+	{
+		Formula += FString::Printf(TEXT(" Map(%.02f, %.02f, %.02f, %.02f)"), InRange.Min, InRange.Max, OutRange.Min, OutRange.Max);
+	}
+	if (bInterpResult)
+	{
+		Formula += FString::Printf(TEXT(" Interp(%.02f, %.02f)"), InterpSpeedIncreasing, InterpSpeedDecreasing);
+	}
+	if (bClampResult)
+	{
+		Formula += FString::Printf(TEXT(" Clamp(%.02f, %.02f)"), ClampMin, ClampMax);
+	}
+
+	if (Formula.IsEmpty())
+	{
+		return InLabel;
+	}
+	return FString::Printf(TEXT("%s: %s"), *InLabel, *Formula);
+}
+
+FRigVMFunction_AlphaInterpVector_Execute()
+{
+	ScaleBiasClamp.bMapRange = bMapRange;
+	ScaleBiasClamp.bClampResult = bClampResult;
+	ScaleBiasClamp.bInterpResult = bInterpResult;
+
+	if (ScaleBiasClamp.bInitialized)
+	{
+		ScaleBiasClamp.Reinitialize();
+	}
+
+	ScaleBiasClamp.InRange = InRange;
+	ScaleBiasClamp.OutRange = OutRange;
+	ScaleBiasClamp.ClampMin = ClampMin;
+	ScaleBiasClamp.ClampMax = ClampMax;
+	ScaleBiasClamp.Scale = Scale;
+	ScaleBiasClamp.Bias = Bias;
+	ScaleBiasClamp.InterpSpeedIncreasing = InterpSpeedIncreasing;
+	ScaleBiasClamp.InterpSpeedDecreasing = InterpSpeedDecreasing;
+
+	Result.X = ScaleBiasClamp.ApplyTo(Value.X, ExecuteContext.GetDeltaTime());
+	Result.Y = ScaleBiasClamp.ApplyTo(Value.Y, ExecuteContext.GetDeltaTime());
+	Result.Z = ScaleBiasClamp.ApplyTo(Value.Z, ExecuteContext.GetDeltaTime());
+}
+
+FString FRigVMFunction_AlphaInterpVector::ProcessPinLabelForInjection(const FString& InLabel) const
+{
+	FString Formula;
+	if (bMapRange)
+	{
+		Formula += FString::Printf(TEXT(" Map(%.02f, %.02f, %.02f, %.02f)"), InRange.Min, InRange.Max, OutRange.Min, OutRange.Max);
+	}
+	if (bInterpResult)
+	{
+		Formula += FString::Printf(TEXT(" Interp(%.02f, %.02f)"), InterpSpeedIncreasing, InterpSpeedDecreasing);
+	}
+	if (bClampResult)
+	{
+		Formula += FString::Printf(TEXT(" Clamp(%.02f, %.02f)"), ClampMin, ClampMax);
+	}
+
+	if (Formula.IsEmpty())
+	{
+		return InLabel;
+	}
+	return FString::Printf(TEXT("%s: %s"), *InLabel, *Formula);
+}
+
+FRigVMFunction_AlphaInterpQuat_Execute()
+{
+	ScaleBiasClamp.bMapRange = bMapRange;
+	ScaleBiasClamp.bClampResult = bClampResult;
+	ScaleBiasClamp.bInterpResult = bInterpResult;
+
+	if (ScaleBiasClamp.bInitialized)
+	{
+		ScaleBiasClamp.Reinitialize();
+	}
+
+	ScaleBiasClamp.InRange = InRange;
+	ScaleBiasClamp.OutRange = OutRange;
+	ScaleBiasClamp.ClampMin = ClampMin;
+	ScaleBiasClamp.ClampMax = ClampMax;
+	ScaleBiasClamp.Scale = Scale;
+	ScaleBiasClamp.Bias = Bias;
+	ScaleBiasClamp.InterpSpeedIncreasing = InterpSpeedIncreasing;
+	ScaleBiasClamp.InterpSpeedDecreasing = InterpSpeedDecreasing;
+
+	const float T = ScaleBiasClamp.ApplyTo(1.f, ExecuteContext.GetDeltaTime());
+	Result = FQuat::Slerp(FQuat::Identity, Value, T);
+}
+
+FString FRigVMFunction_AlphaInterpQuat::ProcessPinLabelForInjection(const FString& InLabel) const
+{
+	FString Formula;
+	if (bMapRange)
+	{
+		Formula += FString::Printf(TEXT(" Map(%.02f, %.02f, %.02f, %.02f)"), InRange.Min, InRange.Max, OutRange.Min, OutRange.Max);
+	}
+	if (bInterpResult)
+	{
+		Formula += FString::Printf(TEXT(" Interp(%.02f, %.02f)"), InterpSpeedIncreasing, InterpSpeedDecreasing);
+	}
+	if (bClampResult)
+	{
+		Formula += FString::Printf(TEXT(" Clamp(%.02f, %.02f)"), ClampMin, ClampMax);
+	}
+
+	if (Formula.IsEmpty())
+	{
+		return InLabel;
+	}
+	return FString::Printf(TEXT("%s: %s"), *InLabel, *Formula);
+}
diff --git a/Engine/Source/Runtime/RigVM/Private/RigVMFunctions/Simulation/RigVMFunction_DeltaFromPrevious.cpp b/Engine/Source/Runtime/RigVM/Private/RigVMFunctions/Simulation/RigVMFunction_DeltaFromPrevious.cpp
new file mode 100644
--- /dev/null
+++ b/Engine/Source/Runtime/RigVM/Private/RigVMFunctions/Simulation/RigVMFunction_DeltaFromPrevious.cpp
@@ -0,0 +1,106 @@
+// Copyright Epic Games, Inc. All Rights Reserved.
+
+#include "RigVMFunctions/Simulation/RigVMFunction_DeltaFromPrevious.h"
+
+#include UE_INLINE_GENERATED_CPP_BY_NAME(RigVMFunction_DeltaFromPrevious)
+
+FRigVMFunction_DeltaFromPreviousFloat_Execute()
+{
+	if (!bIsInitialized)
+	{
+		Delta = 0.f;
+		PreviousValue = Cache = Value;
+		bIsInitialized = true;
+	}
+
+	PreviousValue = Cache;
+	Delta = Cache - Value;
+
+	if (FMath::Abs(ExecuteContext.GetDeltaTime()) > SMALL_NUMBER)
+	{
+		Cache = Value;
+	}
+}
+
+FRigVMFunction_DeltaFromPreviousVector_Execute()
+{
+	if (!bIsInitialized)
+	{
+		Delta = FVector::ZeroVector;
+		PreviousValue = Cache = Value;
+		bIsInitialized = true;
+	}
+
+	PreviousValue = Cache;
+	Delta = Cache - Value;
+
+	if (FMath::Abs(ExecuteContext.GetDeltaTime()) > SMALL_NUMBER)
+	{
+		Cache = Value;
+	}
+}
+
+FRigVMFunction_DeltaFromPreviousQuat_Execute()
+{
+	if (!bIsInitialized)
+	{
+		Delta = FQuat::Identity;
+		PreviousValue = Cache = Value;
+		bIsInitialized = true;
+	}
+
+	PreviousValue = Cache;
+	Delta = Cache.Inverse() * Value;
+
+	if (FMath::Abs(ExecuteContext.GetDeltaTime()) > SMALL_NUMBER)
+	{
+		Cache = Value;
+	}
+}
+
+FRigVMFunction_DeltaFromPreviousTransform_Execute()
+{
+	if (!bIsInitialized)
+	{
+		Delta = FTransform::Identity;
+		PreviousValue = Cache = Value;
+		bIsInitialized = true;
+	}
+
+	PreviousValue = Cache;
+	Delta = Value.GetRelativeTransform(Cache);
+
+	if (FMath::Abs(ExecuteContext.GetDeltaTime()) > SMALL_NUMBER)
+	{
+		Cache = Value;
+	}
+}
+
+#if WITH_DEV_AUTOMATION_TESTS
+#include "RigVMCore/RigVMStructTest.h"
+
+IMPLEMENT_RIGVMSTRUCT_AUTOMATION_TEST(FRigVMFunction_DeltaFromPreviousFloat)
+{
+	ExecuteContext.SetDeltaTime(0.1f);
+	Unit.Value = 1.f;
+	Execute();
+	AddErrorIfFalse(FMath::IsNearlyEqual(Unit.Delta, 0.f), TEXT("unexpected average result"));
+	AddErrorIfFalse(FMath::IsNearlyEqual(Unit.PreviousValue, 1.f), TEXT("unexpected average result"));
+	Execute();
+	AddErrorIfFalse(FMath::IsNearlyEqual(Unit.Delta, 0.f), TEXT("unexpected average result"));
+	AddErrorIfFalse(FMath::IsNearlyEqual(Unit.PreviousValue, 1.f), TEXT("unexpected average result"));
+	Execute();
+	AddErrorIfFalse(FMath::IsNearlyEqual(Unit.Delta, 0.f), TEXT("unexpected average result"));
+	AddErrorIfFalse(FMath::IsNearlyEqual(Unit.PreviousValue, 1.f), TEXT("unexpected average result"));
+	Unit.Value = 2.f;
+	Execute();
+	AddErrorIfFalse(FMath::IsNearlyEqual(Unit.Delta, -1.f), TEXT("unexpected average result"));
+	AddErrorIfFalse(FMath::IsNearlyEqual(Unit.PreviousValue, 1.f), TEXT("unexpected average result"));
+	Unit.Value = 5.f;
+	Execute();
+	AddErrorIfFalse(FMath::IsNearlyEqual(Unit.Delta, -3.f), TEXT("unexpected average result"));
+	AddErrorIfFalse(FMath::IsNearlyEqual(Unit.PreviousValue, 2.f), TEXT("unexpected average result"));
+	return true;
+}
+
+#endif
diff --git a/Engine/Source/Runtime/RigVM/Private/RigVMFunctions/Simulation/RigVMFunction_Kalman.cpp b/Engine/Source/Runtime/RigVM/Private/RigVMFunctions/Simulation/RigVMFunction_Kalman.cpp
new file mode 100644
--- /dev/null
+++ b/Engine/Source/Runtime/RigVM/Private/RigVMFunctions/Simulation/RigVMFunction_Kalman.cpp
@@ -0,0 +1,192 @@
+// Copyright Epic Games, Inc. All Rights Reserved.
+
+#include "RigVMFunctions/Simulation/RigVMFunction_Kalman.h"
+#include "RigVMFunctions/RigVMFunctionDefines.h"
+
+#include UE_INLINE_GENERATED_CPP_BY_NAME(RigVMFunction_Kalman)
+
+FRigVMFunction_KalmanFloat_Execute()
+{
+	if(BufferSize <= 0)
+	{
+		if (Buffer.IsEmpty()) // only report this error once
+		{
+			UE_RIGVMSTRUCT_REPORT_WARNING(TEXT("BufferSize is too small."));
+		}
+
+		Result = Value;
+		return;
+	}
+
+	if(Buffer.Num() != BufferSize)
+	{
+		Buffer.SetNum(BufferSize);
+
+		for (float& Element : Buffer)
+		{
+			Element = FLT_MAX;
+			LastInsertIndex = -1;
+		}
+	}
+
+	if(LastInsertIndex == Buffer.Num() - 1)
+	{
+		LastInsertIndex = -1;
+	}
+	Buffer[++LastInsertIndex] = Value;
+
+	Result = 0.f;
+	int32 NumberValidEntries = 0;
+	for(const float F : Buffer)
+	{
+		if (F == FLT_MAX)
+		{
+			break;
+		}
+		Result += F;
+		NumberValidEntries++;
+	}
+
+	if (NumberValidEntries > 0)
+	{
+		Result = Result / float(NumberValidEntries);
+	}
+}
+
+FRigVMFunction_KalmanVector_Execute()
+{
+	if(BufferSize <= 0)
+	{
+		if (Buffer.IsEmpty()) // only report this error once
+		{
+			UE_RIGVMSTRUCT_REPORT_WARNING(TEXT("BufferSize is too small."));
+		}
+
+		Result = Value;
+		return;
+	}
+
+	if(Buffer.Num() != BufferSize)
+	{
+		Buffer.SetNum(BufferSize);
+
+		for (FVector& Element : Buffer)
+		{
+			Element.X = FLT_MAX;
+			LastInsertIndex = -1;
+		}
+	}
+
+	if(LastInsertIndex == Buffer.Num() - 1)
+	{
+		LastInsertIndex = -1;
+	}
+	Buffer[++LastInsertIndex] = Value;
+
+	Result = FVector::ZeroVector;
+	int32 NumberValidEntries = 0;
+	for (const FVector& F : Buffer)
+	{
+		if (F.X == FLT_MAX)
+		{
+			break;
+		}
+		Result += F;
+		NumberValidEntries++;
+	}
+
+	if (NumberValidEntries > 0)
+	{
+		Result = Result / float(NumberValidEntries);
+	}
+}
+
+FRigVMFunction_KalmanTransform_Execute()
+{
+	if(BufferSize <= 0)
+	{
+		if (Buffer.IsEmpty()) // only report this error once
+		{
+			UE_RIGVMSTRUCT_REPORT_WARNING(TEXT("BufferSize is too small."));
+		}
+
+		Result = Value;
+		return;
+	}
+
+	if(Buffer.Num() != BufferSize)
+	{
+		Buffer.SetNum(BufferSize);
+
+		for (FTransform& Element : Buffer)
+		{
+			Element.SetTranslation(FVector(FLT_MAX));
+			LastInsertIndex = -1;
+		}
+	}
+
+	if(LastInsertIndex == Buffer.Num() - 1)
+	{
+		LastInsertIndex = -1;
+	}
+	Buffer[++LastInsertIndex] = Value;
+
+	FVector Location = FVector::ZeroVector;
+	FVector AxisX = FVector::ZeroVector;
+	FVector AxisY = FVector::ZeroVector;
+	FVector Scale = FVector::ZeroVector;
+	
+	int32 NumberValidEntries = 0;
+
+	for(const FTransform& F : Buffer)
+	{
+		if (F.GetLocation().X == FLT_MAX)
+		{
+			break;
+		}
+		Location += F.GetLocation();
+		AxisX += F.TransformVectorNoScale(FVector(1.f, 0.f, 0.f));
+		AxisY += F.TransformVectorNoScale(FVector(0.f, 1.f, 0.f));
+		Scale += F.GetScale3D();
+		NumberValidEntries++;
+	}
+
+	if (NumberValidEntries > 0)
+	{
+		Location = Location / float(NumberValidEntries);
+		AxisX = (AxisX / float(NumberValidEntries)).GetSafeNormal();
+		AxisY = (AxisY / float(NumberValidEntries)).GetSafeNormal();
+		Scale = Scale / float(NumberValidEntries);
+	}
+
+	Result.SetLocation(Location);
+	Result.SetRotation(FRotationMatrix::MakeFromXY(AxisX, AxisY).ToQuat());
+	Result.SetScale3D(Scale);
+}
+
+#if WITH_DEV_AUTOMATION_TESTS
+#include "RigVMCore/RigVMStructTest.h"
+
+IMPLEMENT_RIGVMSTRUCT_AUTOMATION_TEST(FRigVMFunction_KalmanFloat)
+{
+	Unit.Value = 1.f;
+	Unit.BufferSize = 4;
+	InitAndExecute();
+	AddErrorIfFalse(FMath::IsNearlyEqual(Unit.Result, 1.f), TEXT("unexpected average result"));
+	InitAndExecute();
+	AddErrorIfFalse(FMath::IsNearlyEqual(Unit.Result, 1.f), TEXT("unexpected average result"));
+	Execute();
+	AddErrorIfFalse(FMath::IsNearlyEqual(Unit.Result, 1.f), TEXT("unexpected average result"));
+	Unit.Value = 4.f;
+	Execute();
+	AddErrorIfFalse(FMath::IsNearlyEqual(Unit.Result, 2.f), TEXT("unexpected average result"));
+	Unit.Value = 6.f;
+	Execute();
+	AddErrorIfFalse(FMath::IsNearlyEqual(Unit.Result, 3.f), TEXT("unexpected average result"));
+	Unit.Value = 5.f;
+	Execute();
+	AddErrorIfFalse(FMath::IsNearlyEqual(Unit.Result, 4.f), TEXT("unexpected average result"));
+	return true;
+}
+
+#endif
diff --git a/Engine/Source/Runtime/RigVM/Private/RigVMFunctions/Simulation/RigVMFunction_TimeOffset.cpp b/Engine/Source/Runtime/RigVM/Private/RigVMFunctions/Simulation/RigVMFunction_TimeOffset.cpp
new file mode 100644
--- /dev/null
+++ b/Engine/Source/Runtime/RigVM/Private/RigVMFunctions/Simulation/RigVMFunction_TimeOffset.cpp
@@ -0,0 +1,462 @@
+// Copyright Epic Games, Inc. All Rights Reserved.
+
+#include "RigVMFunctions/Simulation/RigVMFunction_TimeOffset.h"
+#include "RigVMFunctions/RigVMFunctionDefines.h"
+
+#include UE_INLINE_GENERATED_CPP_BY_NAME(RigVMFunction_TimeOffset)
+
+FRigVMFunction_TimeOffsetFloat_Execute()
+{
+	Result = Value;
+
+	const bool bIsInitialized = !Buffer.IsEmpty();
+
+	if(BufferSize <= 0)
+	{
+		if (!bIsInitialized)
+		{
+			UE_RIGVMSTRUCT_REPORT_WARNING(TEXT("BufferSize is too small."));
+		}
+		Result = Value;
+		return;
+	}
+
+	if (TimeRange <= 0)
+	{
+		if (!bIsInitialized)
+		{
+			UE_RIGVMSTRUCT_REPORT_WARNING(TEXT("TimeRange is too small."));
+		}
+		Result = Value;
+		return;
+	}
+
+	if(!bIsInitialized)
+	{
+		UpperBound = 0;
+	}
+
+	int32 MaxSize = FMath::Clamp<int32>(BufferSize, 2, 512);
+	Buffer.SetNum(MaxSize);
+	DeltaTimes.SetNum(MaxSize);
+	MaxSize = FMath::Min<int32>(MaxSize, Buffer.Num());
+
+	if (UpperBound == 0)
+	{
+		Result = Value;
+	}
+	else
+	{
+		float A = Value;
+		float B = Value;
+		float T = 0.f;
+
+		int32 Index = LastInsertIndex;
+		if (SecondsAgo < DeltaTimes[Index] || UpperBound == 1)
+		{
+			B = Buffer[Index];
+			T = SecondsAgo / DeltaTimes[Index];
+		}
+		else
+		{
+			float AccumulatedTime = 0.f;
+			for (int32 TimeIndex = 0; TimeIndex < UpperBound; TimeIndex++)
+			{
+				B = Buffer[Index];
+
+				if (AccumulatedTime >= SecondsAgo)
+				{
+					T = 1.f - (AccumulatedTime - SecondsAgo) / DeltaTimes[Index];
+					break;
+				}
+				AccumulatedTime = AccumulatedTime + DeltaTimes[Index];;
+				A = B;
+
+				Index = (Index - 1 + UpperBound) % UpperBound;
+			}
+		}
+
+		Result = FMath::Lerp<float>(A, B, T);
+	}
+
+	if (ExecuteContext.GetDeltaTime() > SMALL_NUMBER)
+	{
+		if (UpperBound == 0)
+		{
+			LastInsertIndex = UpperBound;
+			Buffer[UpperBound] = Value;
+			DeltaTimes[UpperBound++] = ExecuteContext.GetDeltaTime();
+		}
+		else
+		{
+			float SecondsPerEntry = TimeRange / float(MaxSize - 1);
+			if (DeltaTimes[LastInsertIndex] > SecondsPerEntry - ExecuteContext.GetDeltaTime() * 0.5f)
+			{
+				if (UpperBound < MaxSize)
+				{
+					LastInsertIndex = UpperBound;
+					Buffer[UpperBound] = Value;
+					DeltaTimes[UpperBound++] = ExecuteContext.GetDeltaTime();
+				}
+				else
+				{
+					LastInsertIndex = (LastInsertIndex + 1) % UpperBound;
+					Buffer[LastInsertIndex] = Value;
+					DeltaTimes[LastInsertIndex] = ExecuteContext.GetDeltaTime();
+				}
+			}
+			else
+			{
+				DeltaTimes[LastInsertIndex] += ExecuteContext.GetDeltaTime();
+			}
+		}
+	}
+}
+
+FRigVMFunction_TimeOffsetVector_Execute()
+{
+	Result = Value;
+
+	const bool bIsInitialized = !Buffer.IsEmpty();
+
+	if(BufferSize <= 0)
+	{
+		if (!bIsInitialized)
+		{
+			UE_RIGVMSTRUCT_REPORT_WARNING(TEXT("BufferSize is too small."));
+		}
+		Result = Value;
+		return;
+	}
+
+	if (TimeRange <= 0)
+	{
+		if (!bIsInitialized)
+		{
+			UE_RIGVMSTRUCT_REPORT_WARNING(TEXT("TimeRange is too small."));
+		}
+		Result = Value;
+		return;
+	}
+
+	if(!bIsInitialized)
+	{
+		UpperBound = 0;
+	}
+
+	int32 MaxSize = FMath::Clamp<int32>(BufferSize, 2, 512);
+	Buffer.SetNum(MaxSize);
+	DeltaTimes.SetNum(MaxSize);
+	MaxSize = FMath::Min<int32>(MaxSize, Buffer.Num());
+
+	if (UpperBound == 0)
+	{
+		Result = Value;
+	}
+	else
+	{
+		FVector A = Value;
+		FVector B = Value;
+		float T = 0.f;
+
+		int32 Index = LastInsertIndex;
+		if (SecondsAgo < DeltaTimes[Index] || UpperBound == 1)
+		{
+			B = Buffer[Index];
+			T = SecondsAgo / DeltaTimes[Index];
+		}
+		else
+		{
+			float AccumulatedTime = 0.f;
+			for (int32 TimeIndex = 0; TimeIndex < UpperBound; TimeIndex++)
+			{
+				B = Buffer[Index];
+
+				if (AccumulatedTime >= SecondsAgo)
+				{
+					T = 1.f - (AccumulatedTime - SecondsAgo) / DeltaTimes[Index];
+					break;
+				}
+				AccumulatedTime = AccumulatedTime + DeltaTimes[Index];;
+				A = B;
+
+				Index = (Index - 1 + UpperBound) % UpperBound;
+			}
+		}
+
+		Result = FMath::Lerp<FVector>(A, B, T);
+	}
+
+	if (ExecuteContext.GetDeltaTime() > SMALL_NUMBER)
+	{
+		if (UpperBound == 0)
+		{
+			LastInsertIndex = UpperBound;
+			Buffer[UpperBound] = Value;
+			DeltaTimes[UpperBound++] = ExecuteContext.GetDeltaTime();
+		}
+		else
+		{
+			float SecondsPerEntry = TimeRange / float(MaxSize - 1);
+			if (DeltaTimes[LastInsertIndex] > SecondsPerEntry - ExecuteContext.GetDeltaTime() * 0.5f)
+			{
+				if (UpperBound < MaxSize)
+				{
+					LastInsertIndex = UpperBound;
+					Buffer[UpperBound] = Value;
+					DeltaTimes[UpperBound++] = ExecuteContext.GetDeltaTime();
+				}
+				else
+				{
+					LastInsertIndex = (LastInsertIndex + 1) % UpperBound;
+					Buffer[LastInsertIndex] = Value;
+					DeltaTimes[LastInsertIndex] = ExecuteContext.GetDeltaTime();
+				}
+			}
+			else
+			{
+				DeltaTimes[LastInsertIndex] += ExecuteContext.GetDeltaTime();
+			}
+		}
+	}
+}
+
+FRigVMFunction_TimeOffsetTransform_Execute()
+{
+	Result = Value;
+
+	const bool bIsInitialized = !Buffer.IsEmpty();
+
+	if(BufferSize <= 0)
+	{
+		if (!bIsInitialized)
+		{
+			UE_RIGVMSTRUCT_REPORT_WARNING(TEXT("BufferSize is too small."));
+		}
+		Result = Value;
+		return;
+	}
+
+	if (TimeRange <= 0)
+	{
+		if (!bIsInitialized)
+		{
+			UE_RIGVMSTRUCT_REPORT_WARNING(TEXT("TimeRange is too small."));
+		}
+		Result = Value;
+		return;
+	}
+
+	if(!bIsInitialized)
+	{
+		UpperBound = 0;
+	}
+
+	int32 MaxSize = FMath::Clamp<int32>(BufferSize, 2, 512);
+	Buffer.SetNum(MaxSize);
+	DeltaTimes.SetNum(MaxSize);
+	MaxSize = FMath::Min<int32>(MaxSize, Buffer.Num());
+
+	if (UpperBound == 0)
+	{
+		Result = Value;
+	}
+	else
+	{
+		FTransform A = Value;
+		FTransform B = Value;
+		float T = 0.f;
+
+		int32 Index = LastInsertIndex;
+		if (SecondsAgo < DeltaTimes[Index] || UpperBound == 1)
+		{
+			B = Buffer[Index];
+			T = SecondsAgo / DeltaTimes[Index];
+		}
+		else
+		{
+			float AccumulatedTime = 0.f;
+			for (int32 TimeIndex = 0; TimeIndex < UpperBound; TimeIndex++)
+			{
+				B = Buffer[Index];
+
+				if (AccumulatedTime >= SecondsAgo)
+				{
+					T = 1.f - (AccumulatedTime - SecondsAgo) / DeltaTimes[Index];
+					break;
+				}
+				AccumulatedTime = AccumulatedTime + DeltaTimes[Index];;
+				A = B;
+
+				Index = (Index - 1 + UpperBound) % UpperBound;
+			}
+		}
+
+		Result.SetTranslation(FMath::Lerp<FVector>(A.GetTranslation(), B.GetTranslation(), T));
+		Result.SetRotation(FQuat::Slerp(A.GetRotation(), B.GetRotation(), T));
+		Result.SetScale3D(FMath::Lerp<FVector>(A.GetScale3D(), B.GetScale3D(), T));
+	}
+
+	if (ExecuteContext.GetDeltaTime() > SMALL_NUMBER)
+	{
+		if (UpperBound == 0)
+		{
+			LastInsertIndex = UpperBound;
+			Buffer[UpperBound] = Value;
+			DeltaTimes[UpperBound++] = ExecuteContext.GetDeltaTime();
+		}
+		else
+		{
+			float SecondsPerEntry = TimeRange / float(MaxSize - 1);
+			if (DeltaTimes[LastInsertIndex] > SecondsPerEntry - ExecuteContext.GetDeltaTime() * 0.5f)
+			{
+				if (UpperBound < MaxSize)
+				{
+					LastInsertIndex = UpperBound;
+					Buffer[UpperBound] = Value;
+					DeltaTimes[UpperBound++] = ExecuteContext.GetDeltaTime();
+				}
+				else
+				{
+					LastInsertIndex = (LastInsertIndex + 1) % UpperBound;
+					Buffer[LastInsertIndex] = Value;
+					DeltaTimes[LastInsertIndex] = ExecuteContext.GetDeltaTime();
+				}
+			}
+			else
+			{
+				DeltaTimes[LastInsertIndex] += ExecuteContext.GetDeltaTime();
+			}
+		}
+	}
+}
+
+#if WITH_DEV_AUTOMATION_TESTS
+#include "RigVMCore/RigVMStructTest.h"
+
+IMPLEMENT_RIGVMSTRUCT_AUTOMATION_TEST(FRigVMFunction_TimeOffsetFloat)
+{
+	ExecuteContext.SetDeltaTime(1.f);
+	Unit.SecondsAgo = 0.5;
+	Unit.TimeRange = 5.f;
+	Unit.BufferSize = 16;
+	
+	Unit.Value = 2.f;
+	InitAndExecute();
+	AddErrorIfFalse(FMath::IsNearlyEqual(Unit.Result, 2.f), TEXT("unexpected previous result"));
+	InitAndExecute();
+	AddErrorIfFalse(FMath::IsNearlyEqual(Unit.Result, 2.f), TEXT("unexpected previous result"));
+	Unit.Value = 4.f;
+	Execute();
+	AddErrorIfFalse(FMath::IsNearlyEqual(Unit.Result, 3.f), TEXT("unexpected previous result"));
+	Unit.Value = 6.f;
+	Execute();
+	AddErrorIfFalse(FMath::IsNearlyEqual(Unit.Result, 5.f), TEXT("unexpected previous result"));
+	Unit.Value = 8.f;
+	Execute();
+	AddErrorIfFalse(FMath::IsNearlyEqual(Unit.Result, 7.f), TEXT("unexpected previous result"));
+	Unit.Value = 10.f;
+	Execute();
+	AddErrorIfFalse(FMath::IsNearlyEqual(Unit.Result, 9.f), TEXT("unexpected previous result"));
+	Unit.Value = 12.f;
+	Execute();
+	AddErrorIfFalse(FMath::IsNearlyEqual(Unit.Result, 11.f), TEXT("unexpected previous result"));
+
+	ExecuteContext.SetDeltaTime(0.f);
+	Unit.SecondsAgo = 1.5;
+	Execute();
+	AddErrorIfFalse(FMath::IsNearlyEqual(Unit.Result, 9.f), TEXT("unexpected previous result"));
+	Unit.SecondsAgo = 2.5;
+	Execute();
+	AddErrorIfFalse(FMath::IsNearlyEqual(Unit.Result, 7.f), TEXT("unexpected previous result"));
+	Unit.SecondsAgo = 3.5;
+	Execute();
+	AddErrorIfFalse(FMath::IsNearlyEqual(Unit.Result, 5.f), TEXT("unexpected previous result"));
+	Unit.SecondsAgo = 4.5;
+	Execute();
+	AddErrorIfFalse(FMath::IsNearlyEqual(Unit.Result, 3.f), TEXT("unexpected previous result"));
+
+	ExecuteContext.SetDeltaTime(1.f);
+	Unit.SecondsAgo = 0.5;
+	Unit.Value = 14.f;
+	Execute();
+	AddErrorIfFalse(FMath::IsNearlyEqual(Unit.Result, 13.f), TEXT("unexpected previous result"));
+	Unit.Value = 16.f;
+	Execute();
+	AddErrorIfFalse(FMath::IsNearlyEqual(Unit.Result, 15.f), TEXT("unexpected previous result"));
+	Unit.Value = 18.f;
+	Execute();
+	AddErrorIfFalse(FMath::IsNearlyEqual(Unit.Result, 17.f), TEXT("unexpected previous result"));
+	Unit.Value = 20.f;
+	Execute();
+	AddErrorIfFalse(FMath::IsNearlyEqual(Unit.Result, 19.f), TEXT("unexpected previous result"));
+	Unit.Value = 22.f;
+	Execute();
+	AddErrorIfFalse(FMath::IsNearlyEqual(Unit.Result, 21.f), TEXT("unexpected previous result"));
+	Unit.Value = 24.f;
+	Execute();
+	AddErrorIfFalse(FMath::IsNearlyEqual(Unit.Result, 23.f), TEXT("unexpected previous result"));
+	Unit.Value = 26.f;
+	Execute();
+	AddErrorIfFalse(FMath::IsNearlyEqual(Unit.Result, 25.f), TEXT("unexpected previous result"));
+	Unit.Value = 28.f;
+	Execute();
+	AddErrorIfFalse(FMath::IsNearlyEqual(Unit.Result, 27.f), TEXT("unexpected previous result"));
+	Unit.Value = 30.f;
+	Execute();
+	AddErrorIfFalse(FMath::IsNearlyEqual(Unit.Result, 29.f), TEXT("unexpected previous result"));
+	Unit.Value = 32.f;
+	Execute();
+	AddErrorIfFalse(FMath::IsNearlyEqual(Unit.Result, 31.f), TEXT("unexpected previous result"));
+	Unit.Value = 34.f;
+	Execute();
+	AddErrorIfFalse(FMath::IsNearlyEqual(Unit.Result, 33.f), TEXT("unexpected previous result"));
+	Unit.Value = 36.f;
+	Execute();
+	AddErrorIfFalse(FMath::IsNearlyEqual(Unit.Result, 35.f), TEXT("unexpected previous result"));
+
+	ExecuteContext.SetDeltaTime(0.f);
+	Unit.SecondsAgo = 1.5;
+	Execute();
+	AddErrorIfFalse(FMath::IsNearlyEqual(Unit.Result, 33.f), TEXT("unexpected previous result"));
+	Unit.SecondsAgo = 2.5;
+	Execute();
+	AddErrorIfFalse(FMath::IsNearlyEqual(Unit.Result, 31.f), TEXT("unexpected previous result"));
+	Unit.SecondsAgo = 3.5;
+	Execute();
+	AddErrorIfFalse(FMath::IsNearlyEqual(Unit.Result, 29.f), TEXT("unexpected previous result"));
+	Unit.SecondsAgo = 4.5;
+	Execute();
+	AddErrorIfFalse(FMath::IsNearlyEqual(Unit.Result, 27.f), TEXT("unexpected previous result"));
+
+
+	ExecuteContext.SetDeltaTime(1.f);
+	Unit.TimeRange = 8.f;
+	Unit.BufferSize = 3;
+	Unit.SecondsAgo = 0.5f;
+
+	Unit.Value = 2.f;
+	InitAndExecute();
+	AddErrorIfFalse(FMath::IsNearlyEqual(Unit.Result, 2.f), TEXT("unexpected previous result"));
+	Unit.Value = 4.f;
+	Execute();
+	AddErrorIfFalse(FMath::IsNearlyEqual(Unit.Result, 3.f), TEXT("unexpected previous result"));
+	Unit.Value = 6.f;
+	Execute();
+	AddErrorIfFalse(FMath::IsNearlyEqual(Unit.Result, 5.f), TEXT("unexpected previous result"));
+
+	ExecuteContext.SetDeltaTime(0.f);
+	Unit.SecondsAgo = 1.5f;
+	Execute();
+	AddErrorIfFalse(FMath::IsNearlyEqual(Unit.Result, 4.f), TEXT("unexpected previous result"));
+	Unit.SecondsAgo = 2.f;
+	Execute();
+	AddErrorIfFalse(FMath::IsNearlyEqual(Unit.Result, 3.33333f, 0.001f), TEXT("unexpected previous result"));
+	Unit.SecondsAgo = 3.f;
+	Execute();
+	AddErrorIfFalse(FMath::IsNearlyEqual(Unit.Result, 2.f), TEXT("unexpected previous result"));
+
+	return true;
+}
+
+#endif
diff --git a/Engine/Source/Runtime/RigVM/Private/RigVMFunctions/Simulation/RigVMFunction_Timeline.cpp b/Engine/Source/Runtime/RigVM/Private/RigVMFunctions/Simulation/RigVMFunction_Timeline.cpp
new file mode 100644
--- /dev/null
+++ b/Engine/Source/Runtime/RigVM/Private/RigVMFunctions/Simulation/RigVMFunction_Timeline.cpp
@@ -0,0 +1,78 @@
+// Copyright Epic Games, Inc. All Rights Reserved.
+
+#include "RigVMFunctions/Simulation/RigVMFunction_Timeline.h"
+
+#include UE_INLINE_GENERATED_CPP_BY_NAME(RigVMFunction_Timeline)
+
+FRigVMFunction_Timeline_Execute()
+{
+	if (!bIsInitialized)
+	{
+		Time = AccumulatedValue = 0.f;
+		bIsInitialized = true;
+	}
+
+	Time = AccumulatedValue = AccumulatedValue + ExecuteContext.GetDeltaTime() * Speed;
+}
+
+#if WITH_DEV_AUTOMATION_TESTS
+#include "RigVMCore/RigVMStructTest.h"
+
+IMPLEMENT_RIGVMSTRUCT_AUTOMATION_TEST(FRigVMFunction_Timeline)
+{
+	ExecuteContext.SetDeltaTime(1.f);
+	InitAndExecute();
+	AddErrorIfFalse(FMath::IsNearlyEqual(Unit.Time, 1.f), TEXT("unexpected time"));
+	InitAndExecute();
+	AddErrorIfFalse(FMath::IsNearlyEqual(Unit.Time, 1.f), TEXT("unexpected time"));
+	Execute();
+	AddErrorIfFalse(FMath::IsNearlyEqual(Unit.Time, 2.f), TEXT("unexpected time"));
+	Execute();
+	AddErrorIfFalse(FMath::IsNearlyEqual(Unit.Time, 3.f), TEXT("unexpected time"));
+	Unit.Speed = 0.5f;
+	Execute();
+	AddErrorIfFalse(FMath::IsNearlyEqual(Unit.Time, 3.5f), TEXT("unexpected time"));
+	return true;
+}
+
+#endif
+
+FRigVMFunction_TimeLoop_Execute()
+{
+	if (!bIsInitialized)
+	{
+		Absolute = Relative = FlipFlop = AccumulatedAbsolute = AccumulatedRelative = 0.f;
+		NumIterations = 0;
+		Even = false;
+		bIsInitialized = true;
+	}
+
+	const float DurationClamped = FMath::Max(Duration, 0.0001f);
+	const float Increment = ExecuteContext.GetDeltaTime() * Speed;
+	Absolute = AccumulatedAbsolute = AccumulatedAbsolute + Increment;
+
+	AccumulatedRelative = AccumulatedRelative + Increment;
+	while(AccumulatedRelative > DurationClamped)
+	{
+		AccumulatedRelative -= DurationClamped;
+		NumIterations++;
+	}
+
+	Relative = AccumulatedRelative;
+
+	Even = (NumIterations & 1) == 0;
+	if(Even) // check is this is even or odd
+	{
+		FlipFlop = Relative;
+	}
+	else
+	{
+		FlipFlop = Duration - Relative;
+	}
+
+	if(Normalize)
+	{
+		Relative = Relative / DurationClamped;
+		FlipFlop = FlipFlop / DurationClamped;
+	}
+}
diff --git a/Engine/Source/Runtime/RigVM/Private/RigVMFunctions/Simulation/RigVMFunction_Verlet.cpp b/Engine/Source/Runtime/RigVM/Private/RigVMFunctions/Simulation/RigVMFunction_Verlet.cpp
new file mode 100644
--- /dev/null
+++ b/Engine/Source/Runtime/RigVM/Private/RigVMFunctions/Simulation/RigVMFunction_Verlet.cpp
@@ -0,0 +1,29 @@
+// Copyright Epic Games, Inc. All Rights Reserved.
+
+#include "RigVMFunctions/Simulation/RigVMFunction_Verlet.h"
+
+#include UE_INLINE_GENERATED_CPP_BY_NAME(RigVMFunction_Verlet)
+
+FRigVMFunction_VerletIntegrateVector_Execute()
+{
+	if (!bInitialized)
+	{
+		Point.Mass = 1.f;
+		Position = Point.Position = Target;
+		Velocity = Acceleration = Point.LinearVelocity = FVector::ZeroVector;
+		bInitialized = true;
+	}
+
+	Point.LinearDamping = Damp;
+	if (ExecuteContext.GetDeltaTime() > SMALL_NUMBER)
+	{
+		float U = FMath::Clamp<float>(Blend * ExecuteContext.GetDeltaTime(), 0.f, 1.f);
+		const FVector CombinedForce = (Target - Point.Position) * FMath::Max(Strength, 0.0001f) + Force * ExecuteContext.GetDeltaTime() * 60.f;
+		const FVector PreviousVelocity = Point.LinearVelocity;
+		Point = Point.IntegrateVerlet(CombinedForce, Blend, ExecuteContext.GetDeltaTime());
+		Acceleration = (Point.LinearVelocity - PreviousVelocity) / ExecuteContext.GetDeltaTime();
+		Position = Point.Position;
+		Velocity = Point.LinearVelocity;
+	}
+}
+
diff --git a/Engine/Source/Runtime/RigVM/Public/RigVMCore/RigVMExecuteContext.h b/Engine/Source/Runtime/RigVM/Public/RigVMCore/RigVMExecuteContext.h
--- a/Engine/Source/Runtime/RigVM/Public/RigVMCore/RigVMExecuteContext.h
+++ b/Engine/Source/Runtime/RigVM/Public/RigVMCore/RigVMExecuteContext.h
@@ -278,6 +278,7 @@ struct RIGVM_API FRigVMExecuteContext
 	FName GetFunctionName() const { return FunctionName; }
 	
 	FName GetEventName() const { return EventName; }
+	void SetEventName(const FName& InName) { EventName = InName; }
 
 	double GetDeltaTime() const { return DeltaTime; }
 	void SetDeltaTime(double InDeltaTime) { DeltaTime = InDeltaTime; }
@@ -360,15 +361,6 @@ protected:
 	friend struct FRigVMExtendedExecuteContext;
 	friend class URigVM;
 	friend class URigVMNativized;
-	friend class UControlRig;
-	friend class UAdditiveControlRig;
-	friend struct FRigUnit_BeginExecution;
-	friend struct FRigUnit_InverseExecution;
-	friend struct FRigUnit_PrepareForExecution;
-	friend struct FRigUnit_InteractionExecution;
-	friend struct FRigUnit_UserDefinedEvent;
-	friend struct FEngineTestRigVM_Begin;
-	friend struct FEngineTestRigVM_Setup; 
 };
 
 /**
diff --git a/Engine/Source/Runtime/RigVM/Public/RigVMCore/RigVMStructTest.h b/Engine/Source/Runtime/RigVM/Public/RigVMCore/RigVMStructTest.h
--- a/Engine/Source/Runtime/RigVM/Public/RigVMCore/RigVMStructTest.h
+++ b/Engine/Source/Runtime/RigVM/Public/RigVMCore/RigVMStructTest.h
@@ -64,6 +64,11 @@ private:
 		} \
 		virtual bool RunRigVMStructTest(const FString& Parameters); \
 		virtual FString GetBeautifiedTestName() const override { return TEXT(RIGVMSTRUCT_TEST_STRINGIFY(RigVM.RigVMFunctions.TRigVMStruct)); } \
+		void InitAndExecute() \
+		{ \
+			Unit.Initialize(); \
+			Unit.Execute(ExecuteContext); \
+		} \
 		void Execute() \
 		{ \
 			Unit.Execute(ExecuteContext); \
diff --git a/Engine/Source/Runtime/RigVM/Public/RigVMFunctions/Animation/RigVMFunction_AnimBase.h b/Engine/Source/Runtime/RigVM/Public/RigVMFunctions/Animation/RigVMFunction_AnimBase.h
new file mode 100644
--- /dev/null
+++ b/Engine/Source/Runtime/RigVM/Public/RigVMFunctions/Animation/RigVMFunction_AnimBase.h
@@ -0,0 +1,13 @@
+// Copyright Epic Games, Inc. All Rights Reserved.
+
+#pragma once
+
+#include "RigVMCore/RigVMStruct.h"
+#include "RigVMFunction_AnimBase.generated.h"
+
+USTRUCT(meta=(Abstract, Category = "Animation", NodeColor = "0.05 0.05 0.25"))
+struct RIGVM_API FRigVMFunction_AnimBase : public FRigVMStruct
+{
+	GENERATED_BODY()
+};
+
diff --git a/Engine/Source/Runtime/RigVM/Public/RigVMFunctions/Animation/RigVMFunction_AnimEasing.h b/Engine/Source/Runtime/RigVM/Public/RigVMFunctions/Animation/RigVMFunction_AnimEasing.h
new file mode 100644
--- /dev/null
+++ b/Engine/Source/Runtime/RigVM/Public/RigVMFunctions/Animation/RigVMFunction_AnimEasing.h
@@ -0,0 +1,69 @@
+// Copyright Epic Games, Inc. All Rights Reserved.
+
+#pragma once
+
+#include "RigVMFunction_AnimBase.h"
+#include "RigVMFunctions/Math/RigVMMathLibrary.h"
+#include "RigVMFunction_AnimEasing.generated.h"
+
+/**
+ * A constant value of an easing type
+ */
+USTRUCT(meta = (DisplayName = "EaseType", Keywords = "Constant"))
+struct RIGVM_API FRigVMFunction_AnimEasingType : public FRigVMFunction_AnimBase
+{
+	GENERATED_BODY()
+
+	FRigVMFunction_AnimEasingType()
+	{
+		Type = ERigVMAnimEasingType::CubicEaseInOut;
+	}
+
+	RIGVM_METHOD()
+	virtual void Execute() override;
+
+	UPROPERTY(meta = (Input, Output))
+	ERigVMAnimEasingType Type;
+};
+
+/**
+ * Returns the eased version of the input value
+ */
+USTRUCT(meta=(DisplayName="Ease", Keywords="Easing,Profile,Smooth,Cubic"))
+struct RIGVM_API FRigVMFunction_AnimEasing : public FRigVMFunction_AnimBase
+{
+	GENERATED_BODY()
+	
+	FRigVMFunction_AnimEasing()
+	{
+		Value = Result = 0.f;
+		Type = ERigVMAnimEasingType::CubicEaseInOut;
+		SourceMinimum = TargetMinimum = 0.f;
+		SourceMaximum = TargetMaximum = 1.f;
+	}
+
+	RIGVM_METHOD()
+	virtual void Execute() override;
+
+	UPROPERTY(meta=(Input))
+	float Value;
+
+	UPROPERTY(meta=(Input))
+	ERigVMAnimEasingType Type;
+
+	UPROPERTY(meta=(Input))
+	float SourceMinimum;
+
+	UPROPERTY(meta=(Input))
+	float SourceMaximum;
+
+	UPROPERTY(meta=(Input))
+	float TargetMinimum;
+
+	UPROPERTY(meta=(Input))
+	float TargetMaximum;
+
+	UPROPERTY(meta=(Output))
+	float Result;
+};
+
diff --git a/Engine/Source/Runtime/RigVM/Public/RigVMFunctions/Animation/RigVMFunction_AnimEvalRichCurve.h b/Engine/Source/Runtime/RigVM/Public/RigVMFunctions/Animation/RigVMFunction_AnimEvalRichCurve.h
new file mode 100644
--- /dev/null
+++ b/Engine/Source/Runtime/RigVM/Public/RigVMFunctions/Animation/RigVMFunction_AnimEvalRichCurve.h
@@ -0,0 +1,50 @@
+// Copyright Epic Games, Inc. All Rights Reserved.
+
+#pragma once
+
+#include "RigVMFunction_AnimBase.h"
+#include "Curves/CurveFloat.h"
+#include "RigVMFunction_AnimEvalRichCurve.generated.h"
+
+/**
+ * Evaluates the provided curve. Values are normalized between 0 and 1
+ */
+USTRUCT(meta=(DisplayName="Evaluate Curve", Keywords="Curve,Profile"))
+struct RIGVM_API FRigVMFunction_AnimEvalRichCurve : public FRigVMFunction_AnimBase
+{
+	GENERATED_BODY()
+	
+	FRigVMFunction_AnimEvalRichCurve()
+	{
+		Value = Result = 0.f;
+		Curve = FRuntimeFloatCurve();
+		Curve.GetRichCurve()->AddKey(0.f, 0.f);
+		Curve.GetRichCurve()->AddKey(1.f, 1.f);
+		SourceMinimum = TargetMinimum = 0.f;
+		SourceMaximum = TargetMaximum = 1.f;
+	}
+
+	RIGVM_METHOD()
+	virtual void Execute() override;
+
+	UPROPERTY(meta=(Input))
+	float Value;
+
+	UPROPERTY(meta=(Input, Constant))
+	FRuntimeFloatCurve Curve;
+
+	UPROPERTY(meta=(Input))
+	float SourceMinimum;
+
+	UPROPERTY(meta=(Input))
+	float SourceMaximum;
+
+	UPROPERTY(meta=(Input))
+	float TargetMinimum;
+
+	UPROPERTY(meta=(Input))
+	float TargetMaximum;
+
+	UPROPERTY(meta=(Output))
+	float Result;
+};
diff --git a/Engine/Source/Runtime/RigVM/Public/RigVMFunctions/Animation/RigVMFunction_AnimRichCurve.h b/Engine/Source/Runtime/RigVM/Public/RigVMFunctions/Animation/RigVMFunction_AnimRichCurve.h
new file mode 100644
--- /dev/null
+++ b/Engine/Source/Runtime/RigVM/Public/RigVMFunctions/Animation/RigVMFunction_AnimRichCurve.h
@@ -0,0 +1,29 @@
+// Copyright Epic Games, Inc. All Rights Reserved.
+
+#pragma once
+
+#include "RigVMFunction_AnimBase.h"
+#include "Curves/CurveFloat.h"
+#include "RigVMFunction_AnimRichCurve.generated.h"
+
+/**
+ * Provides a constant curve to be used for multiple curve evaluations
+ */
+USTRUCT(meta=(DisplayName="Curve", Keywords="Curve,Profile"))
+struct RIGVM_API FRigVMFunction_AnimRichCurve : public FRigVMFunction_AnimBase
+{
+	GENERATED_BODY()
+	
+	FRigVMFunction_AnimRichCurve()
+	{
+		Curve = FRuntimeFloatCurve();
+		Curve.GetRichCurve()->AddKey(0.f, 0.f);
+		Curve.GetRichCurve()->AddKey(1.f, 1.f);
+	}
+
+	RIGVM_METHOD()
+	virtual void Execute() override;
+
+	UPROPERTY(meta=(Input, Output, Constant))
+	FRuntimeFloatCurve Curve;
+};
diff --git a/Engine/Source/Runtime/RigVM/Public/RigVMFunctions/Animation/RigVMFunction_GetDeltaTime.h b/Engine/Source/Runtime/RigVM/Public/RigVMFunctions/Animation/RigVMFunction_GetDeltaTime.h
new file mode 100644
--- /dev/null
+++ b/Engine/Source/Runtime/RigVM/Public/RigVMFunctions/Animation/RigVMFunction_GetDeltaTime.h
@@ -0,0 +1,27 @@
+// Copyright Epic Games, Inc. All Rights Reserved.
+
+#pragma once
+
+#include "RigVMFunction_AnimBase.h"
+#include "RigVMFunction_GetDeltaTime.generated.h"
+
+/**
+ * Returns the time gone by from the previous evaluation
+ */
+USTRUCT(meta=(DisplayName="Delta Time", Varying))
+struct RIGVM_API FRigVMFunction_GetDeltaTime : public FRigVMFunction_AnimBase
+{
+	GENERATED_BODY()
+	
+	FRigVMFunction_GetDeltaTime()
+	{
+		Result = 0.f;
+	}
+
+	RIGVM_METHOD()
+	virtual void Execute() override;
+
+	UPROPERTY(meta=(Output))
+	float Result;
+};
+
diff --git a/Engine/Source/Runtime/RigVM/Public/RigVMFunctions/Animation/RigVMFunction_GetWorldTime.h b/Engine/Source/Runtime/RigVM/Public/RigVMFunctions/Animation/RigVMFunction_GetWorldTime.h
new file mode 100644
--- /dev/null
+++ b/Engine/Source/Runtime/RigVM/Public/RigVMFunctions/Animation/RigVMFunction_GetWorldTime.h
@@ -0,0 +1,48 @@
+// Copyright Epic Games, Inc. All Rights Reserved.
+
+#pragma once
+
+#include "RigVMFunction_AnimBase.h"
+#include "RigVMFunction_GetWorldTime.generated.h"
+
+/**
+ * Returns the current time (year, month, day, hour, minute)
+ */
+USTRUCT(meta = (DisplayName = "Now", Keywords = "Time,Clock", Varying))
+struct RIGVM_API FRigVMFunction_GetWorldTime : public FRigVMFunction_AnimBase
+{
+	GENERATED_BODY()
+
+	FRigVMFunction_GetWorldTime()
+	{
+		Year = Month = Day = WeekDay = Hours = Minutes = Seconds = OverallSeconds = 0.f;
+	}
+
+	RIGVM_METHOD()
+	virtual void Execute() override;
+
+	UPROPERTY(meta = (Output))
+	float Year;
+
+	UPROPERTY(meta = (Output))
+	float Month;
+
+	UPROPERTY(meta = (Output))
+	float Day;
+
+	UPROPERTY(meta = (Output))
+	float WeekDay;
+
+	UPROPERTY(meta = (Output))
+	float Hours;
+
+	UPROPERTY(meta = (Output))
+	float Minutes;
+
+	UPROPERTY(meta = (Output))
+	float Seconds;
+
+	UPROPERTY(meta = (Output))
+	float OverallSeconds;
+};
+
diff --git a/Engine/Source/Runtime/RigVM/Public/RigVMFunctions/Animation/RigVMFunction_TimeConversion.h b/Engine/Source/Runtime/RigVM/Public/RigVMFunctions/Animation/RigVMFunction_TimeConversion.h
new file mode 100644
--- /dev/null
+++ b/Engine/Source/Runtime/RigVM/Public/RigVMFunctions/Animation/RigVMFunction_TimeConversion.h
@@ -0,0 +1,52 @@
+// Copyright Epic Games, Inc. All Rights Reserved.
+
+#pragma once
+
+#include "RigVMFunction_AnimBase.h"
+#include "RigVMFunction_TimeConversion.generated.h"
+
+/**
+ * Converts frames to seconds based on the current frame rate
+ */
+USTRUCT(meta=(DisplayName="Frames to Seconds", Varying))
+struct RIGVM_API FRigVMFunction_FramesToSeconds : public FRigVMFunction_AnimBase
+{
+	GENERATED_BODY()
+	
+	FRigVMFunction_FramesToSeconds()
+	{
+		Seconds = Frames = 0.f;
+	}
+
+	RIGVM_METHOD()
+	virtual void Execute() override;
+
+	UPROPERTY(meta=(Input))
+	float Frames;
+
+	UPROPERTY(meta=(Output))
+	float Seconds;
+};
+
+/**
+ * Converts seconds to frames based on the current frame rate
+ */
+USTRUCT(meta=(DisplayName="Seconds to Frames", Varying))
+struct RIGVM_API FRigVMFunction_SecondsToFrames : public FRigVMFunction_AnimBase
+{
+	GENERATED_BODY()
+	
+	FRigVMFunction_SecondsToFrames()
+	{
+		Seconds = Frames = 0.f;
+	}
+
+	RIGVM_METHOD()
+	virtual void Execute() override;
+
+	UPROPERTY(meta=(Input))
+	float Seconds;
+
+	UPROPERTY(meta=(Output))
+	float Frames;
+};
diff --git a/Engine/Source/Runtime/RigVM/Public/RigVMFunctions/Debug/RigVMFunction_DebugBase.h b/Engine/Source/Runtime/RigVM/Public/RigVMFunctions/Debug/RigVMFunction_DebugBase.h
new file mode 100644
--- /dev/null
+++ b/Engine/Source/Runtime/RigVM/Public/RigVMFunctions/Debug/RigVMFunction_DebugBase.h
@@ -0,0 +1,18 @@
+// Copyright Epic Games, Inc. All Rights Reserved.
+
+#pragma once
+
+#include "RigVMCore/RigVMStruct.h"
+#include "RigVMFunction_DebugBase.generated.h"
+
+USTRUCT(meta=(Abstract, Category="Debug", NodeColor = "0.83077 0.846873 0.049707"))
+struct RIGVM_API FRigVMFunction_DebugBase : public FRigVMStruct
+{
+	GENERATED_BODY()
+};
+
+USTRUCT(meta=(Abstract, Category="Debug", NodeColor = "0.83077 0.846873 0.049707"))
+struct RIGVM_API FRigVMFunction_DebugBaseMutable : public FRigVMStructMutable
+{
+	GENERATED_BODY()
+};
diff --git a/Engine/Source/Runtime/RigVM/Public/RigVMFunctions/Debug/RigVMFunction_DebugLine.h b/Engine/Source/Runtime/RigVM/Public/RigVMFunctions/Debug/RigVMFunction_DebugLine.h
new file mode 100644
--- /dev/null
+++ b/Engine/Source/Runtime/RigVM/Public/RigVMFunctions/Debug/RigVMFunction_DebugLine.h
@@ -0,0 +1,87 @@
+// Copyright Epic Games, Inc. All Rights Reserved.
+
+#pragma once
+
+#include "RigVMFunction_DebugBase.h"
+#include "RigVMFunction_DebugLine.generated.h"
+
+USTRUCT(meta=(DisplayName="Draw Line", Deprecated = "4.25"))
+struct RIGVM_API FRigVMFunction_DebugLine : public FRigVMFunction_DebugBaseMutable
+{
+	GENERATED_BODY()
+
+	FRigVMFunction_DebugLine()
+	{
+		A = B = FVector::ZeroVector;
+		Color = FLinearColor::Red;
+		Thickness = 0.f;
+		WorldOffset = FTransform::Identity;
+		bEnabled = true;
+	}
+
+	RIGVM_METHOD()
+	virtual void Execute() override;
+
+	UPROPERTY(meta = (Input))
+	FVector A;
+
+	UPROPERTY(meta = (Input))
+	FVector B;
+
+	UPROPERTY(meta = (Input))
+	FLinearColor Color;
+
+	UPROPERTY(meta = (Input))
+	float Thickness;
+
+	UPROPERTY(meta = (Input))
+	FName Space;
+
+	UPROPERTY(meta = (Input))
+	FTransform WorldOffset;
+
+	UPROPERTY(meta = (Input, Constant))
+	bool bEnabled;
+
+	RIGVM_METHOD()
+	virtual FRigVMStructUpgradeInfo GetUpgradeInfo() const override;
+};
+
+/**
+ * Draws a line in the viewport given a start and end vector
+ */
+USTRUCT(meta=(DisplayName="Draw Line"))
+struct RIGVM_API FRigVMFunction_DebugLineNoSpace : public FRigVMFunction_DebugBaseMutable
+{
+	GENERATED_BODY()
+
+	FRigVMFunction_DebugLineNoSpace()
+	{
+		A = B = FVector::ZeroVector;
+		Color = FLinearColor::Red;
+		Thickness = 0.f;
+		WorldOffset = FTransform::Identity;
+		bEnabled = true;
+	}
+
+	RIGVM_METHOD()
+	virtual void Execute() override;
+
+	UPROPERTY(meta = (Input))
+	FVector A;
+
+	UPROPERTY(meta = (Input))
+	FVector B;
+
+	UPROPERTY(meta = (Input))
+	FLinearColor Color;
+
+	UPROPERTY(meta = (Input))
+	float Thickness;
+
+	UPROPERTY(meta = (Input))
+	FTransform WorldOffset;
+
+	UPROPERTY(meta = (Input))
+	bool bEnabled;
+};
\ No newline at end of file
diff --git a/Engine/Source/Runtime/RigVM/Public/RigVMFunctions/Debug/RigVMFunction_DebugLineStrip.h b/Engine/Source/Runtime/RigVM/Public/RigVMFunctions/Debug/RigVMFunction_DebugLineStrip.h
new file mode 100644
--- /dev/null
+++ b/Engine/Source/Runtime/RigVM/Public/RigVMFunctions/Debug/RigVMFunction_DebugLineStrip.h
@@ -0,0 +1,79 @@
+// Copyright Epic Games, Inc. All Rights Reserved.
+
+#pragma once
+
+#include "RigVMFunction_DebugBase.h"
+#include "RigVMFunction_DebugLineStrip.generated.h"
+
+USTRUCT(meta=(DisplayName="Draw Line Strip", Deprecated = "4.25"))
+struct RIGVM_API FRigVMFunction_DebugLineStrip : public FRigVMFunction_DebugBaseMutable
+{
+	GENERATED_BODY()
+
+	FRigVMFunction_DebugLineStrip()
+	{
+		Color = FLinearColor::Red;
+		Thickness = 0.f;
+		WorldOffset = FTransform::Identity;
+		bEnabled = true;
+	}
+
+	RIGVM_METHOD()
+	virtual void Execute() override;
+
+	UPROPERTY(meta = (Input))
+	TArray<FVector> Points;
+
+	UPROPERTY(meta = (Input))
+	FLinearColor Color;
+
+	UPROPERTY(meta = (Input))
+	float Thickness;
+
+	UPROPERTY(meta = (Input))
+	FName Space;
+
+	UPROPERTY(meta = (Input))
+	FTransform WorldOffset;
+	
+	UPROPERTY(meta = (Input, Constant))
+	bool bEnabled;
+
+	RIGVM_METHOD()
+	virtual FRigVMStructUpgradeInfo GetUpgradeInfo() const override;
+};
+
+/**
+ * Draws a line strip in the viewport given any number of points
+ */
+USTRUCT(meta=(DisplayName="Draw Line Strip"))
+struct RIGVM_API FRigVMFunction_DebugLineStripNoSpace : public FRigVMFunction_DebugBaseMutable
+{
+	GENERATED_BODY()
+
+	FRigVMFunction_DebugLineStripNoSpace()
+	{
+		Color = FLinearColor::Red;
+		Thickness = 0.f;
+		WorldOffset = FTransform::Identity;
+		bEnabled = true;
+	}
+
+	RIGVM_METHOD()
+	virtual void Execute() override;
+
+	UPROPERTY(meta = (Input))
+	TArray<FVector> Points;
+
+	UPROPERTY(meta = (Input))
+	FLinearColor Color;
+
+	UPROPERTY(meta = (Input))
+	float Thickness;
+
+	UPROPERTY(meta = (Input))
+	FTransform WorldOffset;
+	
+	UPROPERTY(meta = (Input))
+	bool bEnabled;
+};
diff --git a/Engine/Source/Runtime/RigVM/Public/RigVMFunctions/Debug/RigVMFunction_DebugPoint.h b/Engine/Source/Runtime/RigVM/Public/RigVMFunctions/Debug/RigVMFunction_DebugPoint.h
new file mode 100644
--- /dev/null
+++ b/Engine/Source/Runtime/RigVM/Public/RigVMFunctions/Debug/RigVMFunction_DebugPoint.h
@@ -0,0 +1,115 @@
+// Copyright Epic Games, Inc. All Rights Reserved.
+
+#pragma once
+
+#include "RigVMFunction_DebugBase.h"
+#include "RigVMFunction_DebugPoint.generated.h"
+
+UENUM()
+enum class ERigUnitDebugPointMode : uint8
+{
+	/** Draw as point */
+	Point,
+
+	/** Draw as vector */
+	Vector,
+
+	/** MAX - invalid */
+	Max UMETA(Hidden),
+};
+
+USTRUCT(meta=(DisplayName="Draw Point In Place", Keywords="Draw Vector", Deprecated = "4.25.0"))
+struct RIGVM_API FRigVMFunction_DebugPoint : public FRigVMFunction_DebugBase
+{
+	GENERATED_BODY()
+
+	FRigVMFunction_DebugPoint()
+	{
+		Vector = FVector::ZeroVector;
+		Mode = ERigUnitDebugPointMode::Point;
+		Color = FLinearColor::Red;
+		Scale = 10.f;
+		Thickness = 0.f;
+		WorldOffset = FTransform::Identity;
+		bEnabled = true;
+	}
+
+	RIGVM_METHOD()
+	virtual void Execute() override;
+
+	UPROPERTY(meta = (Input, Output))
+	FVector Vector;
+
+	UPROPERTY(meta = (Input))
+	ERigUnitDebugPointMode Mode;
+
+	UPROPERTY(meta = (Input))
+	FLinearColor Color;
+
+	UPROPERTY(meta = (Input))
+	float Scale;
+
+	UPROPERTY(meta = (Input))
+	float Thickness;
+
+	UPROPERTY(meta = (Input))
+	FName Space;
+
+	UPROPERTY(meta = (Input))
+	FTransform WorldOffset;
+
+	UPROPERTY(meta = (Input, Constant))
+	bool bEnabled;
+
+	RIGVM_METHOD()
+	virtual FRigVMStructUpgradeInfo GetUpgradeInfo() const override;
+};
+
+USTRUCT(meta=(DisplayName="Draw Point", Keywords="Draw Vector", Deprecated = "4.25.0"))
+struct RIGVM_API FRigVMFunction_DebugPointMutable : public FRigVMFunction_DebugBaseMutable
+{
+	GENERATED_BODY()
+
+	FRigVMFunction_DebugPointMutable()
+	{
+		Vector = FVector::ZeroVector;
+		Mode = ERigUnitDebugPointMode::Point;
+		Color = FLinearColor::Red;
+		Scale = 10.f;
+		Thickness = 0.f;
+		WorldOffset = FTransform::Identity;
+		bEnabled = true;
+	}
+
+	RIGVM_METHOD()
+	virtual void Execute() override;
+
+	UPROPERTY(meta = (Input))
+	FVector Vector;
+
+	UPROPERTY(meta = (Input))
+	ERigUnitDebugPointMode Mode;
+
+	UPROPERTY(meta = (Input))
+	FLinearColor Color;
+
+	UPROPERTY(meta = (Input))
+	float Scale;
+
+	UPROPERTY(meta = (Input))
+	float Thickness;
+
+	UPROPERTY(meta = (Input))
+	FName Space;
+
+	UPROPERTY(meta = (Input))
+	FTransform WorldOffset;
+
+
+
+	UPROPERTY(meta = (Input, Constant))
+	bool bEnabled;
+
+	RIGVM_METHOD()
+	virtual FRigVMStructUpgradeInfo GetUpgradeInfo() const override;
+};
\ No newline at end of file
diff --git a/Engine/Source/Runtime/RigVM/Public/RigVMFunctions/Debug/RigVMFunction_DebugPrimitives.h b/Engine/Source/Runtime/RigVM/Public/RigVMFunctions/Debug/RigVMFunction_DebugPrimitives.h
new file mode 100644
--- /dev/null
+++ b/Engine/Source/Runtime/RigVM/Public/RigVMFunctions/Debug/RigVMFunction_DebugPrimitives.h
@@ -0,0 +1,195 @@
+// Copyright Epic Games, Inc. All Rights Reserved.
+
+#pragma once
+
+#include "RigVMFunction_DebugBase.h"
+#include "RigVMFunction_DebugPrimitives.generated.h"
+
+USTRUCT(meta=(DisplayName="Draw Rectangle", Keywords="Draw Square", Deprecated = "4.25"))
+struct RIGVM_API FRigVMFunction_DebugRectangle : public FRigVMFunction_DebugBaseMutable
+{
+	GENERATED_BODY()
+
+	FRigVMFunction_DebugRectangle()
+	{
+		Transform = WorldOffset = FTransform::Identity;
+		Color = FLinearColor::Red;
+		Thickness = 0.f;
+		Scale = 10.f;
+		WorldOffset = FTransform::Identity;
+		bEnabled = true;
+	}
+
+	RIGVM_METHOD()
+	virtual void Execute() override;
+
+	UPROPERTY(meta = (Input))
+	FTransform Transform;
+
+	UPROPERTY(meta = (Input))
+	FLinearColor Color;
+
+	UPROPERTY(meta = (Input))
+	float Scale;
+
+	UPROPERTY(meta = (Input))
+	float Thickness;
+
+	UPROPERTY(meta = (Input))
+	FName Space;
+
+	UPROPERTY(meta = (Input))
+	FTransform WorldOffset;
+
+	UPROPERTY(meta = (Input, Constant))
+	bool bEnabled;
+
+	RIGVM_METHOD()
+	virtual FRigVMStructUpgradeInfo GetUpgradeInfo() const override;
+};
+
+/**
+ * Draws a rectangle in the viewport given a transform
+ */
+USTRUCT(meta=(DisplayName="Draw Rectangle", Keywords="Draw Square"))
+struct RIGVM_API FRigVMFunction_DebugRectangleNoSpace : public FRigVMFunction_DebugBaseMutable
+{
+	GENERATED_BODY()
+
+	FRigVMFunction_DebugRectangleNoSpace()
+	{
+		Transform = WorldOffset = FTransform::Identity;
+		Color = FLinearColor::Red;
+		Thickness = 0.f;
+		Scale = 10.f;
+		WorldOffset = FTransform::Identity;
+		bEnabled = true;
+	}
+
+	RIGVM_METHOD()
+	virtual void Execute() override;
+
+	UPROPERTY(meta = (Input))
+	FTransform Transform;
+
+	UPROPERTY(meta = (Input))
+	FLinearColor Color;
+
+	UPROPERTY(meta = (Input))
+	float Scale;
+
+	UPROPERTY(meta = (Input))
+	float Thickness;
+
+	UPROPERTY(meta = (Input))
+	FTransform WorldOffset;
+
+	UPROPERTY(meta = (Input))
+	bool bEnabled;
+};
+
+
+USTRUCT(meta=(DisplayName="Draw Arc", Keywords="Draw Ellipse, Draw Circle", Deprecated = "4.25"))
+struct RIGVM_API FRigVMFunction_DebugArc : public FRigVMFunction_DebugBaseMutable
+{
+	GENERATED_BODY()
+
+	FRigVMFunction_DebugArc()
+	{
+		Transform = WorldOffset = FTransform::Identity;
+		Color = FLinearColor::Red;
+		Thickness = MinimumDegrees = 0.f;
+		Radius = 10.f;
+		MaximumDegrees = 360.f;
+		Detail = 16;
+		WorldOffset = FTransform::Identity;
+		bEnabled = true;
+	}
+
+	RIGVM_METHOD()
+	virtual void Execute() override;
+
+	UPROPERTY(meta = (Input))
+	FTransform Transform;
+
+	UPROPERTY(meta = (Input))
+	FLinearColor Color;
+
+	UPROPERTY(meta = (Input))
+	float Radius;
+
+	UPROPERTY(meta = (Input))
+	float MinimumDegrees;
+
+	UPROPERTY(meta = (Input))
+	float MaximumDegrees;
+
+	UPROPERTY(meta = (Input))
+	float Thickness;
+
+	UPROPERTY(meta = (Input))
+	int32 Detail;
+
+	UPROPERTY(meta = (Input))
+	FName Space;
+
+	UPROPERTY(meta = (Input))
+	FTransform WorldOffset;
+
+	UPROPERTY(meta = (Input, Constant))
+	bool bEnabled;
+
+	RIGVM_METHOD()
+	virtual FRigVMStructUpgradeInfo GetUpgradeInfo() const override;
+};
+
+/**
+ * Draws an arc in the viewport, can take in different min and max degrees
+ */
+USTRUCT(meta=(DisplayName="Draw Arc", Keywords="Draw Ellipse, Draw Circle"))
+struct RIGVM_API FRigVMFunction_DebugArcNoSpace : public FRigVMFunction_DebugBaseMutable
+{
+	GENERATED_BODY()
+
+	FRigVMFunction_DebugArcNoSpace()
+	{
+		Transform = WorldOffset = FTransform::Identity;
+		Color = FLinearColor::Red;
+		Thickness = MinimumDegrees = 0.f;
+		Radius = 10.f;
+		MaximumDegrees = 360.f;
+		Detail = 16;
+		WorldOffset = FTransform::Identity;
+		bEnabled = true;
+	}
+
+	RIGVM_METHOD()
+	virtual void Execute() override;
+
+	UPROPERTY(meta = (Input))
+	FTransform Transform;
+
+	UPROPERTY(meta = (Input))
+	FLinearColor Color;
+
+	UPROPERTY(meta = (Input))
+	float Radius;
+
+	UPROPERTY(meta = (Input))
+	float MinimumDegrees;
+
+	UPROPERTY(meta = (Input))
+	float MaximumDegrees;
+
+	UPROPERTY(meta = (Input))
+	float Thickness;
+
+	UPROPERTY(meta = (Input))
+	int32 Detail;
+
+	UPROPERTY(meta = (Input))
+	FTransform WorldOffset;
+
+	UPROPERTY(meta = (Input))
+	bool bEnabled;
+};
\ No newline at end of file
diff --git a/Engine/Source/Runtime/RigVM/Public/RigVMFunctions/Debug/RigVMFunction_DebugTransform.h b/Engine/Source/Runtime/RigVM/Public/RigVMFunctions/Debug/RigVMFunction_DebugTransform.h
new file mode 100644
--- /dev/null
+++ b/Engine/Source/Runtime/RigVM/Public/RigVMFunctions/Debug/RigVMFunction_DebugTransform.h
@@ -0,0 +1,261 @@
+// Copyright Epic Games, Inc. All Rights Reserved.
+
+#pragma once
+
+#include "RigVMFunction_DebugBase.h"
+#include "RigVMFunction_DebugTransform.generated.h"
+
+UENUM()
+enum class ERigUnitDebugTransformMode : uint8
+{
+	/** Draw as point */
+	Point,
+
+	/** Draw as axes */
+	Axes,
+
+	/** Draw as box */
+	Box,
+
+	/** MAX - invalid */
+	Max UMETA(Hidden),
+};
+
+USTRUCT(meta=(DisplayName="Draw Transform In Place", Deprecated = "4.25.0"))
+struct RIGVM_API FRigVMFunction_DebugTransform : public FRigVMFunction_DebugBase
+{
+	GENERATED_BODY()
+
+	FRigVMFunction_DebugTransform()
+	{
+		Transform = WorldOffset = FTransform::Identity;
+		Mode = ERigUnitDebugTransformMode::Axes;
+		Color = FLinearColor::Red;
+		Thickness = 0.f;
+		Scale = 10.f;
+		bEnabled = true;
+	}
+
+	RIGVM_METHOD()
+	virtual void Execute() override;
+
+	UPROPERTY(meta = (Input, Output))
+	FTransform Transform;
+
+	UPROPERTY(meta = (Input))
+	ERigUnitDebugTransformMode Mode;
+
+	UPROPERTY(meta = (Input))
+	FLinearColor Color;
+
+	UPROPERTY(meta = (Input))
+	float Thickness;
+
+	UPROPERTY(meta = (Input))
+	float Scale;
+
+	UPROPERTY(meta = (Input))
+	FName Space;
+
+	UPROPERTY(meta = (Input))
+	FTransform WorldOffset;
+
+	UPROPERTY(meta = (Input, Constant))
+	bool bEnabled;
+
+	RIGVM_METHOD()
+	virtual FRigVMStructUpgradeInfo GetUpgradeInfo() const override;
+};
+
+USTRUCT(meta=(DisplayName="Draw Transform", Deprecated = "4.25"))
+struct RIGVM_API FRigVMFunction_DebugTransformMutable : public FRigVMFunction_DebugBaseMutable
+{
+	GENERATED_BODY()
+
+	FRigVMFunction_DebugTransformMutable()
+	{
+		Transform = WorldOffset = FTransform::Identity;
+		Mode = ERigUnitDebugTransformMode::Axes;
+		Color = FLinearColor::Red;
+		Thickness = 0.f;
+		Scale = 10.f;
+		bEnabled = true;
+	}
+
+	RIGVM_METHOD()
+	virtual void Execute() override;
+
+	UPROPERTY(meta = (Input))
+	FTransform Transform;
+
+	UPROPERTY(meta = (Input))
+	ERigUnitDebugTransformMode Mode;
+
+	UPROPERTY(meta = (Input))
+	FLinearColor Color;
+
+	UPROPERTY(meta = (Input))
+	float Thickness;
+
+	UPROPERTY(meta = (Input))
+	float Scale;
+
+	UPROPERTY(meta = (Input))
+	FName Space;
+
+	UPROPERTY(meta = (Input))
+	FTransform WorldOffset;
+
+	UPROPERTY(meta = (Input, Constant))
+	bool bEnabled;
+
+	RIGVM_METHOD()
+	virtual FRigVMStructUpgradeInfo GetUpgradeInfo() const override;
+};
+
+/**
+ * Given a transform, will draw a point, axis, or a box in the viewport
+ */
+USTRUCT(meta=(DisplayName="Draw Transform"))
+struct RIGVM_API FRigVMFunction_DebugTransformMutableNoSpace : public FRigVMFunction_DebugBaseMutable
+{
+	GENERATED_BODY()
+
+	FRigVMFunction_DebugTransformMutableNoSpace()
+	{
+		Transform = WorldOffset = FTransform::Identity;
+		Mode = ERigUnitDebugTransformMode::Axes;
+		Color = FLinearColor::Red;
+		Thickness = 0.f;
+		Scale = 10.f;
+		bEnabled = true;
+	}
+
+	RIGVM_METHOD()
+	virtual void Execute() override;
+
+	UPROPERTY(meta = (Input))
+	FTransform Transform;
+
+	UPROPERTY(meta = (Input))
+	ERigUnitDebugTransformMode Mode;
+
+	UPROPERTY(meta = (Input))
+	FLinearColor Color;
+
+	UPROPERTY(meta = (Input))
+	float Thickness;
+
+	UPROPERTY(meta = (Input))
+	float Scale;
+
+	UPROPERTY(meta = (Input))
+	FTransform WorldOffset;
+
+	UPROPERTY(meta = (Input))
+	bool bEnabled;
+};
+
+USTRUCT()
+struct RIGVM_API FRigVMFunction_DebugTransformArrayMutable_WorkData
+{
+	GENERATED_BODY()
+		
+	UPROPERTY()
+	TArray<FTransform> DrawTransforms;
+};
+
+USTRUCT(meta=(DisplayName="Draw Transform Array", Deprecated = "4.25.0"))
+struct RIGVM_API FRigVMFunction_DebugTransformArrayMutable : public FRigVMFunction_DebugBaseMutable
+{
+ 	GENERATED_BODY()
+
+	FRigVMFunction_DebugTransformArrayMutable()
+	{
+		Mode = ERigUnitDebugTransformMode::Axes;
+		Color = FLinearColor::Red;
+		Thickness = 0.f;
+		Scale = 10.f;
+		WorldOffset = FTransform::Identity;
+		bEnabled = true;
+	}
+
+	RIGVM_METHOD()
+	virtual void Execute() override;
+
+	UPROPERTY(meta = (Input))
+	TArray<FTransform> Transforms;
+
+	UPROPERTY(meta = (Input))
+	ERigUnitDebugTransformMode Mode;
+
+	UPROPERTY(meta = (Input))
+	FLinearColor Color;
+
+	UPROPERTY(meta = (Input))
+	float Thickness;
+
+	UPROPERTY(meta = (Input))
+	float Scale;
+
+	UPROPERTY(meta = (Input))
+	FName Space;
+
+	UPROPERTY(meta = (Input))
+	FTransform WorldOffset;
+
+	UPROPERTY(meta = (Input, Constant))
+	bool bEnabled;
+
+	UPROPERTY(transient)
+	FRigVMFunction_DebugTransformArrayMutable_WorkData WorkData;
+
+	RIGVM_METHOD()
+	virtual FRigVMStructUpgradeInfo GetUpgradeInfo() const override;
+};
+
+/**
+* Given a transform array, will draw a point, axis, or a box in the viewport
+*/
+USTRUCT(meta=(DisplayName="Draw Transform Array"))
+struct RIGVM_API FRigVMFunction_DebugTransformArrayMutableNoSpace : public FRigVMFunction_DebugBaseMutable
+{
+	GENERATED_BODY()
+
+	FRigVMFunction_DebugTransformArrayMutableNoSpace()
+	{
+		WorldOffset = FTransform::Identity;
+		Mode = ERigUnitDebugTransformMode::Axes;
+		Color = FLinearColor::White;
+		Thickness = 0.f;
+		Scale = 10.f;
+		bEnabled = true;
+	}
+
+	RIGVM_METHOD()
+	virtual void Execute() override;
+
+	UPROPERTY(meta = (Input))
+	TArray<FTransform> Transforms;
+
+	UPROPERTY(meta = (Input))
+	TArray<int32> ParentIndices;
+
+	UPROPERTY(meta = (Input))
+	ERigUnitDebugTransformMode Mode;
+
+	UPROPERTY(meta = (Input))
+	FLinearColor Color;
+
+	UPROPERTY(meta = (Input))
+	float Thickness;
+
+	UPROPERTY(meta = (Input))
+	float Scale;
+
+	UPROPERTY(meta = (Input))
+	FTransform WorldOffset;
+
+	UPROPERTY(meta = (Input))
+	bool bEnabled;
+};
diff --git a/Engine/Source/Runtime/RigVM/Public/RigVMFunctions/Debug/RigVMFunction_VisualDebug.h b/Engine/Source/Runtime/RigVM/Public/RigVMFunctions/Debug/RigVMFunction_VisualDebug.h
new file mode 100644
--- /dev/null
+++ b/Engine/Source/Runtime/RigVM/Public/RigVMFunctions/Debug/RigVMFunction_VisualDebug.h
@@ -0,0 +1,239 @@
+// Copyright Epic Games, Inc. All Rights Reserved.
+
+#pragma once
+
+#include "RigVMFunction_DebugBase.h"
+#include "RigVMFunction_VisualDebug.generated.h"
+
+UENUM()
+enum class ERigUnitVisualDebugPointMode : uint8
+{
+	/** Draw as point */
+	Point,
+
+	/** Draw as vector */
+	Vector,
+
+	/** MAX - invalid */
+	Max UMETA(Hidden),
+};
+
+USTRUCT(meta=(DisplayName = "Visual Debug Vector", TemplateName="VisualDebug", Keywords = "Draw,Point", Deprecated = "4.25", Varying))
+struct RIGVM_API FRigVMFunction_VisualDebugVector : public FRigVMFunction_DebugBase
+{
+	GENERATED_BODY()
+
+	FRigVMFunction_VisualDebugVector()
+	{
+		Value = FVector::ZeroVector;
+		bEnabled = true;
+		Mode = ERigUnitVisualDebugPointMode::Point;
+		Color = FLinearColor::Red;
+		Thickness = 10.f;
+		Scale = 1.f;
+		BoneSpace = NAME_None;
+	}
+
+	RIGVM_METHOD()
+	virtual void Execute() override;
+
+	UPROPERTY(meta = (Input, Output))
+	FVector Value;
+
+	UPROPERTY(meta = (Input))
+	bool bEnabled;
+
+	UPROPERTY(meta = (Input, EditCondition = "bEnabled"))
+	ERigUnitVisualDebugPointMode Mode;
+
+	UPROPERTY(meta = (Input, EditCondition = "bEnabled"))
+	FLinearColor Color;
+
+	UPROPERTY(meta = (Input, EditCondition = "bEnabled"))
+	float Thickness;
+
+	UPROPERTY(meta = (Input, EditCondition = "bEnabled"))
+	float Scale;
+
+	UPROPERTY(meta = (Input, EditCondition = "bEnabled"))
+	FName BoneSpace;
+
+	RIGVM_METHOD()
+	virtual FRigVMStructUpgradeInfo GetUpgradeInfo() const override;
+};
+
+/**
+ * Debug draw parameters for a Point or Vector given a vector
+ */
+USTRUCT(meta=(DisplayName = "Visual Debug Vector", TemplateName="VisualDebug", Keywords = "Draw,Point", Varying))
+struct RIGVM_API FRigVMFunction_VisualDebugVectorNoSpace : public FRigVMFunction_DebugBase
+{
+	GENERATED_BODY()
+
+	FRigVMFunction_VisualDebugVectorNoSpace()
+	{
+		Value = FVector::ZeroVector;
+		bEnabled = true;
+		Mode = ERigUnitVisualDebugPointMode::Point;
+		Color = FLinearColor::Red;
+		Thickness = 10.f;
+		Scale = 1.f;
+	}
+
+	RIGVM_METHOD()
+	virtual void Execute() override;
+
+	UPROPERTY(meta = (Input, Output))
+	FVector Value;
+
+	UPROPERTY(meta = (Input))
+	bool bEnabled;
+
+	UPROPERTY(meta = (Input, EditCondition = "bEnabled"))
+	ERigUnitVisualDebugPointMode Mode;
+
+	UPROPERTY(meta = (Input, EditCondition = "bEnabled"))
+	FLinearColor Color;
+
+	UPROPERTY(meta = (Input, EditCondition = "bEnabled"))
+	float Thickness;
+
+	UPROPERTY(meta = (Input, EditCondition = "bEnabled"))
+	float Scale;
+};
+
+USTRUCT(meta = (DisplayName = "Visual Debug Quat", TemplateName="VisualDebug", Keywords = "Draw,Rotation", Deprecated = "4.25", Varying))
+struct RIGVM_API FRigVMFunction_VisualDebugQuat : public FRigVMFunction_DebugBase
+{
+	GENERATED_BODY()
+
+	FRigVMFunction_VisualDebugQuat()
+	{
+		Value = FQuat::Identity;
+		bEnabled = true;
+		Thickness = 0.f;
+		Scale = 10.f;
+		BoneSpace = NAME_None;
+	}
+
+	RIGVM_METHOD()
+	virtual void Execute() override;
+
+	UPROPERTY(meta = (Input, Output))
+	FQuat Value;
+
+	UPROPERTY(meta = (Input))
+	bool bEnabled;
+
+	UPROPERTY(meta = (Input, EditCondition = "bEnabled"))
+	float Thickness;
+
+	UPROPERTY(meta = (Input, EditCondition = "bEnabled"))
+	float Scale;
+
+	UPROPERTY(meta = (Input, EditCondition = "bEnabled"))
+	FName BoneSpace;
+
+	RIGVM_METHOD()
+	virtual FRigVMStructUpgradeInfo GetUpgradeInfo() const override;
+};
+
+/**
+ * Debug draw parameters for an Axis given a quaternion
+ */
+USTRUCT(meta = (DisplayName = "Visual Debug Quat", TemplateName="VisualDebug", Keywords = "Draw,Rotation", Varying))
+struct RIGVM_API FRigVMFunction_VisualDebugQuatNoSpace : public FRigVMFunction_DebugBase
+{
+	GENERATED_BODY()
+
+	FRigVMFunction_VisualDebugQuatNoSpace()
+	{
+		Value = FQuat::Identity;
+		bEnabled = true;
+		Thickness = 0.f;
+		Scale = 10.f;
+	}
+
+	RIGVM_METHOD()
+	virtual void Execute() override;
+
+	UPROPERTY(meta = (Input, Output))
+	FQuat Value;
+
+	UPROPERTY(meta = (Input))
+	bool bEnabled;
+
+	UPROPERTY(meta = (Input, EditCondition = "bEnabled"))
+	float Thickness;
+
+	UPROPERTY(meta = (Input, EditCondition = "bEnabled"))
+	float Scale;
+};
+
+USTRUCT(meta=(DisplayName="Visual Debug Transform", TemplateName="VisualDebug", Keywords = "Draw,Axes", Deprecated = "4.25", Varying))
+struct RIGVM_API FRigVMFunction_VisualDebugTransform : public FRigVMFunction_DebugBase
+{
+	GENERATED_BODY()
+
+	FRigVMFunction_VisualDebugTransform()
+	{
+		Value = FTransform::Identity;
+		bEnabled = true;
+		Thickness = 0.f;
+		Scale = 10.f;
+		BoneSpace = NAME_None;
+	}
+
+	RIGVM_METHOD()
+	virtual void Execute() override;
+
+	UPROPERTY(meta = (Input, Output))
+	FTransform Value;
+
+	UPROPERTY(meta = (Input))
+	bool bEnabled;
+
+	UPROPERTY(meta = (Input, EditCondition = "bEnabled"))
+	float Thickness;
+
+	UPROPERTY(meta = (Input, EditCondition = "bEnabled"))
+	float Scale;
+
+	UPROPERTY(meta = (Input, EditCondition = "bEnabled"))
+	FName BoneSpace;
+
+	RIGVM_METHOD()
+	virtual FRigVMStructUpgradeInfo GetUpgradeInfo() const override;
+};
+
+/**
+ * Debug draw parameters for an Axis given a transform
+ */
+USTRUCT(meta=(DisplayName="Visual Debug Transform", TemplateName="VisualDebug", Keywords = "Draw,Axes", Varying))
+struct RIGVM_API FRigVMFunction_VisualDebugTransformNoSpace : public FRigVMFunction_DebugBase
+{
+	GENERATED_BODY()
+
+	FRigVMFunction_VisualDebugTransformNoSpace()
+	{
+		Value = FTransform::Identity;
+		bEnabled = true;
+		Thickness = 0.f;
+		Scale = 10.f;
+	}
+
+	RIGVM_METHOD()
+	virtual void Execute() override;
+
+	UPROPERTY(meta = (Input, Output))
+	FTransform Value;
+
+	UPROPERTY(meta = (Input))
+	bool bEnabled;
+
+	UPROPERTY(meta = (Input, EditCondition = "bEnabled"))
+	float Thickness;
+
+	UPROPERTY(meta = (Input, EditCondition = "bEnabled"))
+	float Scale;
+};
\ No newline at end of file
diff --git a/Engine/Source/Runtime/RigVM/Public/RigVMFunctions/Execution/RigVMFunction_ForLoop.h b/Engine/Source/Runtime/RigVM/Public/RigVMFunctions/Execution/RigVMFunction_ForLoop.h
new file mode 100644
--- /dev/null
+++ b/Engine/Source/Runtime/RigVM/Public/RigVMFunctions/Execution/RigVMFunction_ForLoop.h
@@ -0,0 +1,50 @@
+// Copyright Epic Games, Inc. All Rights Reserved.
+
+#pragma once
+
+#include "RigVMCore/RigVMStruct.h"
+#include "RigVMFunction_ForLoop.generated.h"
+
+/**
+ * Given a count, execute iteratively until the count is up
+ */
+USTRUCT(meta=(DisplayName="For Loop", Category="Execution", TitleColor="1 0 0", NodeColor="1 1 1", Keywords="Iterate", Icon="EditorStyle|GraphEditor.Macro.Loop_16x"))
+struct RIGVM_API FRigVMFunction_ForLoopCount : public FRigVMStructMutable
+{
+	GENERATED_BODY()
+
+	FRigVMFunction_ForLoopCount()
+	{
+		BlockToRun = NAME_None;
+		Count = 1;
+		Index = 0;
+		Ratio = 0.f;
+	}
+
+	// FRigVMStruct overrides
+	virtual const TArray<FName>& GetControlFlowBlocks_Impl() const override
+	{
+		static const TArray<FName> Blocks = {ExecuteContextName, ForLoopCompletedPinName};
+		return Blocks;
+	}
+	virtual const bool IsControlFlowBlockSliced(const FName& InBlockName) const { return InBlockName == ExecuteContextName; }
+	virtual int32 GetNumSlices() const override { return Count; }
+
+	RIGVM_METHOD()
+	virtual void Execute() override;
+
+	UPROPERTY(meta = (Singleton))
+	FName BlockToRun;
+
+	UPROPERTY(meta = (Singleton, Input))
+	int32 Count;
+
+	UPROPERTY(meta = (Singleton, Output))
+	int32 Index;
+
+	UPROPERTY(meta = (Singleton, Output))
+	float Ratio;
+
+	UPROPERTY(EditAnywhere, Transient, Category = "ForLoop", meta = (Output))
+	FRigVMExecuteContext Completed;
+};
diff --git a/Engine/Source/Runtime/RigVM/Public/RigVMFunctions/Execution/RigVMFunction_UserDefinedEvent.h b/Engine/Source/Runtime/RigVM/Public/RigVMFunctions/Execution/RigVMFunction_UserDefinedEvent.h
new file mode 100644
--- /dev/null
+++ b/Engine/Source/Runtime/RigVM/Public/RigVMFunctions/Execution/RigVMFunction_UserDefinedEvent.h
@@ -0,0 +1,35 @@
+// Copyright Epic Games, Inc. All Rights Reserved.
+
+#pragma once
+
+#include "RigVMCore/RigVMStruct.h"
+#include "RigVMFunction_UserDefinedEvent.generated.h"
+
+/**
+ * User Defined Event for running custom logic
+ */
+USTRUCT(meta=(DisplayName="User Defined Event", Category="Events", NodeColor="1, 0, 0", Keywords="Event,Entry,MyEvent"))
+struct RIGVM_API FRigVMFunction_UserDefinedEvent : public FRigVMStruct
+{
+	GENERATED_BODY()
+
+	FRigVMFunction_UserDefinedEvent()
+	{
+		EventName = TEXT("My Event");
+	}
+
+	RIGVM_METHOD()
+	virtual void Execute() override;
+
+	virtual FString GetUnitLabel() const override;
+	virtual FName GetEventName() const override { return EventName; }
+	virtual bool CanOnlyExistOnce() const override { return false; }
+
+	// The execution result
+	UPROPERTY(EditAnywhere, Transient, DisplayName = "Execute", Category = "UserDefinedEvent", meta = (Output))
+	FRigVMExecuteContext ExecuteContext;
+
+	// True if the current interaction is a rotation
+	UPROPERTY(EditAnywhere, Transient, Category = "UserDefinedEvent", meta = (Input,Constant,DetailsOnly))
+	FName EventName;
+};
diff --git a/Engine/Source/Runtime/RigVM/Public/RigVMFunctions/Math/RigVMMathLibrary.h b/Engine/Source/Runtime/RigVM/Public/RigVMFunctions/Math/RigVMMathLibrary.h
--- a/Engine/Source/Runtime/RigVM/Public/RigVMFunctions/Math/RigVMMathLibrary.h
+++ b/Engine/Source/Runtime/RigVM/Public/RigVMFunctions/Math/RigVMMathLibrary.h
@@ -98,6 +98,69 @@ struct RIGVM_API FRigVMMirrorSettings
 	FVector MirrorVector(const FVector& InVector) const;
 };
 
+UENUM()
+enum class ERigVMSimPointIntegrateType : uint8
+{
+	Verlet,
+	SemiExplicitEuler
+};
+
+USTRUCT(BlueprintType)
+struct RIGVM_API FRigVMSimPoint
+{
+	GENERATED_BODY()
+
+	FRigVMSimPoint()
+	{
+		Mass = 1.f;
+		Size = 0.f;
+		LinearDamping = 0.01f;
+		InheritMotion = 0.f;
+		Position = LinearVelocity = FVector::ZeroVector;
+	}
+
+	/**
+	 * The mass of the point
+	 */
+	UPROPERTY(EditAnywhere, Category=Simulation)
+	float Mass;
+
+	/**
+	 * Size of the point - only used for collision
+	 */
+	UPROPERTY(EditAnywhere, Category=Simulation)
+	float Size;
+
+	/**
+	 * The linear damping of the point
+	 */
+	UPROPERTY(EditAnywhere, Category=Simulation)
+	float LinearDamping;
+
+	/**
+	 * Defines how much the point will inherit motion from its input.
+	 * This does not have an effect on passive (mass == 0.0) points.
+	 * Values can be higher than 1 due to timestep - but they are clamped internally.
+	 */
+	UPROPERTY(EditAnywhere, Category=Simulation)
+	float InheritMotion;
+
+	/**
+	 * The position of the point
+	 */
+	UPROPERTY(EditAnywhere, Category=Simulation)
+	FVector Position;
+
+	/**
+	 * The velocity of the point per second
+	 */
+	UPROPERTY(EditAnywhere, Category=Simulation)
+	FVector LinearVelocity;
+
+	FRigVMSimPoint IntegrateVerlet(const FVector& InForce, float InBlend, float InDeltaTime) const;
+	FRigVMSimPoint IntegrateSemiExplicitEuler(const FVector& InForce, float InDeltaTime) const;
+};
+
 class RIGVM_API FRigVMMathLibrary
 {
 public:
diff --git a/Engine/Source/Runtime/RigVM/Public/RigVMFunctions/Simulation/RigVMFunction_Accumulate.h b/Engine/Source/Runtime/RigVM/Public/RigVMFunctions/Simulation/RigVMFunction_Accumulate.h
new file mode 100644
--- /dev/null
+++ b/Engine/Source/Runtime/RigVM/Public/RigVMFunctions/Simulation/RigVMFunction_Accumulate.h
@@ -0,0 +1,499 @@
+// Copyright Epic Games, Inc. All Rights Reserved.
+
+#pragma once
+
+#include "RigVMFunction_SimBase.h"
+#include "RigVMFunction_Accumulate.generated.h"
+
+USTRUCT(meta=(Category="Simulation|Accumulate"))
+struct RIGVM_API FRigVMFunction_AccumulateBase : public FRigVMFunction_SimBase
+{
+	GENERATED_BODY()
+};
+
+/**
+ * Adds a value over time over and over again
+ */
+USTRUCT(meta=(DisplayName="Accumulate Add (Float)", TemplateName="AccumulateAdd", Keywords="Simulate,++"))
+struct RIGVM_API FRigVMFunction_AccumulateFloatAdd : public FRigVMFunction_AccumulateBase
+{
+	GENERATED_BODY()
+	
+	FRigVMFunction_AccumulateFloatAdd()
+	{
+		InitialValue = Increment = Result = AccumulatedValue = 0.f;
+		bIntegrateDeltaTime = false;
+		bIsInitialized = false;
+	}
+
+	virtual void Initialize() override { bIsInitialized = false; }
+
+	RIGVM_METHOD()
+	virtual void Execute() override;
+
+	UPROPERTY(meta=(Input))
+	float Increment;
+
+	UPROPERTY(meta = (Input))
+	float InitialValue;
+
+	UPROPERTY(meta=(Input))
+	bool bIntegrateDeltaTime;
+
+	UPROPERTY(meta=(Output))
+	float Result;
+
+	UPROPERTY()
+	float AccumulatedValue;
+
+	UPROPERTY()
+	bool bIsInitialized;
+};
+
+/**
+ * Adds a vector over time over and over again
+ */
+USTRUCT(meta = (DisplayName="Accumulate Add (Vector)", TemplateName="AccumulateAdd", Keywords="Simulate,++"))
+struct RIGVM_API FRigVMFunction_AccumulateVectorAdd : public FRigVMFunction_AccumulateBase
+{
+	GENERATED_BODY()
+
+	FRigVMFunction_AccumulateVectorAdd()
+	{
+		InitialValue = Increment = Result = AccumulatedValue = FVector::ZeroVector;
+		bIntegrateDeltaTime = false;
+		bIsInitialized = false;
+	}
+
+	virtual void Initialize() override { bIsInitialized = false; }
+
+	RIGVM_METHOD()
+	virtual void Execute() override;
+
+	UPROPERTY(meta = (Input))
+	FVector Increment;
+
+	UPROPERTY(meta = (Input))
+	FVector InitialValue;
+
+	UPROPERTY(meta = (Input))
+	bool bIntegrateDeltaTime;
+
+	UPROPERTY(meta = (Output))
+	FVector Result;
+
+	UPROPERTY()
+	FVector AccumulatedValue;
+
+	UPROPERTY()
+	bool bIsInitialized;
+};
+
+/**
+ * Multiplies a value over time over and over again
+ */
+USTRUCT(meta=(DisplayName="Accumulate Mul (Float)", TemplateName="AccumulateMul", Keywords="Simulate,**"))
+struct RIGVM_API FRigVMFunction_AccumulateFloatMul : public FRigVMFunction_AccumulateBase
+{
+	GENERATED_BODY()
+	
+	FRigVMFunction_AccumulateFloatMul()
+	{
+		InitialValue = Multiplier = Result = AccumulatedValue = 1.f;
+		bIntegrateDeltaTime = false;
+		bIsInitialized = false;
+	}
+
+	virtual void Initialize() override { bIsInitialized = false; }
+
+	RIGVM_METHOD()
+	virtual void Execute() override;
+
+	UPROPERTY(meta=(Input))
+	float Multiplier;
+
+	UPROPERTY(meta = (Input))
+	float InitialValue;
+
+	UPROPERTY(meta=(Input))
+	bool bIntegrateDeltaTime;
+
+	UPROPERTY(meta=(Output))
+	float Result;
+
+	UPROPERTY()
+	float AccumulatedValue;
+
+	UPROPERTY()
+	bool bIsInitialized;
+};
+
+/**
+ * Multiplies a vector over time over and over again
+ */
+USTRUCT(meta = (DisplayName="Accumulate Mul (Vector)", TemplateName="AccumulateMul", Keywords="Simulate,**"))
+struct RIGVM_API FRigVMFunction_AccumulateVectorMul : public FRigVMFunction_AccumulateBase
+{
+	GENERATED_BODY()
+
+	FRigVMFunction_AccumulateVectorMul()
+	{
+		InitialValue = Multiplier = Result = AccumulatedValue = FVector::OneVector;
+		bIntegrateDeltaTime = false;
+		bIsInitialized = false;
+	}
+
+	virtual void Initialize() override { bIsInitialized = false; }
+
+	RIGVM_METHOD()
+	virtual void Execute() override;
+
+	UPROPERTY(meta = (Input))
+	FVector Multiplier;
+
+	UPROPERTY(meta = (Input))
+	FVector InitialValue;
+
+	UPROPERTY(meta = (Input))
+	bool bIntegrateDeltaTime;
+
+	UPROPERTY(meta = (Output))
+	FVector Result;
+
+	UPROPERTY()
+	FVector AccumulatedValue;
+
+	UPROPERTY()
+	bool bIsInitialized;
+};
+
+/**
+ * Multiplies a quaternion over time over and over again
+ */
+USTRUCT(meta=(DisplayName="Accumulate Mul (Quaternion)", TemplateName="AccumulateMul", Keywords="Simulate,**"))
+struct RIGVM_API FRigVMFunction_AccumulateQuatMul : public FRigVMFunction_AccumulateBase
+{
+	GENERATED_BODY()
+	
+	FRigVMFunction_AccumulateQuatMul()
+	{
+		InitialValue = Multiplier = Result = AccumulatedValue = FQuat::Identity;
+		bIntegrateDeltaTime = bFlipOrder = false;
+		bIsInitialized = false;
+	}
+
+	virtual void Initialize() override { bIsInitialized = false; }
+
+	RIGVM_METHOD()
+	virtual void Execute() override;
+
+	UPROPERTY(meta=(Input))
+	FQuat Multiplier;
+
+	UPROPERTY(meta = (Input))
+	FQuat InitialValue;
+
+	UPROPERTY(meta = (Input))
+	bool bFlipOrder;
+
+	UPROPERTY(meta=(Input))
+	bool bIntegrateDeltaTime;
+
+	UPROPERTY(meta=(Output))
+	FQuat Result;
+
+	UPROPERTY()
+	FQuat AccumulatedValue;
+
+	UPROPERTY()
+	bool bIsInitialized;
+};
+
+/**
+ * Multiplies a transform over time over and over again
+ */
+USTRUCT(meta = (DisplayName="Accumulate Mul (Transform)", TemplateName="AccumulateMul", Keywords="Simulate,**"))
+struct RIGVM_API FRigVMFunction_AccumulateTransformMul : public FRigVMFunction_AccumulateBase
+{
+	GENERATED_BODY()
+
+	FRigVMFunction_AccumulateTransformMul()
+	{
+		InitialValue = Multiplier = Result = AccumulatedValue = FTransform::Identity;
+		bIntegrateDeltaTime = bFlipOrder = false;
+		bIsInitialized = false;
+	}
+
+	virtual void Initialize() override { bIsInitialized = false; }
+
+	RIGVM_METHOD()
+	virtual void Execute() override;
+
+	UPROPERTY(meta = (Input))
+	FTransform Multiplier;
+
+	UPROPERTY(meta = (Input))
+	FTransform InitialValue;
+
+	UPROPERTY(meta = (Input))
+	bool bFlipOrder;
+
+	UPROPERTY(meta = (Input))
+	bool bIntegrateDeltaTime;
+
+	UPROPERTY(meta = (Output))
+	FTransform Result;
+
+	UPROPERTY()
+	FTransform AccumulatedValue;
+
+	UPROPERTY()
+	bool bIsInitialized;
+};
+
+/**
+ * Interpolates two values over time over and over again
+ */
+USTRUCT(meta=(DisplayName="Accumulate Lerp (Float)", TemplateName="AccumulateLerp", Keywords="Simulate,Ramp"))
+struct RIGVM_API FRigVMFunction_AccumulateFloatLerp : public FRigVMFunction_AccumulateBase
+{
+	GENERATED_BODY()
+	
+	FRigVMFunction_AccumulateFloatLerp()
+	{
+		InitialValue = TargetValue = Blend = Result = AccumulatedValue = 0.f;
+		bIntegrateDeltaTime = false;
+		bIsInitialized = false;
+	}
+
+	virtual void Initialize() override { bIsInitialized = false; }
+
+	RIGVM_METHOD()
+	virtual void Execute() override;
+
+	UPROPERTY(meta=(Input))
+	float TargetValue;
+
+	UPROPERTY(meta = (Input))
+	float InitialValue;
+
+	UPROPERTY(meta = (Input))
+	float Blend;
+
+	UPROPERTY(meta=(Input))
+	bool bIntegrateDeltaTime;
+
+	UPROPERTY(meta=(Output))
+	float Result;
+
+	UPROPERTY()
+	float AccumulatedValue;
+
+	UPROPERTY()
+	bool bIsInitialized;
+};
+
+/**
+ * Interpolates two vectors over time over and over again
+ */
+USTRUCT(meta = (DisplayName="Accumulate Lerp (Vector)", TemplateName="AccumulateLerp", Keywords="Simulate,Ramp"))
+struct RIGVM_API FRigVMFunction_AccumulateVectorLerp : public FRigVMFunction_AccumulateBase
+{
+	GENERATED_BODY()
+
+	FRigVMFunction_AccumulateVectorLerp()
+	{
+		InitialValue = TargetValue = Result = AccumulatedValue = FVector::ZeroVector;
+		Blend = 0.f;
+		bIntegrateDeltaTime = false;
+		bIsInitialized = false;
+	}
+
+	virtual void Initialize() override { bIsInitialized = false; }
+
+	RIGVM_METHOD()
+	virtual void Execute() override;
+
+	UPROPERTY(meta = (Input))
+	FVector TargetValue;
+
+	UPROPERTY(meta = (Input))
+	FVector InitialValue;
+
+	UPROPERTY(meta = (Input))
+	float Blend;
+
+	UPROPERTY(meta = (Input))
+	bool bIntegrateDeltaTime;
+
+	UPROPERTY(meta = (Output))
+	FVector Result;
+
+	UPROPERTY()
+	FVector AccumulatedValue;
+
+	UPROPERTY()
+	bool bIsInitialized;
+};
+
+/**
+ * Interpolates two quaternions over time over and over again
+ */
+USTRUCT(meta=(DisplayName="Accumulate Lerp (Quaternion)", TemplateName="AccumulateLerp", Keywords="Simulate,Ramp"))
+struct RIGVM_API FRigVMFunction_AccumulateQuatLerp : public FRigVMFunction_AccumulateBase
+{
+	GENERATED_BODY()
+	
+	FRigVMFunction_AccumulateQuatLerp()
+	{
+		InitialValue = TargetValue = Result = AccumulatedValue = FQuat::Identity;
+		Blend = 0.f;
+		bIntegrateDeltaTime = false;
+		bIsInitialized = false;
+	}
+
+	virtual void Initialize() override { bIsInitialized = false; }
+
+	RIGVM_METHOD()
+	virtual void Execute() override;
+
+	UPROPERTY(meta=(Input))
+	FQuat TargetValue;
+
+	UPROPERTY(meta = (Input))
+	FQuat InitialValue;
+
+	UPROPERTY(meta = (Input))
+	float Blend;
+
+	UPROPERTY(meta=(Input))
+	bool bIntegrateDeltaTime;
+
+	UPROPERTY(meta=(Output))
+	FQuat Result;
+
+	UPROPERTY()
+	FQuat AccumulatedValue;
+
+	UPROPERTY()
+	bool bIsInitialized;
+};
+
+/**
+ * Interpolates two transforms over time over and over again
+ */
+USTRUCT(meta = (DisplayName="Accumulate Lerp (Transform)", TemplateName="AccumulateLerp", Keywords="Simulate,Ramp"))
+struct RIGVM_API FRigVMFunction_AccumulateTransformLerp : public FRigVMFunction_AccumulateBase
+{
+	GENERATED_BODY()
+
+	FRigVMFunction_AccumulateTransformLerp()
+	{
+		InitialValue = TargetValue = Result = AccumulatedValue = FTransform::Identity;
+		Blend = 0.f;
+		bIntegrateDeltaTime = false;
+		bIsInitialized = false;
+	}
+
+	virtual void Initialize() override { bIsInitialized = false; }
+
+	RIGVM_METHOD()
+	virtual void Execute() override;
+
+	UPROPERTY(meta = (Input))
+	FTransform TargetValue;
+
+	UPROPERTY(meta = (Input))
+	FTransform InitialValue;
+
+	UPROPERTY(meta = (Input))
+	float Blend;
+
+	UPROPERTY(meta = (Input))
+	bool bIntegrateDeltaTime;
+
+	UPROPERTY(meta = (Output))
+	FTransform Result;
+
+	UPROPERTY()
+	FTransform AccumulatedValue;
+
+	UPROPERTY()
+	bool bIsInitialized;
+};
+
+/**
+ * Accumulates the min and max values over time
+ */
+USTRUCT(meta=(DisplayName="Accumulate Range (Float)", TemplateName="AccumulateRange", Keywords="Range"))
+struct RIGVM_API FRigVMFunction_AccumulateFloatRange : public FRigVMFunction_AccumulateBase
+{
+	GENERATED_BODY()
+	
+	FRigVMFunction_AccumulateFloatRange()
+	{
+		Value = Minimum = Maximum = AccumulatedMinimum = AccumulatedMaximum = 0.f;
+		bIsInitialized = false;
+	}
+
+	virtual void Initialize() override { bIsInitialized = false; }
+
+	RIGVM_METHOD()
+	virtual void Execute() override;
+
+	UPROPERTY(meta = (Input))
+	float Value;
+
+	UPROPERTY(meta= (Output))
+	float Minimum;
+
+	UPROPERTY(meta= (Output))
+	float Maximum;
+
+	UPROPERTY()
+	float AccumulatedMinimum;
+
+	UPROPERTY()
+	float AccumulatedMaximum;
+
+	UPROPERTY()
+	bool bIsInitialized;
+};
+
+/**
+ * Accumulates the min and max values over time
+ */
+USTRUCT(meta = (DisplayName="Accumulate Range (Vector)", TemplateName="AccumulateRange", Keywords="Range"))
+struct RIGVM_API FRigVMFunction_AccumulateVectorRange : public FRigVMFunction_AccumulateBase
+{
+	GENERATED_BODY()
+
+	FRigVMFunction_AccumulateVectorRange()
+	{
+		Value = Minimum = Maximum = AccumulatedMinimum = AccumulatedMaximum = FVector::ZeroVector;
+		bIsInitialized = false;
+	}
+
+	virtual void Initialize() override { bIsInitialized = false; }
+
+	RIGVM_METHOD()
+	virtual void Execute() override;
+
+	UPROPERTY(meta = (Input))
+	FVector Value;
+
+	UPROPERTY(meta= (Output))
+	FVector Minimum;
+
+	UPROPERTY(meta= (Output))
+	FVector Maximum;
+
+	UPROPERTY()
+	FVector AccumulatedMinimum;
+
+	UPROPERTY()
+	FVector AccumulatedMaximum;
+
+	UPROPERTY()
+	bool bIsInitialized;
+};
diff --git a/Engine/Source/Runtime/RigVM/Public/RigVMFunctions/Simulation/RigVMFunction_AlphaInterp.h b/Engine/Source/Runtime/RigVM/Public/RigVMFunctions/Simulation/RigVMFunction_AlphaInterp.h
new file mode 100644
--- /dev/null
+++ b/Engine/Source/Runtime/RigVM/Public/RigVMFunctions/Simulation/RigVMFunction_AlphaInterp.h
@@ -0,0 +1,226 @@
+// Copyright Epic Games, Inc. All Rights Reserved.
+
+#pragma once
+
+#include "RigVMFunction_SimBase.h"
+#include "Animation/InputScaleBias.h"
+#include "RigVMFunction_AlphaInterp.generated.h"
+
+/**
+ * Takes in a float value and outputs an accumulated value with a customized scale and clamp
+ */
+USTRUCT(meta=(DisplayName="Alpha Interpolate", Keywords="Alpha,Lerp,LinearInterpolate", Category = "Simulation|Time", TemplateName="AlphaInterp", MenuDescSuffix = "(Float)"))
+struct RIGVM_API FRigVMFunction_AlphaInterp : public FRigVMFunction_SimBase
+{
+	GENERATED_BODY()
+	
+	FRigVMFunction_AlphaInterp()
+	{
+		Value = Result = 0.f;
+		ScaleBiasClamp = FInputScaleBiasClamp();
+		bMapRange = ScaleBiasClamp.bMapRange;
+		bClampResult = ScaleBiasClamp.bClampResult;
+		bInterpResult = ScaleBiasClamp.bInterpResult;
+		InRange = ScaleBiasClamp.InRange;
+		OutRange = ScaleBiasClamp.OutRange;
+		Scale = ScaleBiasClamp.Scale;
+		Bias = ScaleBiasClamp.Bias;
+		ClampMin = ScaleBiasClamp.ClampMin;
+		ClampMax = ScaleBiasClamp.ClampMax;
+		InterpSpeedIncreasing = ScaleBiasClamp.InterpSpeedIncreasing;
+		InterpSpeedDecreasing = ScaleBiasClamp.InterpSpeedDecreasing;
+	}
+
+	RIGVM_METHOD()
+	virtual void Execute() override;
+
+	virtual FString ProcessPinLabelForInjection(const FString& InLabel) const override;
+
+	UPROPERTY(meta=(Input))
+	float Value;
+
+	UPROPERTY(meta = (Input))
+	float Scale;
+
+	UPROPERTY(meta = (Input))
+	float Bias;
+
+	UPROPERTY(meta=(Input, Constant))
+	bool bMapRange;
+
+	UPROPERTY(meta=(Input, EditCondition = "bMapRange"))
+	FInputRange InRange;
+
+	UPROPERTY(meta=(Input, EditCondition = "bMapRange"))
+	FInputRange OutRange;
+
+	UPROPERTY(meta = (Input, Constant))
+	bool bClampResult;
+
+	UPROPERTY(meta=(Input, EditCondition = "bClampResult"))
+	float ClampMin;
+
+	UPROPERTY(meta=(Input, EditCondition = "bClampResult"))
+	float ClampMax;
+
+	UPROPERTY(meta = (Input, Constant))
+	bool bInterpResult;
+
+	UPROPERTY(meta=(Input, EditCondition = "bInterpResult"))
+	float InterpSpeedIncreasing;
+
+	UPROPERTY(meta=(Input, EditCondition = "bInterpResult"))
+	float InterpSpeedDecreasing;
+
+	UPROPERTY(meta=(Output))
+	float Result;
+
+	UPROPERTY()
+	FInputScaleBiasClamp ScaleBiasClamp;
+};
+
+/**
+ * Takes in a vector value and outputs an accumulated value with a customized scale and clamp
+ */
+USTRUCT(meta=(DisplayName="Alpha Interpolate", Keywords="Alpha,Lerp,LinearInterpolate", Category = "Simulation|Time", TemplateName="AlphaInterp", MenuDescSuffix = "(Vector)"))
+struct RIGVM_API FRigVMFunction_AlphaInterpVector : public FRigVMFunction_SimBase
+{
+	GENERATED_BODY()
+	
+	FRigVMFunction_AlphaInterpVector()
+	{
+		Value = Result = FVector::ZeroVector;
+		ScaleBiasClamp = FInputScaleBiasClamp();
+		bMapRange = ScaleBiasClamp.bMapRange;
+		bClampResult = ScaleBiasClamp.bClampResult;
+		bInterpResult = ScaleBiasClamp.bInterpResult;
+		InRange = ScaleBiasClamp.InRange;
+		OutRange = ScaleBiasClamp.OutRange;
+		Scale = ScaleBiasClamp.Scale;
+		Bias = ScaleBiasClamp.Bias;
+		ClampMin = ScaleBiasClamp.ClampMin;
+		ClampMax = ScaleBiasClamp.ClampMax;
+		InterpSpeedIncreasing = ScaleBiasClamp.InterpSpeedIncreasing;
+		InterpSpeedDecreasing = ScaleBiasClamp.InterpSpeedDecreasing;
+	}
+
+	RIGVM_METHOD()
+	virtual void Execute() override;
+
+	virtual FString ProcessPinLabelForInjection(const FString& InLabel) const override;
+
+	UPROPERTY(meta=(Input))
+	FVector Value;
+
+	UPROPERTY(meta = (Input))
+	float Scale;
+
+	UPROPERTY(meta = (Input))
+	float Bias;
+
+	UPROPERTY(meta = (Input, Constant))
+	bool bMapRange;
+
+	UPROPERTY(meta = (Input, EditCondition = "bMapRange"))
+	FInputRange InRange;
+
+	UPROPERTY(meta = (Input, EditCondition = "bMapRange"))
+	FInputRange OutRange;
+
+	UPROPERTY(meta = (Input, Constant))
+	bool bClampResult;
+
+	UPROPERTY(meta = (Input, EditCondition = "bClampResult"))
+	float ClampMin;
+
+	UPROPERTY(meta = (Input, EditCondition = "bClampResult"))
+	float ClampMax;
+
+	UPROPERTY(meta = (Input, Constant))
+	bool bInterpResult;
+
+	UPROPERTY(meta = (Input, EditCondition = "bInterpResult"))
+	float InterpSpeedIncreasing;
+
+	UPROPERTY(meta = (Input, EditCondition = "bInterpResult"))
+	float InterpSpeedDecreasing;
+
+	UPROPERTY(meta=(Output))
+	FVector Result;
+
+	UPROPERTY()
+	FInputScaleBiasClamp ScaleBiasClamp;
+};
+
+/**
+ * Takes in a vector value and outputs an accumulated value with a customized scale and clamp
+ */
+USTRUCT(meta=(DisplayName="Alpha Interpolate", Keywords="Alpha,Lerp,LinearInterpolate", Category = "Simulation|Time", TemplateName="AlphaInterp", MenuDescSuffix = "(Quat)"))
+struct RIGVM_API FRigVMFunction_AlphaInterpQuat : public FRigVMFunction_SimBase
+{
+	GENERATED_BODY()
+	
+	FRigVMFunction_AlphaInterpQuat()
+	{
+		Value = Result = FQuat::Identity;
+		ScaleBiasClamp = FInputScaleBiasClamp();
+		bMapRange = ScaleBiasClamp.bMapRange;
+		bClampResult = ScaleBiasClamp.bClampResult;
+		bInterpResult = ScaleBiasClamp.bInterpResult;
+		InRange = ScaleBiasClamp.InRange;
+		OutRange = ScaleBiasClamp.OutRange;
+		Scale = ScaleBiasClamp.Scale;
+		Bias = ScaleBiasClamp.Bias;
+		ClampMin = ScaleBiasClamp.ClampMin;
+		ClampMax = ScaleBiasClamp.ClampMax;
+		InterpSpeedIncreasing = ScaleBiasClamp.InterpSpeedIncreasing;
+		InterpSpeedDecreasing = ScaleBiasClamp.InterpSpeedDecreasing;
+	}
+
+	RIGVM_METHOD()
+	virtual void Execute() override;
+
+	virtual FString ProcessPinLabelForInjection(const FString& InLabel) const override;
+
+	UPROPERTY(meta=(Input))
+	FQuat Value;
+
+	UPROPERTY(meta = (Input))
+	float Scale;
+
+	UPROPERTY(meta = (Input))
+	float Bias;
+
+	UPROPERTY(meta = (Input, Constant))
+	bool bMapRange;
+
+	UPROPERTY(meta = (Input, EditCondition = "bMapRange"))
+	FInputRange InRange;
+
+	UPROPERTY(meta = (Input, EditCondition = "bMapRange"))
+	FInputRange OutRange;
+
+	UPROPERTY(meta = (Input, Constant))
+	bool bClampResult;
+
+	UPROPERTY(meta = (Input, EditCondition = "bClampResult"))
+	float ClampMin;
+
+	UPROPERTY(meta = (Input, EditCondition = "bClampResult"))
+	float ClampMax;
+
+	UPROPERTY(meta = (Input, Constant))
+	bool bInterpResult;
+
+	UPROPERTY(meta = (Input, EditCondition = "bInterpResult"))
+	float InterpSpeedIncreasing;
+
+	UPROPERTY(meta = (Input, EditCondition = "bInterpResult"))
+	float InterpSpeedDecreasing;
+
+	UPROPERTY(meta=(Output))
+	FQuat Result;
+
+	UPROPERTY()
+	FInputScaleBiasClamp ScaleBiasClamp;
+};
diff --git a/Engine/Source/Runtime/RigVM/Public/RigVMFunctions/Simulation/RigVMFunction_DeltaFromPrevious.h b/Engine/Source/Runtime/RigVM/Public/RigVMFunctions/Simulation/RigVMFunction_DeltaFromPrevious.h
new file mode 100644
--- /dev/null
+++ b/Engine/Source/Runtime/RigVM/Public/RigVMFunctions/Simulation/RigVMFunction_DeltaFromPrevious.h
@@ -0,0 +1,138 @@
+// Copyright Epic Games, Inc. All Rights Reserved.
+
+#pragma once
+
+#include "RigVMFunction_SimBase.h"
+#include "RigVMFunction_DeltaFromPrevious.generated.h"
+
+/**
+ * Computes the difference from the previous value going through the node
+ */
+USTRUCT(meta=(DisplayName="DeltaFromPrevious (Float)", Category = "Simulation|Time", TemplateName="DeltaFromPrevious", Keywords="Difference,Velocity,Acceleration"))
+struct RIGVM_API FRigVMFunction_DeltaFromPreviousFloat : public FRigVMFunction_SimBase
+{
+	GENERATED_BODY()
+	
+	FRigVMFunction_DeltaFromPreviousFloat()
+	{
+		Value = Delta = PreviousValue = Cache = 0.f;
+		bIsInitialized = false;
+	}
+
+	RIGVM_METHOD()
+	virtual void Execute() override;
+
+	UPROPERTY(meta = (Input))
+	float Value;
+
+	UPROPERTY(meta=(Output))
+	float Delta;
+
+	UPROPERTY(meta=(Output))
+	float PreviousValue;
+
+	UPROPERTY()
+	float Cache;
+
+	UPROPERTY()
+	bool bIsInitialized;
+};
+
+/**
+ * Computes the difference from the previous value going through the node
+ */
+USTRUCT(meta=(DisplayName="DeltaFromPrevious (Vector)", Category = "Simulation|Time", TemplateName="DeltaFromPrevious", Keywords="Difference,Velocity,Acceleration"))
+struct RIGVM_API FRigVMFunction_DeltaFromPreviousVector : public FRigVMFunction_SimBase
+{
+	GENERATED_BODY()
+	
+	FRigVMFunction_DeltaFromPreviousVector()
+	{
+		Value = Delta = PreviousValue = Cache = FVector::ZeroVector;
+		bIsInitialized = false;
+	}
+
+	RIGVM_METHOD()
+	virtual void Execute() override;
+
+	UPROPERTY(meta = (Input))
+	FVector Value;
+
+	UPROPERTY(meta=(Output))
+	FVector Delta;
+
+	UPROPERTY(meta=(Output))
+	FVector PreviousValue;
+
+	UPROPERTY()
+	FVector Cache;
+
+	UPROPERTY()
+	bool bIsInitialized;
+};
+
+/**
+ * Computes the difference from the previous value going through the node
+ */
+USTRUCT(meta=(DisplayName="DeltaFromPrevious (Quaternion)", Category = "Simulation|Time", TemplateName="DeltaFromPrevious", Keywords="Difference,Velocity,Acceleration"))
+struct RIGVM_API FRigVMFunction_DeltaFromPreviousQuat : public FRigVMFunction_SimBase
+{
+	GENERATED_BODY()
+	
+	FRigVMFunction_DeltaFromPreviousQuat()
+	{
+		Value = Delta = PreviousValue = Cache = FQuat::Identity;
+		bIsInitialized = false;
+	}
+
+	RIGVM_METHOD()
+	virtual void Execute() override;
+
+	UPROPERTY(meta = (Input))
+	FQuat Value;
+
+	UPROPERTY(meta=(Output))
+	FQuat Delta;
+
+	UPROPERTY(meta=(Output))
+	FQuat PreviousValue;
+
+	UPROPERTY()
+	FQuat Cache;
+
+	UPROPERTY()
+	bool bIsInitialized;
+};
+
+/**
+ * Computes the difference from the previous value going through the node
+ */
+USTRUCT(meta=(DisplayName="DeltaFromPrevious (Transform)", Category = "Simulation|Time", TemplateName="DeltaFromPrevious", Keywords="Difference,Velocity,Acceleration"))
+struct RIGVM_API FRigVMFunction_DeltaFromPreviousTransform : public FRigVMFunction_SimBase
+{
+	GENERATED_BODY()
+	
+	FRigVMFunction_DeltaFromPreviousTransform()
+	{
+		Value = Delta = PreviousValue = Cache = FTransform::Identity;
+		bIsInitialized = false;
+	}
+
+	RIGVM_METHOD()
+	virtual void Execute() override;
+
+	UPROPERTY(meta = (Input))
+	FTransform Value;
+
+	UPROPERTY(meta=(Output))
+	FTransform Delta;
+
+	UPROPERTY(meta=(Output))
+	FTransform PreviousValue;
+
+	UPROPERTY()
+	FTransform Cache;
+
+	UPROPERTY()
+	bool bIsInitialized;
+};
diff --git a/Engine/Source/Runtime/RigVM/Public/RigVMFunctions/Simulation/RigVMFunction_Kalman.h b/Engine/Source/Runtime/RigVM/Public/RigVMFunctions/Simulation/RigVMFunction_Kalman.h
new file mode 100644
--- /dev/null
+++ b/Engine/Source/Runtime/RigVM/Public/RigVMFunctions/Simulation/RigVMFunction_Kalman.h
@@ -0,0 +1,117 @@
+// Copyright Epic Games, Inc. All Rights Reserved.
+
+#pragma once
+
+#include "RigVMFunction_SimBase.h"
+#include "RigVMFunction_Kalman.generated.h"
+
+/**
+ * Averages a value over time.
+ * This uses a Kalman Filter internally.
+ */
+USTRUCT(meta=(DisplayName="Average Over Time (Float)", Category = "Simulation|Time", TemplateName="KalmanFilter", Keywords="Average,Smooth,KalmanFilter"))
+struct RIGVM_API FRigVMFunction_KalmanFloat : public FRigVMFunction_SimBase
+{
+	GENERATED_BODY()
+	
+	FRigVMFunction_KalmanFloat()
+	{
+		Value = Result = 0.f;
+		BufferSize = 16;
+		LastInsertIndex = 0;
+	}
+
+	virtual void Initialize() override { Buffer.Reset(); }
+
+	RIGVM_METHOD()
+	virtual void Execute() override;
+
+	UPROPERTY(meta = (Input))
+	float Value;
+
+	UPROPERTY(meta=(Input, Constant))
+	int32 BufferSize;
+
+	UPROPERTY(meta=(Output))
+	float Result;
+
+	UPROPERTY(meta=(ArraySize = "FMath::Clamp<int32>(BufferSize, 1, 512)"))
+	TArray<float> Buffer;
+
+	UPROPERTY()
+	int32 LastInsertIndex;
+};
+
+/**
+ * Averages a value over time.
+ * This uses a Kalman Filter internally.
+ */
+USTRUCT(meta=(DisplayName="Average Over Time (Vector)", Category = "Simulation|Time", TemplateName="KalmanFilter", Keywords="Average,Smooth,KalmanFilter"))
+struct RIGVM_API FRigVMFunction_KalmanVector : public FRigVMFunction_SimBase
+{
+	GENERATED_BODY()
+	
+	FRigVMFunction_KalmanVector()
+	{
+		Value = Result = FVector::ZeroVector;
+		BufferSize = 16;
+		LastInsertIndex = 0;
+	}
+	
+	virtual void Initialize() override { Buffer.Reset(); }
+
+	RIGVM_METHOD()
+	virtual void Execute() override;
+
+	UPROPERTY(meta = (Input))
+	FVector Value;
+
+	UPROPERTY(meta=(Input, Constant))
+	int32 BufferSize;
+
+	UPROPERTY(meta=(Output))
+	FVector Result;
+
+	UPROPERTY(meta=(ArraySize = "FMath::Clamp<int32>(BufferSize, 1, 512)"))
+	TArray<FVector> Buffer;
+
+	UPROPERTY()
+	int32 LastInsertIndex;
+};
+
+/**
+ * Averages a transform over time.
+ * This uses a Kalman Filter internally.
+ */
+USTRUCT(meta=(DisplayName="Average Over Time (Transform)", Category = "Simulation|Time", TemplateName="KalmanFilter", Keywords="Average,Smooth,KalmanFilter"))
+struct RIGVM_API FRigVMFunction_KalmanTransform : public FRigVMFunction_SimBase
+{
+	GENERATED_BODY()
+	
+	FRigVMFunction_KalmanTransform()
+	{
+		Value = Result = FTransform::Identity;
+		BufferSize = 16;
+		LastInsertIndex = 0;
+	}
+
+	virtual void Initialize() override { Buffer.Reset(); }
+
+	RIGVM_METHOD()
+	virtual void Execute() override;
+
+	UPROPERTY(meta = (Input))
+	FTransform Value;
+
+	UPROPERTY(meta=(Input, Constant))
+	int32 BufferSize;
+
+	UPROPERTY(meta=(Output))
+	FTransform Result;
+
+	UPROPERTY(meta = (ArraySize = "FMath::Clamp<int32>(BufferSize, 1, 512)"))
+	TArray<FTransform> Buffer;
+
+	UPROPERTY()
+	int32 LastInsertIndex;
+};
diff --git a/Engine/Source/Runtime/RigVM/Public/RigVMFunctions/Simulation/RigVMFunction_SimBase.h b/Engine/Source/Runtime/RigVM/Public/RigVMFunctions/Simulation/RigVMFunction_SimBase.h
new file mode 100644
--- /dev/null
+++ b/Engine/Source/Runtime/RigVM/Public/RigVMFunctions/Simulation/RigVMFunction_SimBase.h
@@ -0,0 +1,18 @@
+// Copyright Epic Games, Inc. All Rights Reserved.
+
+#pragma once
+
+#include "RigVMCore/RigVMStruct.h"
+#include "RigVMFunction_SimBase.generated.h"
+
+USTRUCT(meta=(Abstract, Category = "Simulation", NodeColor = "0.25 0.05 0.05"))
+struct RIGVM_API FRigVMFunction_SimBase : public FRigVMStruct
+{
+	GENERATED_BODY()
+};
+
+USTRUCT(meta=(Abstract, Category = "Simulation", NodeColor = "0.25 0.05 0.05"))
+struct RIGVM_API FRigVMFunction_SimBaseMutable : public FRigVMStructMutable
+{
+	GENERATED_BODY()
+};
diff --git a/Engine/Source/Runtime/RigVM/Public/RigVMFunctions/Simulation/RigVMFunction_TimeOffset.h b/Engine/Source/Runtime/RigVM/Public/RigVMFunctions/Simulation/RigVMFunction_TimeOffset.h
new file mode 100644
--- /dev/null
+++ b/Engine/Source/Runtime/RigVM/Public/RigVMFunctions/Simulation/RigVMFunction_TimeOffset.h
@@ -0,0 +1,171 @@
+// Copyright Epic Games, Inc. All Rights Reserved.
+
+#pragma once
+
+#include "RigVMFunction_SimBase.h"
+#include "RigVMFunction_TimeOffset.generated.h"
+
+/**
+ * Records a value over time and can access the value from the past
+ */
+USTRUCT(meta=(DisplayName="Value Over Time (Float)", Category = "Simulation|Time", TemplateName="TimeOffset", Keywords="Buffer,Delta,History,Previous,TimeOffset,Delay"))
+struct RIGVM_API FRigVMFunction_TimeOffsetFloat : public FRigVMFunction_SimBase
+{
+	GENERATED_BODY()
+	
+	FRigVMFunction_TimeOffsetFloat()
+	{
+		SecondsAgo = 0.5f;
+		Value = Result = 0.f;
+		BufferSize = 16;
+		TimeRange = 1.f;
+		LastInsertIndex = 0;
+		UpperBound = 0;
+	}
+
+	virtual void Initialize() override { Buffer.Reset(); }
+
+	RIGVM_METHOD()
+	virtual void Execute() override;
+	
+	/** The value to record */
+	UPROPERTY(meta = (Input))
+	float Value;
+
+	/** Seconds of time in the past you want to query the value for */
+	UPROPERTY(meta = (Input))
+	float SecondsAgo;
+
+	/** The sampling precision of the buffer. The higher the more precise - the more memory usage. */
+	UPROPERTY(meta=(Input, Constant))
+	int32 BufferSize;
+
+	/** The maximum time required for offsetting in seconds. */
+	UPROPERTY(meta=(Input, Constant))
+	float TimeRange;
+
+	UPROPERTY(meta=(Output))
+	float Result;
+
+	UPROPERTY(meta = (ArraySize = "FMath::Clamp<int32>(BufferSize, 2, 512)"))
+	TArray<float> Buffer;
+
+	UPROPERTY(meta = (ArraySize = "FMath::Clamp<int32>(BufferSize, 2, 512)"))
+	TArray<float> DeltaTimes;
+
+	UPROPERTY()
+	int32 LastInsertIndex;
+
+	UPROPERTY()
+	int32 UpperBound;
+};
+
+/**
+ * Records a value over time and can access the value from the past
+ */
+USTRUCT(meta=(DisplayName="Value Over Time (Vector)", Category = "Simulation|Time", TemplateName="TimeOffset", Keywords="Buffer,Delta,History,Previous,TimeOffset,Delay"))
+struct RIGVM_API FRigVMFunction_TimeOffsetVector : public FRigVMFunction_SimBase
+{
+	GENERATED_BODY()
+	
+	FRigVMFunction_TimeOffsetVector()
+	{
+		SecondsAgo = 0.5f;
+		Value = Result = FVector::ZeroVector;
+		BufferSize = 16;
+		TimeRange = 1.f;
+		LastInsertIndex = 0;
+		UpperBound = 0;
+	}
+
+	virtual void Initialize() override { Buffer.Reset(); }
+
+	RIGVM_METHOD()
+	virtual void Execute() override;
+	
+	/** The value to record */
+	UPROPERTY(meta = (Input))
+	FVector Value;
+
+	/** Seconds of time in the past you want to query the value for */
+	UPROPERTY(meta = (Input))
+	float SecondsAgo;
+
+	/** The sampling precision of the buffer. The higher the more precise - the more memory usage. */
+	UPROPERTY(meta=(Input, Constant))
+	int32 BufferSize;
+
+	/** The maximum time required for offsetting in seconds. */
+	UPROPERTY(meta=(Input, Constant))
+	float TimeRange;
+
+	UPROPERTY(meta=(Output))
+	FVector Result;
+
+	UPROPERTY(meta = (ArraySize = "FMath::Clamp<int32>(BufferSize, 2, 512)"))
+	TArray<FVector> Buffer;
+
+	UPROPERTY(meta = (ArraySize = "FMath::Clamp<int32>(BufferSize, 2, 512)"))
+	TArray<float> DeltaTimes;
+
+	UPROPERTY()
+	int32 LastInsertIndex;
+
+	UPROPERTY()
+	int32 UpperBound;
+};
+
+/**
+ * Records a value over time and can access the value from the past
+ */
+USTRUCT(meta=(DisplayName="Value Over Time (Transform)", Category = "Simulation|Time", TemplateName="TimeOffset", Keywords="Buffer,Delta,History,Previous,TimeOffset,Delay"))
+struct RIGVM_API FRigVMFunction_TimeOffsetTransform : public FRigVMFunction_SimBase
+{
+	GENERATED_BODY()
+	
+	FRigVMFunction_TimeOffsetTransform()
+	{
+		SecondsAgo = 0.5f;
+		Value = Result = FTransform::Identity;
+		BufferSize = 16;
+		TimeRange = 1.f;
+		LastInsertIndex = 0;
+		UpperBound = 0;
+	}
+
+	virtual void Initialize() override { Buffer.Reset(); }
+
+	RIGVM_METHOD()
+	virtual void Execute() override;
+	
+	/** The value to record */
+	UPROPERTY(meta = (Input))
+	FTransform Value;
+
+	/** Seconds of time in the past you want to query the value for */
+	UPROPERTY(meta = (Input))
+	float SecondsAgo;
+
+	/** The sampling precision of the buffer. The higher the more precise - the more memory usage. */
+	UPROPERTY(meta=(Input, Constant))
+	int32 BufferSize;
+
+	/** The maximum time required for offsetting in seconds. */
+	UPROPERTY(meta=(Input, Constant))
+	float TimeRange;
+
+	UPROPERTY(meta=(Output))
+	FTransform Result;
+
+	UPROPERTY(meta = (ArraySize = "FMath::Clamp<int32>(BufferSize, 2, 512)"))
+	TArray<FTransform> Buffer;
+
+	UPROPERTY(meta = (ArraySize = "FMath::Clamp<int32>(BufferSize, 2, 512)"))
+	TArray<float> DeltaTimes;
+
+	UPROPERTY()
+	int32 LastInsertIndex;
+
+	UPROPERTY()
+	int32 UpperBound;
+};
diff --git a/Engine/Source/Runtime/RigVM/Public/RigVMFunctions/Simulation/RigVMFunction_Timeline.h b/Engine/Source/Runtime/RigVM/Public/RigVMFunctions/Simulation/RigVMFunction_Timeline.h
new file mode 100644
--- /dev/null
+++ b/Engine/Source/Runtime/RigVM/Public/RigVMFunctions/Simulation/RigVMFunction_Timeline.h
@@ -0,0 +1,106 @@
+// Copyright Epic Games, Inc. All Rights Reserved.
+
+#pragma once
+
+#include "RigVMFunction_SimBase.h"
+#include "RigVMFunction_Timeline.generated.h"
+
+/**
+ * Simulates a time value - can act as a timeline playing back
+ */
+USTRUCT(meta=(DisplayName="Accumulated Time", Category = "Simulation|Accumulate", Keywords="Playback,Pause,Timeline"))
+struct RIGVM_API FRigVMFunction_Timeline : public FRigVMFunction_SimBase
+{
+	GENERATED_BODY()
+	
+	FRigVMFunction_Timeline()
+	{
+		Speed = 1.f;
+		Time = AccumulatedValue = 0.f;
+		bIsInitialized = false;
+	}
+
+	virtual void Initialize() override { bIsInitialized = false; }
+
+	RIGVM_METHOD()
+	virtual void Execute() override;
+
+	UPROPERTY(meta = (Input))
+	float Speed;
+
+	UPROPERTY(meta=(Output))
+	float Time;
+
+	UPROPERTY()
+	float AccumulatedValue;
+
+	UPROPERTY()
+	bool bIsInitialized;
+};
+
+/**
+ * Simulates a time value - and outputs loop information
+ */
+USTRUCT(meta=(DisplayName="Time Loop", Category = "Simulation|Accumulate", Keywords="Playback,Pause,Timeline"))
+struct RIGVM_API FRigVMFunction_TimeLoop : public FRigVMFunction_SimBase
+{
+	GENERATED_BODY()
+	
+	FRigVMFunction_TimeLoop()
+	{
+		Speed = 1.f;
+		Duration = 1.f;
+		Normalize = false;
+		Absolute = Relative = FlipFlop = AccumulatedAbsolute = AccumulatedRelative = 0.f;
+		NumIterations = 0;
+		Even = false;
+		bIsInitialized = false;
+	}
+
+	virtual void Initialize() override { bIsInitialized = false; }
+
+	RIGVM_METHOD()
+	virtual void Execute() override;
+
+	UPROPERTY(meta = (Input))
+	float Speed;
+
+	// the duration of a single loop in seconds
+	UPROPERTY(meta = (Input))
+	float Duration;
+
+	// if set to true the output relative and flipflop
+	// will be normalized over the duration.
+	UPROPERTY(meta = (Input))
+	bool Normalize;
+
+	// the overall time in seconds
+	UPROPERTY(meta=(Output))
+	float Absolute;
+
+	// the relative time in seconds (within the loop)
+	UPROPERTY(meta=(Output))
+	float Relative;
+
+	// the relative time in seconds (within the loop),
+	// going from 0 to duration and then back from duration to 0,
+	// or 0 to 1 and 1 to 0 if Normalize is turned on
+	UPROPERTY(meta=(Output))
+	float FlipFlop;
+
+	// true if the iteration of the loop is even
+	UPROPERTY(meta=(Output))
+	bool Even;
+
+	UPROPERTY()
+	float AccumulatedAbsolute;
+
+	UPROPERTY()
+	float AccumulatedRelative;
+
+	UPROPERTY()
+	int32 NumIterations;
+
+	UPROPERTY()
+	bool bIsInitialized;
+};
\ No newline at end of file
diff --git a/Engine/Source/Runtime/RigVM/Public/RigVMFunctions/Simulation/RigVMFunction_Verlet.h b/Engine/Source/Runtime/RigVM/Public/RigVMFunctions/Simulation/RigVMFunction_Verlet.h
new file mode 100644
--- /dev/null
+++ b/Engine/Source/Runtime/RigVM/Public/RigVMFunctions/Simulation/RigVMFunction_Verlet.h
@@ -0,0 +1,65 @@
+// Copyright Epic Games, Inc. All Rights Reserved.
+
+#pragma once
+
+#include "RigVMFunction_SimBase.h"
+#include "RigVMFunctions/Math/RigVMMathLibrary.h"
+#include "RigVMFunction_Verlet.generated.h"
+
+/**
+ * Simulates a single position over time using Verlet integration. It is recommended to use SpringInterp instead as it
+ * is more accurate and stable, and has more meaningful parameters.
+ */
+USTRUCT(meta=(DisplayName="Verlet (Vector)", Category = "Simulation|Springs", TemplateName="Verlet", Keywords="Simulate,Integrate"))
+struct RIGVM_API FRigVMFunction_VerletIntegrateVector : public FRigVMFunction_SimBase
+{
+	GENERATED_BODY()
+	
+	FRigVMFunction_VerletIntegrateVector()
+	{
+		Target = Force = Position = Velocity = Acceleration = FVector::ZeroVector;
+		Strength = 64.f;
+		Damp = 0.01;
+		Blend = 5.f;
+		Point = FRigVMSimPoint();
+		bInitialized = false;
+	}
+
+	RIGVM_METHOD()
+	virtual void Execute() override;
+
+	UPROPERTY(meta = (Input))
+	FVector Target;
+
+	/** The strength of the verlet spring */
+	UPROPERTY(meta = (Input))
+	float Strength;
+
+	/** The amount of damping to apply ( 0.0 to 1.0, but usually really low like 0.005 )*/
+	UPROPERTY(meta = (Input))
+	float Damp;
+
+	/** The amount of blending to apply per second */
+	UPROPERTY(meta = (Input))
+	float Blend;
+
+	/** The force feeding into the solver. Can be used for gravity. */
+	UPROPERTY(meta = (Input))
+	FVector Force;
+
+	UPROPERTY(meta = (Output))
+	FVector Position;
+
+	UPROPERTY(meta = (Output))
+	FVector Velocity;
+
+	UPROPERTY(meta = (Output))
+	FVector Acceleration;
+
+	UPROPERTY(transient)
+	FRigVMSimPoint Point;
+
+	UPROPERTY(transient)
+	bool bInitialized;
+};
+
