Control Rig: Fix crash in test_control_rig_add_remove

diff --git a/Engine/Source/Runtime/RigVM/Private/RigVMCore/RigVMRegistry.cpp b/Engine/Source/Runtime/RigVM/Private/RigVMCore/RigVMRegistry.cpp
--- a/Engine/Source/Runtime/RigVM/Private/RigVMCore/RigVMRegistry.cpp
+++ b/Engine/Source/Runtime/RigVM/Private/RigVMCore/RigVMRegistry.cpp
@@ -12,6 +12,7 @@
 #include "UObject/UObjectIterator.h"
 #include "Misc/CoreDelegates.h"
 #include "Misc/DelayedAutoRegister.h"
+#include "UObject/CoreRedirects.h"
 
 #if WITH_EDITOR
 #include "UObject/CoreRedirects.h"
@@ -710,15 +711,38 @@ const FRigVMTemplateArgumentType& FRigVMRegistry::FindTypeFromCPPType(const FStr
 
 TRigVMTypeIndex FRigVMRegistry::GetTypeIndexFromCPPType(const FString& InCPPType) const
 {
+	TRigVMTypeIndex Result = INDEX_NONE;
 	if(ensure(!InCPPType.IsEmpty()))
 	{
 		const FName CPPTypeName = *InCPPType;
-		return Types.IndexOfByPredicate([CPPTypeName](const FTypeInfo& Info) -> bool
+		Result = Types.IndexOfByPredicate([CPPTypeName](const FTypeInfo& Info) -> bool
 		{
 			return Info.Type.CPPType == CPPTypeName;
 		});
+
+		// If not found, try to find a redirect
+		if (Result == INDEX_NONE)
+		{
+			FCoreRedirectObjectName NewObjectName;
+			const bool bFoundRedirect = FCoreRedirects::RedirectNameAndValues(
+				ECoreRedirectFlags::Type_Class | ECoreRedirectFlags::Type_Struct | ECoreRedirectFlags::Type_Enum,
+				FCoreRedirectObjectName(*InCPPType),
+				NewObjectName,
+				nullptr,
+				ECoreRedirectMatchFlags::AllowPartialMatch); // AllowPartialMatch to allow redirects from one package to another (see /Script/ControlRig.CRFourPointBezier -> /Script/RigVM.RigVMFourPointBezier)
+
+			if (!bFoundRedirect)
+			{
+				return INDEX_NONE;
+			}
+
+			Result = Types.IndexOfByPredicate([NewObjectName](const FTypeInfo& Info) -> bool
+			{
+				return Info.Type.CPPType == NewObjectName.ObjectName;
+			});
+		}
 	}
-	return INDEX_NONE;
+	return Result;
 }
 
 bool FRigVMRegistry::IsArrayType(TRigVMTypeIndex InTypeIndex) const
