Control Rig: Serialize RigVMExecuteContext from ControlRigExecuteContext

diff --git a/Engine/Source/Runtime/RigVM/Private/RigVMCore/RigVMExecuteContext.cpp b/Engine/Source/Runtime/RigVM/Private/RigVMCore/RigVMExecuteContext.cpp
--- a/Engine/Source/Runtime/RigVM/Private/RigVMCore/RigVMExecuteContext.cpp
+++ b/Engine/Source/Runtime/RigVM/Private/RigVMCore/RigVMExecuteContext.cpp
@@ -4,3 +4,19 @@
 
 #include UE_INLINE_GENERATED_CPP_BY_NAME(RigVMExecuteContext)
 
+bool FRigVMExecuteContext::SerializeFromMismatchedTag(const FPropertyTag& Tag, FStructuredArchive::FSlot Slot)
+{
+	static const FName ControlRigExecuteContextName("ControlRigExecuteContext");
+	if (Tag.Type == NAME_StructProperty && Tag.StructName == ControlRigExecuteContextName)
+	{
+		static const FString CRExecuteContextPath = TEXT("/Script/ControlRig.ControlRigExecuteContext");
+		UScriptStruct* OldStruct = FindFirstObject<UScriptStruct>(*CRExecuteContextPath, EFindFirstObjectOptions::NativeFirst | EFindFirstObjectOptions::EnsureIfAmbiguous);
+		checkf(OldStruct, TEXT("FControlRigExecuteContext was not found."));
+
+		TSharedPtr<FStructOnScope> StructOnScope = MakeShareable(new FStructOnScope(OldStruct));
+		OldStruct->SerializeItem(Slot, StructOnScope->GetStructMemory(), nullptr);		
+		return true;
+	}
+
+	return false;
+}
diff --git a/Engine/Source/Runtime/RigVM/Public/RigVMCore/RigVMExecuteContext.h b/Engine/Source/Runtime/RigVM/Public/RigVMCore/RigVMExecuteContext.h
--- a/Engine/Source/Runtime/RigVM/Public/RigVMCore/RigVMExecuteContext.h
+++ b/Engine/Source/Runtime/RigVM/Public/RigVMCore/RigVMExecuteContext.h
@@ -321,6 +321,11 @@ struct RIGVM_API FRigVMExecuteContext
 		NameCache = InOtherContext->NameCache;
 	}
 
+	/**
+	 * Serialize this type from another
+	 */
+	bool SerializeFromMismatchedTag(const FPropertyTag& Tag, FStructuredArchive::FSlot Slot);
+
 protected:
 
 	virtual void Reset()
@@ -363,6 +368,16 @@ protected:
 	friend class URigVMNativized;
 };
 
+template<>
+struct TStructOpsTypeTraits<FRigVMExecuteContext> : public TStructOpsTypeTraitsBase2<FRigVMExecuteContext>
+{
+	enum
+	{
+		WithStructuredSerializeFromMismatchedTag = true,
+	};
+};
+
+
 /**
  * The execute context is used for mutable nodes to
  * indicate execution order.
