Control Rig: Fix cooking error on old rig due to missing redirects

diff --git a/Engine/Source/Runtime/RigVM/Private/RigVMCore/RigVMMemoryDeprecated.cpp b/Engine/Source/Runtime/RigVM/Private/RigVMCore/RigVMMemoryDeprecated.cpp
--- a/Engine/Source/Runtime/RigVM/Private/RigVMCore/RigVMMemoryDeprecated.cpp
+++ b/Engine/Source/Runtime/RigVM/Private/RigVMCore/RigVMMemoryDeprecated.cpp
@@ -6,6 +6,7 @@
 #include "UObject/PropertyPortFlags.h"
 #include "UObject/Package.h"
 #include "RigVMModule.h"
+#include "UObject/CoreRedirects.h"
 
 #include UE_INLINE_GENERATED_CPP_BY_NAME(RigVMMemoryDeprecated)
 
@@ -597,6 +598,17 @@ void FRigVMMemoryContainer::Load(FArchive& Ar)
 	{
 		UScriptStruct* ScriptStruct = FindObject<UScriptStruct>(nullptr, *ScriptStructPath);
 
+		// try to find a redirect
+		if (ScriptStruct == nullptr)
+		{
+			const FCoreRedirectObjectName OldObjectName(ScriptStructPath);
+			const FCoreRedirectObjectName NewObjectName = FCoreRedirects::GetRedirectedName(ECoreRedirectFlags::Type_Struct, OldObjectName);
+			if (OldObjectName != NewObjectName)
+			{
+				ScriptStruct = FindObject<UScriptStruct>(nullptr, *NewObjectName.ToString());
+			}
+		}
+
 		// this might have happened if a given script struct no longer 
 		// exists or cannot be loaded.
 		if (ScriptStruct == nullptr)
