RigVM: Array Dispatches

diff --git a/Engine/Plugins/Animation/ControlRig/Source/ControlRigDeveloper/Private/Graph/ControlRigGraphNode.cpp b/Engine/Plugins/Animation/ControlRig/Source/ControlRigDeveloper/Private/Graph/ControlRigGraphNode.cpp
--- a/Engine/Plugins/Animation/ControlRig/Source/ControlRigDeveloper/Private/Graph/ControlRigGraphNode.cpp
+++ b/Engine/Plugins/Animation/ControlRig/Source/ControlRigDeveloper/Private/Graph/ControlRigGraphNode.cpp
@@ -32,6 +32,7 @@
 #include "RigVMModel/Nodes/RigVMFunctionReturnNode.h"
 #include "RigVMModel/Nodes/RigVMCollapseNode.h"
 #include "RigVMModel/Nodes/RigVMInvokeEntryNode.h"
+#include "RigVMModel/Nodes/RigVMDispatchNode.h"
 
 #include UE_INLINE_GENERATED_CPP_BY_NAME(ControlRigGraphNode)
 
@@ -1237,6 +1238,7 @@ bool UControlRigGraphNode::ShowPaletteIconOnNode() const
 			ModelNode->IsA<URigVMFunctionReferenceNode>() ||
 			ModelNode->IsA<URigVMCollapseNode>() ||
 			ModelNode->IsA<URigVMUnitNode>() ||
+			ModelNode->IsA<URigVMDispatchNode>() ||
 			ModelNode->IsLoopNode();
 	}
 	return false;
@@ -1250,7 +1252,6 @@ FSlateIcon UControlRigGraphNode::GetIconAndTint(FLinearColor& OutColor) const
 	static FSlateIcon EventIcon(FAppStyle::GetAppStyleSetName(), "GraphEditor.Event_16x");
 	static FSlateIcon EntryReturnIcon(FAppStyle::GetAppStyleSetName(), "GraphEditor.Default_16x");
 	static FSlateIcon CollapsedNodeIcon(FAppStyle::GetAppStyleSetName(), "GraphEditor.SubGraph_16x");
-	static FSlateIcon ArrayNodeIteratorIcon(FAppStyle::GetAppStyleSetName(), "GraphEditor.Macro.ForEach_16x");
 	static FSlateIcon TemplateNodeIcon("ControlRigEditorStyle", "ControlRig.Template");
 
 	if (URigVMNode* ModelNode = GetModelNode())
@@ -1281,24 +1282,18 @@ FSlateIcon UControlRigGraphNode::GetIconAndTint(FLinearColor& OutColor) const
 			return EntryReturnIcon;
 		}
 
-		if (URigVMArrayNode* ArrayNode = Cast<URigVMArrayNode>(ModelNode))
+		const UScriptStruct* MetadataScriptStruct = nullptr;
+		if (const URigVMUnitNode* UnitNode = Cast<URigVMUnitNode>(ModelNode))
 		{
-			if(ArrayNode->IsLoopNode())
-			{
-				return ArrayNodeIteratorIcon;
-			}
+			MetadataScriptStruct = UnitNode->GetScriptStruct();
+		}
+		else if (const URigVMDispatchNode* DispatchNode = Cast<URigVMDispatchNode>(ModelNode))
+		{
+			MetadataScriptStruct = DispatchNode->GetFactory()->GetScriptStruct();
 		}
 
-		if (URigVMUnitNode *UnitNode = Cast<URigVMUnitNode>(ModelNode))
+		if(MetadataScriptStruct && MetadataScriptStruct->HasMetaDataHierarchical(FRigVMStruct::IconMetaName))
 		{
-			if(const FRigVMTemplate* Template = UnitNode->GetTemplate())
-			{
-				if(Template->NumPermutations() > 1)
-				{
-					return TemplateNodeIcon;
-				}
-			}
-			
 			FString IconPath;
 			const int32 NumOfIconPathNames = 4;
 			
@@ -1309,31 +1304,40 @@ FSlateIcon UControlRigGraphNode::GetIconAndTint(FLinearColor& OutColor) const
 				NAME_None  // StatusOverlayStyleName
 			};
 
-			if(UnitNode->GetScriptStruct())
-			{
-				// icon path format: StyleSetName|StyleName|SmallStyleName|StatusOverlayStyleName
-				// the last two names are optional, see FSlateIcon() for reference
-				UnitNode->GetScriptStruct()->GetStringMetaDataHierarchical(FRigVMStruct::IconMetaName, &IconPath);
+			// icon path format: StyleSetName|StyleName|SmallStyleName|StatusOverlayStyleName
+			// the last two names are optional, see FSlateIcon() for reference
+			MetadataScriptStruct->GetStringMetaDataHierarchical(FRigVMStruct::IconMetaName, &IconPath);
+
+			int32 NameIndex = 0;
 
-				int32 NameIndex = 0;
+			while (!IconPath.IsEmpty() && NameIndex < NumOfIconPathNames)
+			{
+				FString Left;
+				FString Right;
 
-				while (!IconPath.IsEmpty() && NameIndex < NumOfIconPathNames)
+				if (!IconPath.Split(TEXT("|"), &Left, &Right))
 				{
-					FString Left;
-					FString Right;
+					Left = IconPath;
+				}
 
-					if (!IconPath.Split(TEXT("|"), &Left, &Right))
-					{
-						Left = IconPath;
-					}
+				IconPathNames[NameIndex] = FName(*Left);
 
-					IconPathNames[NameIndex] = FName(*Left);
+				NameIndex++;
+				IconPath = Right;
+			}
+			
+			return FSlateIcon(IconPathNames[0], IconPathNames[1], IconPathNames[2], IconPathNames[3]);
+		}
 
-					NameIndex++;
-					IconPath = Right;
+		if (const URigVMTemplateNode* TemplateNode = Cast<URigVMTemplateNode>(ModelNode))
+		{
+			if(const FRigVMTemplate* Template = TemplateNode->GetTemplate())
+			{
+				if(Template->NumPermutations() > 1)
+				{
+					return TemplateNodeIcon;
 				}
 			}
-			return FSlateIcon(IconPathNames[0], IconPathNames[1], IconPathNames[2], IconPathNames[3]);
 		}
 	}
 
diff --git a/Engine/Plugins/Animation/ControlRig/Source/ControlRigEditor/Private/ControlRigEditorModule.cpp b/Engine/Plugins/Animation/ControlRig/Source/ControlRigEditor/Private/ControlRigEditorModule.cpp
--- a/Engine/Plugins/Animation/ControlRig/Source/ControlRigEditor/Private/ControlRigEditorModule.cpp
+++ b/Engine/Plugins/Animation/ControlRig/Source/ControlRigEditor/Private/ControlRigEditorModule.cpp
@@ -52,7 +52,6 @@
 #include "Graph/NodeSpawners/ControlRigTemplateNodeSpawner.h"
 #include "Graph/NodeSpawners/ControlRigEnumNodeSpawner.h"
 #include "Graph/NodeSpawners/ControlRigFunctionRefNodeSpawner.h"
-#include "Graph/NodeSpawners/ControlRigArrayNodeSpawner.h"
 #include "Graph/NodeSpawners/ControlRigInvokeEntryNodeSpawner.h"
 #include "Kismet2/BlueprintEditorUtils.h"
 #include "Kismet2/KismetDebugUtilities.h"
@@ -1115,21 +1114,6 @@ void FControlRigEditorModule::GetTypeActions(UControlRigBlueprint* CRB, FBluepri
 		LOCTEXT("RerouteSpawnerTooltip", "Adds a new reroute node to the graph"));
 	ActionRegistrar.AddBlueprintAction(ActionKey, RerouteNodeSpawner);
 
-	const int32 FirstArrayOpCode = (int32)ERigVMOpCode::FirstArrayOpCode; 
-	const int32 LastArrayOpCode = (int32)ERigVMOpCode::LastArrayOpCode;
-	for(int32 OpCodeIndex = FirstArrayOpCode; OpCodeIndex <= LastArrayOpCode; OpCodeIndex++)
-	{
-		ERigVMOpCode OpCode = (ERigVMOpCode)OpCodeIndex;
-		FString OpCodeString = URigVMArrayNode::GetNodeTitle(OpCode);
-
-		UBlueprintNodeSpawner* ArrayNodeSpawner = UControlRigArrayNodeSpawner::CreateGeneric(
-			OpCode,
-			FText::FromString(OpCodeString),
-			LOCTEXT("ArraySpawnerCategory", "Array"),
-			FText::FromString(FString::Printf(TEXT("Adds a new '%s' node to the graph"), *OpCodeString)));
-		ActionRegistrar.AddBlueprintAction(ActionKey, ArrayNodeSpawner);
-	}
-
 	for (TObjectIterator<UEnum> EnumIt; EnumIt; ++EnumIt)
 	{
 		UEnum* EnumToConsider = (*EnumIt);
@@ -1544,8 +1528,7 @@ void FControlRigEditorModule::GetContextMenuActions(const UControlRigGraphSchema
 
 					if (Cast<URigVMUnitNode>(ModelPin->GetNode()) != nullptr || 
 						Cast<URigVMDispatchNode>(ModelPin->GetNode()) != nullptr || 
-						Cast<URigVMLibraryNode>(ModelPin->GetNode()) != nullptr ||
-						Cast<URigVMArrayNode>(ModelPin->GetNode()) != nullptr)
+						Cast<URigVMLibraryNode>(ModelPin->GetNode()) != nullptr)
 					{
 						if (ModelPin->GetDirection() == ERigVMPinDirection::Input && 
 							bIsEditablePin)
diff --git a/Engine/Plugins/Animation/ControlRig/Source/ControlRigEditor/Private/Graph/NodeSpawners/ControlRigArrayNodeSpawner.cpp b/Engine/Plugins/Animation/ControlRig/Source/ControlRigEditor/Private/Graph/NodeSpawners/ControlRigArrayNodeSpawner.cpp
deleted file mode 100644
diff --git a/Engine/Plugins/Animation/ControlRig/Source/ControlRigEditor/Private/Graph/NodeSpawners/ControlRigArrayNodeSpawner.h b/Engine/Plugins/Animation/ControlRig/Source/ControlRigEditor/Private/Graph/NodeSpawners/ControlRigArrayNodeSpawner.h
deleted file mode 100644
diff --git a/Engine/Source/Developer/RigVMDeveloper/Private/RigVMCompiler/RigVMAST.cpp b/Engine/Source/Developer/RigVMDeveloper/Private/RigVMCompiler/RigVMAST.cpp
--- a/Engine/Source/Developer/RigVMDeveloper/Private/RigVMCompiler/RigVMAST.cpp
+++ b/Engine/Source/Developer/RigVMDeveloper/Private/RigVMCompiler/RigVMAST.cpp
@@ -10,7 +10,6 @@
 #include "RigVMModel/Nodes/RigVMEnumNode.h"
 #include "RigVMModel/Nodes/RigVMFunctionReturnNode.h"
 #include "RigVMModel/Nodes/RigVMFunctionEntryNode.h"
-#include "RigVMModel/Nodes/RigVMArrayNode.h"
 #include "RigVMModel/Nodes/RigVMInvokeEntryNode.h"
 #include "RigVMModel/RigVMGraph.h"
 #include "RigVMModel/RigVMController.h"
@@ -82,10 +81,6 @@ FName FRigVMExprAST::GetTypeName() const
 		{
 			return TEXT("[.Exit..]");
 		}
-		case EType::Array:
-		{
-			return TEXT("[Array..]");
-		}
 		case EType::Invalid:
 		{
 			return TEXT("[Invalid]");
@@ -584,23 +579,11 @@ bool FRigVMVarExprAST::IsEnumValue() const
 
 bool FRigVMVarExprAST::SupportsSoftLinks() const
 {
-	if (URigVMUnitNode* UnitNode = Cast<URigVMUnitNode>(GetPin()->GetNode()))
-	{
-		if (UnitNode->IsControlFlowNode())
-		{
-			return !UnitNode->GetControlFlowBlocks().Contains(GetPin()->GetFName());
-		}
-	}
-	else if(URigVMArrayNode* ArrayNode = Cast<URigVMArrayNode>(GetPin()->GetNode()))
+	if (const URigVMNode* Node = Cast<URigVMNode>(GetPin()->GetNode()))
 	{
-		// this is still using is loop node - but that's to be deprecated soon
-		if(ArrayNode->IsLoopNode())
+		if (Node->IsControlFlowNode())
 		{
-			if (GetPin()->GetFName() != FRigVMStruct::ExecuteContextName &&
-				GetPin()->GetName() != URigVMArrayNode::CompletedName)
-			{
-				return true;
-			}
+			return !Node->GetControlFlowBlocks().Contains(GetPin()->GetFName());
 		}
 	}
 	return false;
@@ -956,10 +939,6 @@ FRigVMExprAST* FRigVMParserAST::CreateExpressionForNode(const FRigVMASTProxy& In
 		{
 			NodeExpr = MakeExpr<FRigVMNoOpExprAST>(InNodeProxy);
 		}
-		else if (InNodeProxy.IsA<URigVMArrayNode>())
-		{
-			NodeExpr = MakeExpr<FRigVMArrayExprAST>(InNodeProxy);
-		}
 		else if (InNodeProxy.IsA<URigVMInvokeEntryNode>())
 		{
 			NodeExpr = MakeExpr<FRigVMInvokeEntryExprAST>(InNodeProxy);
@@ -1274,8 +1253,7 @@ FRigVMExprAST* FRigVMParserAST::TraverseLink(int32 InLinkIndex, FRigVMExprAST* I
 	{
 		// if this is a copy expression - we should require the copy to use a ref instead
 		if (NodeExpr->IsA(FRigVMExprAST::EType::CallExtern) ||
-			NodeExpr->IsA(FRigVMExprAST::EType::InlineFunction) ||
-			NodeExpr->IsA(FRigVMExprAST::EType::Array))
+			NodeExpr->IsA(FRigVMExprAST::EType::InlineFunction))
 		{
 			for (FRigVMExprAST* ChildExpr : *NodeExpr)
 			{
@@ -1691,9 +1669,9 @@ void FRigVMParserAST::FoldAssignments()
 		}
 
 		// if this node is a loop node - let's skip the folding
-		if (URigVMUnitNode* UnitNode = Cast<URigVMUnitNode>(TargetPin->GetNode()))
+		if (const URigVMNode* ModelNode = TargetPin->GetNode())
 		{
-			if (UnitNode->IsControlFlowNode())
+			if (ModelNode->IsControlFlowNode())
 			{
 				continue;
 			}
@@ -1707,16 +1685,6 @@ void FRigVMParserAST::FoldAssignments()
 				continue;
 			}
 		}
-		
-		// if this node is an array iterator node - let's skip the folding
-		if (URigVMArrayNode* ArrayNode = Cast<URigVMArrayNode>(TargetPin->GetNode()))
-		{
-			// still using IsLoopNode - but to be deprecated soon 
-			if (ArrayNode->IsLoopNode())
-			{
-				continue;
-			}
-		}
 
 		FRigVMExprAST* Parent = AssignExpr->Parents[0];
 		if (!Parent->IsA(FRigVMExprAST::EType::Var))
@@ -1865,6 +1833,25 @@ bool FRigVMParserAST::FoldConstantValuesToLiterals(TArray<URigVMGraph*> InGraphs
 			for (const int32 LinkIndex : LinkIndices)
 			{
 				const FRigVMASTProxy& SourcePinProxy = Links[LinkIndex].SourceProxy;
+
+				// don't compute values which are also wired to an sub pins
+				bool bHasLinkToSubPin = false;
+				const TArray<int32> CurrentTargetLinkIndices = GetTargetLinkIndices(SourcePinProxy, true);
+				for(int32 TargetLinkIndex : CurrentTargetLinkIndices)
+				{
+					if(!Links[TargetLinkIndex].SegmentPath.IsEmpty())
+					{
+						bHasLinkToSubPin = true;
+						break;
+					}
+				}
+
+				if(bHasLinkToSubPin)
+				{
+					bFoundValidSourcePin = false;
+					break;
+				}
+				
 				URigVMNode* SourceNode = SourcePinProxy.GetSubjectChecked<URigVMPin>()->GetNode();
 				FRigVMASTProxy SourceNodeProxy = SourcePinProxy.GetSibling(SourceNode);
 
@@ -2608,12 +2595,6 @@ FString FRigVMParserAST::DumpDot() const
 					Label = TEXT("NoOp");
 					break;
 				}
-				case FRigVMExprAST::EType::Array:
-				{
-					ERigVMOpCode OpCode = CastChecked<URigVMArrayNode>(InExpr->To<FRigVMArrayExprAST>()->GetNode())->GetOpCode();
-					Label = StaticEnum<ERigVMOpCode>()->GetDisplayNameTextByValue((int64)OpCode).ToString();
-					break;
-				}
 				case FRigVMExprAST::EType::Exit:
 				{
 					Label = TEXT("Exit");
diff --git a/Engine/Source/Developer/RigVMDeveloper/Private/RigVMCompiler/RigVMCodeGenerator.cpp b/Engine/Source/Developer/RigVMDeveloper/Private/RigVMCompiler/RigVMCodeGenerator.cpp
--- a/Engine/Source/Developer/RigVMDeveloper/Private/RigVMCompiler/RigVMCodeGenerator.cpp
+++ b/Engine/Source/Developer/RigVMDeveloper/Private/RigVMCompiler/RigVMCodeGenerator.cpp
@@ -86,23 +86,6 @@ static constexpr TCHAR RigVM_JumpIfOpFormat[] = TEXT("\tif ({0} == {1}) { goto I
 static constexpr TCHAR RigVM_JumpToBranchFormat[] = TEXT("\tif ({0} == BlockName_{1}) { goto Instruction{2}Label; }");
 static constexpr TCHAR RigVM_BeginBlockOpFormat[] = TEXT("\tBeginSlice({0}, {1});");
 static constexpr TCHAR RigVM_EndBlockOpFormat[] = TEXT("\tEndSlice();");
-static constexpr TCHAR RigVM_ArrayResetOpFormat[] = TEXT("\t{0}.Reset();");
-static constexpr TCHAR RigVM_ArrayGetNumOpFormat[] = TEXT("\t{0} = {1}.Num();");
-static constexpr TCHAR RigVM_ArraySetNumOpFormat[] = TEXT("\tif ({0}) { {1}.SetNum({2}); }");
-static constexpr TCHAR RigVM_ArrayGetAtIndexOpFormat[] = TEXT("\tif ({0}) { {1} = {2}[TemporaryArrayIndex]; }");
-static constexpr TCHAR RigVM_ArraySetAtIndexOpFormat[] = TEXT("\tif ({0}) { {1}[TemporaryArrayIndex] = {2}; }");
-static constexpr TCHAR RigVM_ArrayAddOpFormat[] = TEXT("\tif ({0}) { {1} = {2}.Add({3}); }");
-static constexpr TCHAR RigVM_ArrayNumPlusOneFormat[] = TEXT("{0}.Num() + 1");
-static constexpr TCHAR RigVM_ArrayInsertOpFormat[] = TEXT("\tif ({0}) { {1}.Insert({2}, {3}); }");
-static constexpr TCHAR RigVM_ArrayRemoveOpFormat[] = TEXT("\t{0}.RemoveAt({1});");
-static constexpr TCHAR RigVM_ArrayFindOpFormat[] = TEXT("\t{0} = ArrayOp_Find{1}({2}, {3}, {4});");
-static constexpr TCHAR RigVM_ArrayAppendOpFormat[] = TEXT("\tif ({0} { {1}.Append({2}); }");
-static constexpr TCHAR RigVM_ArrayAppendOpNumFormat[] = TEXT("{0}.Num() + {1}.Num()");
-static constexpr TCHAR RigVM_ArrayIteratorOpFormat[] = TEXT("\t{0} = ArrayOp_Iterator{1}({2}, {3}, {4}, {5}, {6});");
-static constexpr TCHAR RigVM_ArrayUnionOpFormat[] = TEXT("\tArrayOp_Union{0}({1}, {2});");
-static constexpr TCHAR RigVM_ArrayDifferenceOpFormat[] = TEXT("\tArrayOp_Difference{0}({1}, {2});");
-static constexpr TCHAR RigVM_ArrayIntersectionOpFormat[] = TEXT("\tArrayOp_Intersection{0}({1}, {2});");
-static constexpr TCHAR RigVM_ArrayReverseOpFormat[] = TEXT("\tArrayOp_Reverse{0}({1});");
 static constexpr TCHAR RigVM_ReturnFailedFormat[] = TEXT("\treturn ERigVMExecuteResult::Failed;");
 static constexpr TCHAR RigVM_ReturnSucceededFormat[] = TEXT("\treturn ERigVMExecuteResult::Succeeded;");
 static constexpr TCHAR RigVM_CopyrightFormat[] = TEXT("// Copyright Epic Games, Inc. All Rights Reserved.");
@@ -804,7 +787,6 @@ FString FRigVMCodeGenerator::DumpInstructions(const FString& InPrefix, const TAr
 				break;
 			}
 			case ERigVMOpCode::Copy:
-			case ERigVMOpCode::ArrayClone:
 			{
 				const FRigVMBinaryOp& Op = ByteCode.GetOpAt<FRigVMBinaryOp>(Instruction);
 				const FString TargetOperand = GetOperandName(Op.ArgB, false, false);
@@ -917,124 +899,6 @@ FString FRigVMCodeGenerator::DumpInstructions(const FString& InPrefix, const TAr
 				Lines.Add(Prefix + RigVM_EndBlockOpFormat);
 				break;
 			}
-			case ERigVMOpCode::ArrayReset:
-			{
-				const FRigVMUnaryOp& Op = ByteCode.GetOpAt<FRigVMUnaryOp>(Instruction);
-				Lines.Add(Prefix + Format(RigVM_ArrayResetOpFormat, *GetOperandName(Op.Arg, false)));
-				break;
-			}
-			case ERigVMOpCode::ArrayGetNum:
-			{
-				const FRigVMBinaryOp& Op = ByteCode.GetOpAt<FRigVMBinaryOp>(Instruction);
-				Lines.Add(Prefix + Format(RigVM_ArrayGetNumOpFormat, *GetOperandName(Op.ArgB, false), *GetOperandName(Op.ArgA, false)));
-				break;
-			} 
-			case ERigVMOpCode::ArraySetNum:
-			{
-				const FRigVMBinaryOp& Op = ByteCode.GetOpAt<FRigVMBinaryOp>(Instruction);
-				const FString ConditionString = Format(RigVM_IsValidArraySizeFormat, *GetOperandName(Op.ArgB, false));
-				Lines.Add(Prefix + Format(RigVM_ArraySetNumOpFormat, *ConditionString, *GetOperandName(Op.ArgA, false), *GetOperandName(Op.ArgB, false)));
-				break;
-			}
-			case ERigVMOpCode::ArrayGetAtIndex:
-			{
-				const FRigVMTernaryOp& Op = ByteCode.GetOpAt<FRigVMTernaryOp>(Instruction);
-				const FString ConditionString = Format(RigVM_IsValidArrayIndexFormat, *GetOperandCPPBaseType(Op.ArgA), *GetOperandName(Op.ArgA, false));
-				Lines.Add(Prefix + Format(RigVM_TemporaryArrayIndexFormat, *GetOperandName(Op.ArgB, false)));
-				Lines.Add(Prefix + Format(RigVM_ArrayGetAtIndexOpFormat, *ConditionString, *GetOperandName(Op.ArgC, false), *GetOperandName(Op.ArgA, false)));
-				break;
-			}  
-			case ERigVMOpCode::ArraySetAtIndex:
-			{
-				const FRigVMTernaryOp& Op = ByteCode.GetOpAt<FRigVMTernaryOp>(Instruction);
-				const FString ConditionString = Format(RigVM_IsValidArrayIndexFormat,  *GetOperandCPPBaseType(Op.ArgA), *GetOperandName(Op.ArgA, false));
-				Lines.Add(Prefix + Format(RigVM_TemporaryArrayIndexFormat, *GetOperandName(Op.ArgB, false)));
-				Lines.Add(Prefix + Format(RigVM_ArraySetAtIndexOpFormat, *ConditionString, *GetOperandName(Op.ArgA, false), *GetOperandName(Op.ArgC, false), *GetOperandName(Op.ArgC, false)));
-				break;
-			}
-			case ERigVMOpCode::ArrayAdd:
-			{
-				const FRigVMTernaryOp& Op = ByteCode.GetOpAt<FRigVMTernaryOp>(Instruction);
-				const FString NumString = Format(RigVM_ArrayNumPlusOneFormat, *GetOperandName(Op.ArgA, false));
-				const FString ConditionString = Format(RigVM_IsValidArraySizeFormat, *NumString);
-				Lines.Add(Prefix + Format(RigVM_ArrayAddOpFormat, *ConditionString, *GetOperandName(Op.ArgC, false), *GetOperandName(Op.ArgA, false), *GetOperandName(Op.ArgB, false)));
-				break;
-			}
-			case ERigVMOpCode::ArrayInsert:
-			{
-				const FRigVMTernaryOp& Op = ByteCode.GetOpAt<FRigVMTernaryOp>(Instruction);
-				const FString NumString = Format(RigVM_ArrayNumPlusOneFormat, *GetOperandName(Op.ArgA, false));
-				const FString ConditionString = Format(RigVM_IsValidArraySizeFormat, *NumString);
-				Lines.Add(Prefix + Format(RigVM_ArrayInsertOpFormat, *ConditionString, *GetOperandName(Op.ArgA, false), *GetOperandName(Op.ArgC, false), *GetOperandName(Op.ArgB, false)));
-				break;
-			}
-			case ERigVMOpCode::ArrayRemove:
-			{
-				const FRigVMBinaryOp& Op = ByteCode.GetOpAt<FRigVMBinaryOp>(Instruction);
-				Lines.Add(Prefix + Format(RigVM_ArrayRemoveOpFormat, *GetOperandName(Op.ArgA, false), *GetOperandName(Op.ArgB, false)));
-				break;
-			}
-			case ERigVMOpCode::ArrayFind:
-			{
-				const FRigVMQuaternaryOp& Op = ByteCode.GetOpAt<FRigVMQuaternaryOp>(Instruction);
-				const FString ArrayElementCPPType = RigVMTypeUtils::BaseTypeFromArrayType(GetOperandCPPBaseType(Op.ArgA));
-				const FString ElementCPPType = GetOperandCPPType(Op.ArgB);
-				const FString TemplateSuffix = ArrayElementCPPType == ElementCPPType ? Format(RigVM_TemplateOneArgFormat, *ElementCPPType) : FString();
-				Lines.Add(Prefix + Format(RigVM_ArrayFindOpFormat, *GetOperandName(Op.ArgD, false), *TemplateSuffix, *GetOperandName(Op.ArgA, false), *GetOperandName(Op.ArgB, false), *GetOperandName(Op.ArgC, false)));
-				break;
-			}
-			case ERigVMOpCode::ArrayAppend:
-			{
-				const FRigVMBinaryOp& Op = ByteCode.GetOpAt<FRigVMBinaryOp>(Instruction);
-				const FString NumString = Format(RigVM_ArrayAppendOpNumFormat, *GetOperandName(Op.ArgA, false), *GetOperandName(Op.ArgB, false));
-				const FString ConditionString = Format(RigVM_IsValidArraySizeFormat, *NumString);
-				Lines.Add(Prefix + Format(RigVM_ArrayAppendOpFormat, *ConditionString, *GetOperandName(Op.ArgA, false), *GetOperandName(Op.ArgB, false)));
-				break;
-			}
-			case ERigVMOpCode::ArrayIterator:
-			{
-				const FRigVMSenaryOp& Op = ByteCode.GetOpAt<FRigVMSenaryOp>(Instruction);
-				const FString ArrayElementCPPType = RigVMTypeUtils::BaseTypeFromArrayType(GetOperandCPPType(Op.ArgA));
-				const FString ElementCPPType = GetOperandCPPType(Op.ArgB);
-				const FString TemplateSuffix = ArrayElementCPPType == ElementCPPType ? Format(RigVM_TemplateOneArgFormat, *ElementCPPType) : FString();
-				Lines.Add(Prefix + Format(RigVM_ArrayIteratorOpFormat, *GetOperandName(Op.ArgF, false), *TemplateSuffix, *GetOperandName(Op.ArgA, false), *GetOperandName(Op.ArgB, false), *GetOperandName(Op.ArgC, false), *GetOperandName(Op.ArgD, false), *GetOperandName(Op.ArgE, false)));
-				break;
-			}
-			case ERigVMOpCode::ArrayUnion:
-			{
-				const FRigVMBinaryOp& Op = ByteCode.GetOpAt<FRigVMBinaryOp>(Instruction);
-				const FString ElementCPPTypeA = RigVMTypeUtils::BaseTypeFromArrayType(GetOperandCPPType(Op.ArgA));
-				const FString ElementCPPTypeB = RigVMTypeUtils::BaseTypeFromArrayType(GetOperandCPPType(Op.ArgA));
-				const FString TemplateSuffix = ElementCPPTypeA == ElementCPPTypeB ? Format(RigVM_TemplateOneArgFormat, *ElementCPPTypeA) : FString();
-				Lines.Add(Prefix + Format(RigVM_ArrayUnionOpFormat, *TemplateSuffix, *GetOperandName(Op.ArgA, false, false), *GetOperandName(Op.ArgB, false, false)));
-				break;
-			}
-			case ERigVMOpCode::ArrayDifference:
-			{
-				const FRigVMBinaryOp& Op = ByteCode.GetOpAt<FRigVMBinaryOp>(Instruction);
-				const FString ElementCPPTypeA = RigVMTypeUtils::BaseTypeFromArrayType(GetOperandCPPType(Op.ArgA));
-				const FString ElementCPPTypeB = RigVMTypeUtils::BaseTypeFromArrayType(GetOperandCPPType(Op.ArgA));
-				const FString TemplateSuffix = ElementCPPTypeA == ElementCPPTypeB ? Format(RigVM_TemplateOneArgFormat, *ElementCPPTypeA) : FString();
-				Lines.Add(Prefix + Format(RigVM_ArrayDifferenceOpFormat, *TemplateSuffix, *GetOperandName(Op.ArgA, false), *GetOperandName(Op.ArgB, false)));
-				break;
-			}
-			case ERigVMOpCode::ArrayIntersection:
-			{
-				const FRigVMBinaryOp& Op = ByteCode.GetOpAt<FRigVMBinaryOp>(Instruction);
-				const FString ElementCPPTypeA = RigVMTypeUtils::BaseTypeFromArrayType(GetOperandCPPType(Op.ArgA));
-				const FString ElementCPPTypeB = RigVMTypeUtils::BaseTypeFromArrayType(GetOperandCPPType(Op.ArgA));
-				const FString TemplateSuffix = ElementCPPTypeA == ElementCPPTypeB ? Format(RigVM_TemplateOneArgFormat, *ElementCPPTypeA) : FString();
-				Lines.Add(Prefix + Format(RigVM_ArrayIntersectionOpFormat, *TemplateSuffix, *GetOperandName(Op.ArgA, false), *GetOperandName(Op.ArgB, false)));
-				break;
-			}
-			case ERigVMOpCode::ArrayReverse:
-			{
-				const FRigVMUnaryOp& Op = ByteCode.GetOpAt<FRigVMUnaryOp>(Instruction);
-				const FString ElementCPPType = RigVMTypeUtils::BaseTypeFromArrayType(GetOperandCPPType(Op.Arg));
-				const FString TemplateSuffix = Format(RigVM_TemplateOneArgFormat, *ElementCPPType);
-				Lines.Add(Prefix + Format(RigVM_ArrayReverseOpFormat, *TemplateSuffix, *GetOperandName(Op.Arg, false)));
-				break;
-			}
 			case ERigVMOpCode::InvokeEntry:
 			{
 				const FRigVMInvokeEntryOp& Op = ByteCode.GetOpAt<FRigVMInvokeEntryOp>(Instruction);
diff --git a/Engine/Source/Developer/RigVMDeveloper/Private/RigVMCompiler/RigVMCompiler.cpp b/Engine/Source/Developer/RigVMDeveloper/Private/RigVMCompiler/RigVMCompiler.cpp
--- a/Engine/Source/Developer/RigVMDeveloper/Private/RigVMCompiler/RigVMCompiler.cpp
+++ b/Engine/Source/Developer/RigVMDeveloper/Private/RigVMCompiler/RigVMCompiler.cpp
@@ -4,6 +4,7 @@
 #include "RigVMModel/RigVMController.h"
 #include "RigVMModel/Nodes/RigVMDispatchNode.h"
 #include "RigVMModel/Nodes/RigVMBranchNode.h"
+#include "RigVMModel/Nodes/RigVMArrayNode.h"
 #include "RigVMCore/RigVMExecuteContext.h"
 #include "RigVMCore/RigVMNativized.h"
 #include "RigVMDeveloperModule.h"
@@ -14,6 +15,7 @@
 #include "RigVMCore/RigVMGraphFunctionDefinition.h"
 #include "ProfilingDebugging/ScopedTimers.h"
 #include "RigVMModel/RigVMClient.h"
+#include "RigVMFunctions/RigVMDispatch_Array.h"
 #include "Algo/Count.h"
 #include "String/Join.h"
 
@@ -492,6 +494,13 @@ bool URigVMCompiler::Compile(TArray<URigVMGraph*> InGraphs, URigVMController* In
 				bEncounteredGraphError = true;
 			}
 
+			if(ModelNode->IsA<UDEPRECATED_RigVMArrayNode>())
+			{
+				static const FString LinkedMessage = TEXT("Node @@ is a deprecated array node. Cannot compile.");
+				Settings.ASTSettings.Report(EMessageSeverity::Error, ModelNode, LinkedMessage);
+				bEncounteredGraphError = true;
+			}
+
 			if(!InController->RemoveUnusedOrphanedPins(ModelNode))
 			{
 				static const FString LinkedMessage = TEXT("Node @@ uses pins that no longer exist. Please rewire the links and re-compile.");
@@ -1246,11 +1255,6 @@ void URigVMCompiler::TraverseExpression(const FRigVMExprAST* InExpr, FRigVMCompi
 			TraverseExit(InExpr->To<FRigVMExitExprAST>(), WorkData);
 			break;
 		}
-		case FRigVMExprAST::EType::Array:
-		{
-			TraverseArray(InExpr->To<FRigVMArrayExprAST>(), WorkData);
-			break;
-		}
 		case FRigVMExprAST::EType::InvokeEntry:
 		{
 			TraverseInvokeEntry(InExpr->To<FRigVMInvokeEntryExprAST>(), WorkData);
@@ -2309,326 +2313,6 @@ void URigVMCompiler::TraverseExit(const FRigVMExitExprAST* InExpr, FRigVMCompile
 	}
 }
 
-void URigVMCompiler::TraverseArray(const FRigVMArrayExprAST* InExpr, FRigVMCompilerWorkData& WorkData)
-{
-	URigVMArrayNode* ArrayNode = Cast<URigVMArrayNode>(InExpr->GetNode());
-	if(!ValidateNode(ArrayNode))
-	{
-		return;
-	}
-	
-	if (WorkData.bSetupMemory)
-	{
-		TraverseChildren(InExpr, WorkData);
-	}
-	else
-	{
-		const FRigVMCallstack Callstack = InExpr->GetProxy().GetCallstack();
-
-		static const FName ExecuteName = FRigVMStruct::ExecuteName;
-		static const FName ArrayName = *URigVMArrayNode::ArrayName;
-		static const FName NumName = *URigVMArrayNode::NumName;
-		static const FName IndexName = *URigVMArrayNode::IndexName;
-		static const FName ElementName = *URigVMArrayNode::ElementName;
-		static const FName SuccessName = *URigVMArrayNode::SuccessName;
-		static const FName OtherName = *URigVMArrayNode::OtherName;
-		static const FName CloneName = *URigVMArrayNode::CloneName;
-		static const FName CountName = *URigVMArrayNode::CountName;
-		static const FName RatioName = *URigVMArrayNode::RatioName;
-		static const FName ResultName = *URigVMArrayNode::ResultName;
-		static const FName ContinueName = *URigVMArrayNode::ContinueName;
-		static const FName CompletedName = *URigVMArrayNode::CompletedName;
-
-		const ERigVMOpCode OpCode = ArrayNode->GetOpCode();
-		switch(OpCode)
-		{
-			case ERigVMOpCode::ArrayReset:
-			{
-				TraverseChildren(InExpr, WorkData);
-				const FRigVMOperand& ArrayOperand = WorkData.ExprToOperand.FindChecked(GetSourceVarExpr(InExpr->ChildAt(1)));
-				WorkData.VM->GetByteCode().AddArrayResetOp(ArrayOperand);
-				if (Settings.SetupNodeInstructionIndex)
-				{
-					WorkData.VM->GetByteCode().SetSubject(WorkData.VM->GetByteCode().GetNumInstructions() - 1, Callstack.GetCallPath(), Callstack.GetStack());
-				}
-				break;
-			}
-			case ERigVMOpCode::ArrayGetNum:
-			{
-				TraverseChildren(InExpr, WorkData);
-				const FRigVMOperand& ArrayOperand = WorkData.ExprToOperand.FindChecked(GetSourceVarExpr(InExpr->ChildAt(0)));
-				const FRigVMOperand& NumOperand = WorkData.ExprToOperand.FindChecked(GetSourceVarExpr(InExpr->ChildAt(1)));
-				WorkData.VM->GetByteCode().AddArrayGetNumOp(ArrayOperand, NumOperand);
-				if (Settings.SetupNodeInstructionIndex)
-				{
-					WorkData.VM->GetByteCode().SetSubject(WorkData.VM->GetByteCode().GetNumInstructions() - 1, Callstack.GetCallPath(), Callstack.GetStack());
-				}
-				break;
-			}
-			case ERigVMOpCode::ArraySetNum:
-			{
-				TraverseChildren(InExpr, WorkData);
-				const FRigVMOperand& ArrayOperand = WorkData.ExprToOperand.FindChecked(GetSourceVarExpr(InExpr->ChildAt(1)));
-				const FRigVMOperand& NumOperand = WorkData.ExprToOperand.FindChecked(GetSourceVarExpr(InExpr->ChildAt(2)));
-				WorkData.VM->GetByteCode().AddArraySetNumOp(ArrayOperand, NumOperand);
-				if (Settings.SetupNodeInstructionIndex)
-				{
-					WorkData.VM->GetByteCode().SetSubject(WorkData.VM->GetByteCode().GetNumInstructions() - 1, Callstack.GetCallPath(), Callstack.GetStack());
-				}
-				break;
-			}
-			case ERigVMOpCode::ArrayGetAtIndex:
-			{
-				TraverseChildren(InExpr, WorkData);
-				const FRigVMOperand& ArrayOperand = WorkData.ExprToOperand.FindChecked(GetSourceVarExpr(InExpr->ChildAt(0)));
-				const FRigVMOperand& IndexOperand = WorkData.ExprToOperand.FindChecked(GetSourceVarExpr(InExpr->ChildAt(1)));
-				const FRigVMOperand& ElementOperand = WorkData.ExprToOperand.FindChecked(GetSourceVarExpr(InExpr->ChildAt(2)));
-				WorkData.VM->GetByteCode().AddArrayGetAtIndexOp(ArrayOperand, IndexOperand, ElementOperand);
-				if (Settings.SetupNodeInstructionIndex)
-				{
-					WorkData.VM->GetByteCode().SetSubject(WorkData.VM->GetByteCode().GetNumInstructions() - 1, Callstack.GetCallPath(), Callstack.GetStack());
-				}
-				break;
-			}
-			case ERigVMOpCode::ArraySetAtIndex:
-			{
-				TraverseChildren(InExpr, WorkData);
-				const FRigVMOperand& ArrayOperand = WorkData.ExprToOperand.FindChecked(GetSourceVarExpr(InExpr->ChildAt(1)));
-				const FRigVMOperand& IndexOperand = WorkData.ExprToOperand.FindChecked(GetSourceVarExpr(InExpr->ChildAt(2)));
-				const FRigVMOperand& ElementOperand = WorkData.ExprToOperand.FindChecked(GetSourceVarExpr(InExpr->ChildAt(3)));
-				WorkData.VM->GetByteCode().AddArraySetAtIndexOp(ArrayOperand, IndexOperand, ElementOperand);
-				if (Settings.SetupNodeInstructionIndex)
-				{
-					WorkData.VM->GetByteCode().SetSubject(WorkData.VM->GetByteCode().GetNumInstructions() - 1, Callstack.GetCallPath(), Callstack.GetStack());
-				}
-				break;
-			}
-			case ERigVMOpCode::ArrayAdd:
-			{
-				TraverseChildren(InExpr, WorkData);
-				const FRigVMOperand& ArrayOperand = WorkData.ExprToOperand.FindChecked(GetSourceVarExpr(InExpr->ChildAt(1)));
-				const FRigVMOperand& ElementOperand = WorkData.ExprToOperand.FindChecked(GetSourceVarExpr(InExpr->ChildAt(2)));
-				const FRigVMOperand& IndexOperand = WorkData.ExprToOperand.FindChecked(GetSourceVarExpr(InExpr->ChildAt(3)));
-				WorkData.VM->GetByteCode().AddArrayAddOp(ArrayOperand, ElementOperand, IndexOperand);
-				if (Settings.SetupNodeInstructionIndex)
-				{
-					WorkData.VM->GetByteCode().SetSubject(WorkData.VM->GetByteCode().GetNumInstructions() - 1, Callstack.GetCallPath(), Callstack.GetStack());
-				}
-				break;
-			}
-			case ERigVMOpCode::ArrayInsert:
-			{
-				TraverseChildren(InExpr, WorkData);
-				const FRigVMOperand& ArrayOperand = WorkData.ExprToOperand.FindChecked(GetSourceVarExpr(InExpr->ChildAt(1)));
-				const FRigVMOperand& IndexOperand = WorkData.ExprToOperand.FindChecked(GetSourceVarExpr(InExpr->ChildAt(2)));
-				const FRigVMOperand& ElementOperand = WorkData.ExprToOperand.FindChecked(GetSourceVarExpr(InExpr->ChildAt(3)));
-				WorkData.VM->GetByteCode().AddArrayInsertOp(ArrayOperand, IndexOperand, ElementOperand);
-				if (Settings.SetupNodeInstructionIndex)
-				{
-					WorkData.VM->GetByteCode().SetSubject(WorkData.VM->GetByteCode().GetNumInstructions() - 1, Callstack.GetCallPath(), Callstack.GetStack());
-				}
-				break;
-			}
-			case ERigVMOpCode::ArrayRemove:
-			{
-				TraverseChildren(InExpr, WorkData);
-				const FRigVMOperand& ArrayOperand = WorkData.ExprToOperand.FindChecked(GetSourceVarExpr(InExpr->ChildAt(1)));
-				const FRigVMOperand& IndexOperand = WorkData.ExprToOperand.FindChecked(GetSourceVarExpr(InExpr->ChildAt(2)));
-				WorkData.VM->GetByteCode().AddArrayRemoveOp(ArrayOperand, IndexOperand);
-				if (Settings.SetupNodeInstructionIndex)
-				{
-					WorkData.VM->GetByteCode().SetSubject(WorkData.VM->GetByteCode().GetNumInstructions() - 1, Callstack.GetCallPath(), Callstack.GetStack());
-				}
-				break;
-			}
-			case ERigVMOpCode::ArrayFind:
-			{
-				TraverseChildren(InExpr, WorkData);
-				const FRigVMOperand& ArrayOperand = WorkData.ExprToOperand.FindChecked(GetSourceVarExpr(InExpr->ChildAt(0)));
-				const FRigVMOperand& ElementOperand = WorkData.ExprToOperand.FindChecked(GetSourceVarExpr(InExpr->ChildAt(1)));
-				const FRigVMOperand& IndexOperand = WorkData.ExprToOperand.FindChecked(GetSourceVarExpr(InExpr->ChildAt(2)));
-				const FRigVMOperand& SuccessOperand = WorkData.ExprToOperand.FindChecked(GetSourceVarExpr(InExpr->ChildAt(3)));
-				WorkData.VM->GetByteCode().AddArrayFindOp(ArrayOperand, ElementOperand, IndexOperand, SuccessOperand);
-				if (Settings.SetupNodeInstructionIndex)
-				{
-					WorkData.VM->GetByteCode().SetSubject(WorkData.VM->GetByteCode().GetNumInstructions() - 1, Callstack.GetCallPath(), Callstack.GetStack());
-				}
-				break;
-			}
-			case ERigVMOpCode::ArrayAppend:
-			{
-				TraverseChildren(InExpr, WorkData);
-				const FRigVMOperand& ArrayOperand = WorkData.ExprToOperand.FindChecked(GetSourceVarExpr(InExpr->ChildAt(1)));
-				const FRigVMOperand& OtherOperand = WorkData.ExprToOperand.FindChecked(GetSourceVarExpr(InExpr->ChildAt(2)));
-				WorkData.VM->GetByteCode().AddArrayAppendOp(ArrayOperand, OtherOperand);
-				if (Settings.SetupNodeInstructionIndex)
-				{
-					WorkData.VM->GetByteCode().SetSubject(WorkData.VM->GetByteCode().GetNumInstructions() - 1, Callstack.GetCallPath(), Callstack.GetStack());
-				}
-				break;
-			}
-			case ERigVMOpCode::ArrayClone:
-			{
-				TraverseChildren(InExpr, WorkData);
-				const FRigVMOperand& ArrayOperand = WorkData.ExprToOperand.FindChecked(GetSourceVarExpr(InExpr->ChildAt(0)));
-				const FRigVMOperand& CloneOperand = WorkData.ExprToOperand.FindChecked(GetSourceVarExpr(InExpr->ChildAt(1)));
-				WorkData.VM->GetByteCode().AddArrayCloneOp(ArrayOperand, CloneOperand);
-				if (Settings.SetupNodeInstructionIndex)
-				{
-					WorkData.VM->GetByteCode().SetSubject(WorkData.VM->GetByteCode().GetNumInstructions() - 1, Callstack.GetCallPath(), Callstack.GetStack());
-				}
-				break;
-			}
-			case ERigVMOpCode::ArrayIterator:
-			{
-				const FRigVMExprAST* ExecuteExpr = InExpr->ChildAt(0);
-				const FRigVMExprAST* ArrayExpr = InExpr->ChildAt(1);
-				const FRigVMExprAST* ElementExpr = InExpr->ChildAt(2);
-				const FRigVMExprAST* IndexExpr = InExpr->ChildAt(3);
-				const FRigVMExprAST* CountExpr = InExpr->ChildAt(4);
-				const FRigVMExprAST* RatioExpr = InExpr->ChildAt(5);
-				const FRigVMExprAST* ContinueExpr = InExpr->ChildAt(6);
-				const FRigVMExprAST* CompletedExpr = InExpr->ChildAt(7);
-				const FRigVMOperand& ArrayOperand = WorkData.ExprToOperand.FindChecked(GetSourceVarExpr(ArrayExpr));
-				const FRigVMOperand& ElementOperand = WorkData.ExprToOperand.FindChecked(GetSourceVarExpr(ElementExpr));
-				const FRigVMOperand& IndexOperand = WorkData.ExprToOperand.FindChecked(GetSourceVarExpr(IndexExpr));
-				const FRigVMOperand& CountOperand = WorkData.ExprToOperand.FindChecked(GetSourceVarExpr(CountExpr));
-				const FRigVMOperand& RatioOperand = WorkData.ExprToOperand.FindChecked(GetSourceVarExpr(RatioExpr));
-				const FRigVMOperand& ContinueOperand = WorkData.ExprToOperand.FindChecked(GetSourceVarExpr(ContinueExpr));
-
-				WorkData.ExprToSkip.AddUnique(ExecuteExpr);
-				WorkData.ExprToSkip.AddUnique(CompletedExpr);
-
-				// traverse the input array
-				TraverseExpression(ArrayExpr, WorkData);
-
-				// zero the index
-				WorkData.VM->GetByteCode().AddZeroOp(IndexOperand);
-				if (Settings.SetupNodeInstructionIndex)
-				{
-					WorkData.VM->GetByteCode().SetSubject(WorkData.VM->GetByteCode().GetNumInstructions() - 1, Callstack.GetCallPath(), Callstack.GetStack());
-				}
-
-				// add the iterator
-				WorkData.VM->GetByteCode().AddArrayIteratorOp(ArrayOperand, ElementOperand, IndexOperand, CountOperand, RatioOperand, ContinueOperand);
-				const int32 IteratorInstruction = WorkData.VM->GetByteCode().GetNumInstructions() - 1;
-				if (Settings.SetupNodeInstructionIndex)
-				{
-					WorkData.VM->GetByteCode().SetSubject(WorkData.VM->GetByteCode().GetNumInstructions() - 1, Callstack.GetCallPath(), Callstack.GetStack());
-				}
-
-				// jump to the end of the loop
-				const uint64 JumpToEndByte = WorkData.VM->GetByteCode().AddJumpIfOp(ERigVMOpCode::JumpForwardIf, 0, ContinueOperand, false);
-				const int32 JumpToEndInstruction = WorkData.VM->GetByteCode().GetNumInstructions() - 1;
-				if (Settings.SetupNodeInstructionIndex)
-				{
-					WorkData.VM->GetByteCode().SetSubject(WorkData.VM->GetByteCode().GetNumInstructions() - 1, Callstack.GetCallPath(), Callstack.GetStack());
-				}
-
-				// begin the block
-				WorkData.VM->GetByteCode().AddBeginBlockOp(CountOperand, IndexOperand);
-				if (Settings.SetupNodeInstructionIndex)
-				{
-					WorkData.VM->GetByteCode().SetSubject(WorkData.VM->GetByteCode().GetNumInstructions() - 1, Callstack.GetCallPath(), Callstack.GetStack());
-				}
-
-				// traverse the per iteration instructions
-				WorkData.ExprToSkip.Remove(ExecuteExpr);
-				TraverseExpression(ExecuteExpr, WorkData);
-
-				// end the block
-				WorkData.VM->GetByteCode().AddEndBlockOp();
-				if (Settings.SetupNodeInstructionIndex)
-				{
-					WorkData.VM->GetByteCode().SetSubject(WorkData.VM->GetByteCode().GetNumInstructions() - 1, Callstack.GetCallPath(), Callstack.GetStack());
-				}
-
-				// increment index per loop iteration
-				WorkData.VM->GetByteCode().AddIncrementOp(IndexOperand);
-				if (Settings.SetupNodeInstructionIndex)
-				{
-					WorkData.VM->GetByteCode().SetSubject(WorkData.VM->GetByteCode().GetNumInstructions() - 1, Callstack.GetCallPath(), Callstack.GetStack());
-
-				}
-
-				// jump backwards instruction (to the beginning of the iterator)
-				const int32 JumpToStartInstruction = WorkData.VM->GetByteCode().GetNumInstructions();
-				WorkData.VM->GetByteCode().AddJumpOp(ERigVMOpCode::JumpBackward, JumpToStartInstruction - IteratorInstruction);
-				if (Settings.SetupNodeInstructionIndex)
-				{
-					WorkData.VM->GetByteCode().SetSubject(WorkData.VM->GetByteCode().GetNumInstructions() - 1, Callstack.GetCallPath(), Callstack.GetStack());
-				}
-
-				// fix up the first jump instruction
-				const int32 InstructionsToEnd = WorkData.VM->GetByteCode().GetNumInstructions() - JumpToEndInstruction;
-				WorkData.VM->GetByteCode().GetOpAt<FRigVMJumpIfOp>(JumpToEndByte).InstructionIndex = InstructionsToEnd;
-				if (Settings.SetupNodeInstructionIndex)
-				{
-					WorkData.VM->GetByteCode().SetSubject(WorkData.VM->GetByteCode().GetNumInstructions() - 1, Callstack.GetCallPath(), Callstack.GetStack());
-				}
-					
-				WorkData.ExprToSkip.Remove(CompletedExpr);
-				TraverseExpression(CompletedExpr, WorkData);
-				break;
-			}
-			case ERigVMOpCode::ArrayUnion:
-			{
-				TraverseChildren(InExpr, WorkData);
-				const FRigVMOperand& ArrayOperand = WorkData.ExprToOperand.FindChecked(GetSourceVarExpr(InExpr->ChildAt(1)));
-				const FRigVMOperand& OtherOperand = WorkData.ExprToOperand.FindChecked(GetSourceVarExpr(InExpr->ChildAt(2)));
-				WorkData.VM->GetByteCode().AddArrayUnionOp(ArrayOperand, OtherOperand);
-				if (Settings.SetupNodeInstructionIndex)
-				{
-					WorkData.VM->GetByteCode().SetSubject(WorkData.VM->GetByteCode().GetNumInstructions() - 1, Callstack.GetCallPath(), Callstack.GetStack());
-				}
-				break;
-			}
-			case ERigVMOpCode::ArrayDifference:
-			{
-				TraverseChildren(InExpr, WorkData);
-				const FRigVMOperand& ArrayOperand = WorkData.ExprToOperand.FindChecked(GetSourceVarExpr(InExpr->ChildAt(0)));
-				const FRigVMOperand& OtherOperand = WorkData.ExprToOperand.FindChecked(GetSourceVarExpr(InExpr->ChildAt(1)));
-				const FRigVMOperand& ResultOperand = WorkData.ExprToOperand.FindChecked(GetSourceVarExpr(InExpr->ChildAt(2)));
-				WorkData.VM->GetByteCode().AddArrayDifferenceOp(ArrayOperand, OtherOperand, ResultOperand);
-				if (Settings.SetupNodeInstructionIndex)
-				{
-					WorkData.VM->GetByteCode().SetSubject(WorkData.VM->GetByteCode().GetNumInstructions() - 1, Callstack.GetCallPath(), Callstack.GetStack());
-				}
-				break;
-			}
-			case ERigVMOpCode::ArrayIntersection:
-			{
-				TraverseChildren(InExpr, WorkData);
-				const FRigVMOperand& ArrayOperand = WorkData.ExprToOperand.FindChecked(GetSourceVarExpr(InExpr->ChildAt(0)));
-				const FRigVMOperand& OtherOperand = WorkData.ExprToOperand.FindChecked(GetSourceVarExpr(InExpr->ChildAt(1)));
-				const FRigVMOperand& ResultOperand = WorkData.ExprToOperand.FindChecked(GetSourceVarExpr(InExpr->ChildAt(2)));
-				WorkData.VM->GetByteCode().AddArrayIntersectionOp(ArrayOperand, OtherOperand, ResultOperand);
-				if (Settings.SetupNodeInstructionIndex)
-				{
-					WorkData.VM->GetByteCode().SetSubject(WorkData.VM->GetByteCode().GetNumInstructions() - 1, Callstack.GetCallPath(), Callstack.GetStack());
-				}
-				break;
-			}
-			case ERigVMOpCode::ArrayReverse:
-			{
-				TraverseChildren(InExpr, WorkData);
-				const FRigVMOperand& ArrayOperand = WorkData.ExprToOperand.FindChecked(GetSourceVarExpr(InExpr->ChildAt(1)));
-				WorkData.VM->GetByteCode().AddArrayReverseOp(ArrayOperand);
-				if (Settings.SetupNodeInstructionIndex)
-				{
-					WorkData.VM->GetByteCode().SetSubject(WorkData.VM->GetByteCode().GetNumInstructions() - 1, Callstack.GetCallPath(), Callstack.GetStack());
-				}
-				break;
-			}
-			default:
-			{
-				checkNoEntry();
-				break;
-			}
-		}
-	}
-}
-
 void URigVMCompiler::TraverseInvokeEntry(const FRigVMInvokeEntryExprAST* InExpr, FRigVMCompilerWorkData& WorkData)
 {
 	URigVMInvokeEntryNode* InvokeEntryNode = Cast<URigVMInvokeEntryNode>(InExpr->GetNode());
@@ -2766,8 +2450,11 @@ void URigVMCompiler::AddCopyOperator(const FRigVMCopyOp& InOp, const FRigVMAssig
 			{
 				if(InOp.Source.GetRegisterOffset() == INDEX_NONE &&
 					InOp.Target.GetRegisterOffset() == INDEX_NONE)
-				{ 
-					WorkData.VM->GetByteCode().AddArrayCloneOp(InOp.Source, InOp.Target);
+				{
+					static const FString ArrayCloneName =
+						FRigVMRegistry::Get().FindOrAddSingletonDispatchFunction<FRigVMDispatch_ArrayClone>();
+					const int32 FunctionIndex = WorkData.VM->AddRigVMFunction(ArrayCloneName);
+					WorkData.VM->GetByteCode().AddExecuteOp(FunctionIndex, {InOp.Source, InOp.Target});
 					bAddCopyOp = false;
 				}
 			}
diff --git a/Engine/Source/Developer/RigVMDeveloper/Private/RigVMModel/Nodes/RigVMArrayNode.cpp b/Engine/Source/Developer/RigVMDeveloper/Private/RigVMModel/Nodes/RigVMArrayNode.cpp
--- a/Engine/Source/Developer/RigVMDeveloper/Private/RigVMModel/Nodes/RigVMArrayNode.cpp
+++ b/Engine/Source/Developer/RigVMDeveloper/Private/RigVMModel/Nodes/RigVMArrayNode.cpp
@@ -6,817 +6,19 @@
 
 #include UE_INLINE_GENERATED_CPP_BY_NAME(RigVMArrayNode)
 
-const FString URigVMArrayNode::ArrayName = TEXT("Array");
-const FString URigVMArrayNode::NumName = TEXT("Num");
-const FString URigVMArrayNode::IndexName = TEXT("Index");
-const FString URigVMArrayNode::ElementName = TEXT("Element");
-const FString URigVMArrayNode::SuccessName = TEXT("Success");
-const FString URigVMArrayNode::OtherName = TEXT("Other");
-const FString URigVMArrayNode::CloneName = TEXT("Clone");
-const FString URigVMArrayNode::CountName = TEXT("Count");
-const FString URigVMArrayNode::RatioName = TEXT("Ratio");
-const FString URigVMArrayNode::ResultName = TEXT("Result");
-const FString URigVMArrayNode::ContinueName = TEXT("Continue");
-const FString URigVMArrayNode::CompletedName = TEXT("Completed");
-
-URigVMArrayNode::URigVMArrayNode()
+UDEPRECATED_RigVMArrayNode::UDEPRECATED_RigVMArrayNode()
 : OpCode(ERigVMOpCode::Invalid)
 {
 }
 
-FString URigVMArrayNode::GetNodeTitle() const
-{
-	return GetNodeTitle(OpCode);
-}
-
-FString URigVMArrayNode::GetNodeTitle(ERigVMOpCode InOpCode)
-{
-	if(InOpCode == ERigVMOpCode::ArrayGetAtIndex)
-	{
-		static const FString ArrayGetAtIndexNodeTitle = TEXT("At"); 
-		return ArrayGetAtIndexNodeTitle;
-	}
-	else if(InOpCode == ERigVMOpCode::ArraySetAtIndex)
-	{
-		static const FString ArraySetAtIndexNodeTitle = TEXT("Set At"); 
-		return ArraySetAtIndexNodeTitle;
-	}
-	else if(InOpCode == ERigVMOpCode::ArrayGetNum)
-	{
-		static const FString ArrayGetNumNodeTitle = TEXT("Num"); 
-		return ArrayGetNumNodeTitle;
-	}
-	else if(InOpCode == ERigVMOpCode::ArrayIterator)
-	{
-		static const FString ArrayIteratorNodeTitle = TEXT("For Each"); 
-		return ArrayIteratorNodeTitle;
-	}
-
-	FString OpCodeString = StaticEnum<ERigVMOpCode>()->GetDisplayNameTextByValue((int64)InOpCode).ToString();
-	OpCodeString.RemoveFromStart(URigVMArrayNode::ArrayName);
-	return OpCodeString;
-}
-
-FText URigVMArrayNode::GetToolTipText() const
-{
-	switch(OpCode)
-	{
-		case ERigVMOpCode::ArrayReset:
-		{
-			return FText::FromString(TEXT("Removes all elements from an array.\nModifies the input array."));
-		}
-		case ERigVMOpCode::ArrayGetNum:
-		{
-			return FText::FromString(TEXT("Returns the number of elements of an array"));
-		} 
-		case ERigVMOpCode::ArraySetNum:
-		{
-			return FText::FromString(TEXT("Sets the numbers of elements of an array.\nModifies the input array."));
-		}
-		case ERigVMOpCode::ArrayGetAtIndex:
-		{
-			return FText::FromString(TEXT("Returns an element of an array by index."));
-		}  
-		case ERigVMOpCode::ArraySetAtIndex:
-		{
-			return FText::FromString(TEXT("Sets an element of an array by index.\nModifies the input array."));
-		}
-		case ERigVMOpCode::ArrayAdd:
-		{
-			return FText::FromString(TEXT("Adds an element to an array and returns the new element's index.\nModifies the input array."));
-		}
-		case ERigVMOpCode::ArrayInsert:
-		{
-			return FText::FromString(TEXT("Inserts an element into an array at a given index.\nModifies the input array."));
-		}
-		case ERigVMOpCode::ArrayRemove:
-		{
-			return FText::FromString(TEXT("Removes an element from an array by index.\nModifies the input array."));
-		}
-		case ERigVMOpCode::ArrayFind:
-		{
-			return FText::FromString(TEXT("Searchs a potential element in an array and returns its index."));
-		}
-		case ERigVMOpCode::ArrayAppend:
-		{
-			return FText::FromString(TEXT("Appends the another array to the main one.\nModifies the input array."));
-		}
-		case ERigVMOpCode::ArrayClone:
-		{
-			return FText::FromString(TEXT("Clones an array and returns a duplicate."));
-		}
-		case ERigVMOpCode::ArrayIterator:
-		{
-			return FText::FromString(TEXT("Loops over the elements in an array."));
-		}
-		case ERigVMOpCode::ArrayUnion:
-		{
-			return FText::FromString(TEXT("Merges two arrays while ensuring unique elements.\nModifies the input array."));
-		}
-		case ERigVMOpCode::ArrayDifference:
-		{
-			return FText::FromString(TEXT("Creates a new array containing the difference between the two input arrays.\nDifference here means elements that are only contained in either A or B."));
-		}
-		case ERigVMOpCode::ArrayIntersection:
-		{
-			return FText::FromString(TEXT("Creates a new array containing the intersection between the two input arrays.\nDifference here means elements that are contained in both A and B."));
-		}
-		case ERigVMOpCode::ArrayReverse:
-		{
-			return FText::FromString(TEXT("Reverses the order of the elements of an array.\nModifies the input array."));
-		}
-		default:
-		{
-			ensure(false);
-			break;
-		}
-	}
-	return Super::GetToolTipText();
-}
-
-FName URigVMArrayNode::GetNotation() const
-{
-    const FString TemplateName = StaticEnum<ERigVMOpCode>()->GetDisplayNameTextByValue((int64)OpCode).ToString();
-	switch(OpCode)
-    {
-    	case ERigVMOpCode::ArrayReset:
-    	{
-    		static constexpr TCHAR Format[] = TEXT("%s(io %s)");
-    		const FName Notation = *FString::Printf(Format, *TemplateName, *ArrayName);
-            return Notation;
-    	}
-    	case ERigVMOpCode::ArrayGetNum:
-    	{
-    		static constexpr TCHAR Format[] = TEXT("%s(in %s,out %s)");
-    		const FName Notation = *FString::Printf(Format, *TemplateName, *ArrayName, *NumName);
-    		return Notation;
-    	} 
-    	case ERigVMOpCode::ArraySetNum:
-    	{
-    		static constexpr TCHAR Format[] = TEXT("%s(io %s,in %s)");
-    		const FName Notation = *FString::Printf(Format, *TemplateName, *ArrayName, *NumName);
-    		return Notation;
-    	}
-    	case ERigVMOpCode::ArrayGetAtIndex:
-    	{
-    		static constexpr TCHAR Format[] = TEXT("%s(in %s,in %s,out %s)");
-    		const FName Notation = *FString::Printf(Format, *TemplateName, *ArrayName, *IndexName, *ElementName);
-    		return Notation;
-    	}  
-    	case ERigVMOpCode::ArraySetAtIndex:
-    	{
-    		static constexpr TCHAR Format[] = TEXT("%s(io %s,in %s,in %s)");
-    		const FName Notation = *FString::Printf(Format, *TemplateName, *ArrayName, *IndexName, *ElementName);
-    		return Notation;
-    	}
-    	case ERigVMOpCode::ArrayAdd:
-    	{
-    		static constexpr TCHAR Format[] = TEXT("%s(io %s,in %s,out %s)");
-    		const FName Notation = *FString::Printf(Format, *TemplateName, *ArrayName, *ElementName, *IndexName);
-    		return Notation;
-    	}
-    	case ERigVMOpCode::ArrayInsert:
-    	{
-    		static constexpr TCHAR Format[] = TEXT("%s(io %s,in %s,in %s)");
-    		const FName Notation = *FString::Printf(Format, *TemplateName, *ArrayName, *IndexName, *ElementName);
-    		return Notation;
-    	}
-    	case ERigVMOpCode::ArrayRemove:
-    	{
-    		static constexpr TCHAR Format[] = TEXT("%s(io %s,in %s)");
-    		const FName Notation = *FString::Printf(Format, *TemplateName, *ArrayName, *IndexName);
-    		return Notation;
-    	}
-    	case ERigVMOpCode::ArrayFind:
-    	{
-    		static constexpr TCHAR Format[] = TEXT("%s(in %s,in %s,out %s,out %s)");
-    		const FName Notation = *FString::Printf(Format, *TemplateName, *ArrayName, *ElementName, *IndexName, *SuccessName);
-    		return Notation;
-    	}
-    	case ERigVMOpCode::ArrayAppend:
-    	{
-    		static constexpr TCHAR Format[] = TEXT("%s(io %s,in %s)");
-    		const FName Notation = *FString::Printf(Format, *TemplateName, *ArrayName, *OtherName);
-    		return Notation;
-    	}
-    	case ERigVMOpCode::ArrayClone:
-    	{
-    		static constexpr TCHAR Format[] = TEXT("%s(in %s,out %s)");
-    		const FName Notation = *FString::Printf(Format, *TemplateName, *ArrayName, *CloneName);
-    		return Notation;
-    	}
-    	case ERigVMOpCode::ArrayIterator:
-    	{
-    		static constexpr TCHAR Format[] = TEXT("%s(in %s,out %s,out %s,out %s,out %s)");
-    		const FName Notation = *FString::Printf(Format, *TemplateName,
-    			*ArrayName, *ElementName, *IndexName, *CountName, *RatioName);
-    		return Notation;
-    	}
-    	case ERigVMOpCode::ArrayUnion:
-    	{
-    		static constexpr TCHAR Format[] = TEXT("%s(io %s,in %s)");
-    		const FName Notation = *FString::Printf(Format, *TemplateName, *ArrayName, *OtherName);
-    		return Notation;
-    	}
-    	case ERigVMOpCode::ArrayDifference:
-    	{
-    		static constexpr TCHAR Format[] = TEXT("%s(in %s,in %s,out %s)");
-    		const FName Notation = *FString::Printf(Format, *TemplateName, *ArrayName, *OtherName, *ResultName);
-    		return Notation;
-    	}
-    	case ERigVMOpCode::ArrayIntersection:
-    	{
-    		static constexpr TCHAR Format[] = TEXT("%s(in %s,in %s,out %s)");
-    		const FName Notation = *FString::Printf(Format, *TemplateName, *ArrayName, *OtherName, *ResultName);
-    		return Notation;
-    	}
-    	case ERigVMOpCode::ArrayReverse:
-    	{
-    		static constexpr TCHAR Format[] = TEXT("%s(io %s)");
-    		const FName Notation = *FString::Printf(Format, *TemplateName, *ArrayName);
-    		return Notation;
-    	}
-    	default:
-    	{
-    		ensure(false);
-    		break;
-    	}
-    }
-    return Super::GetNotation();
-}
-
-const FRigVMTemplate* URigVMArrayNode::GetTemplate() const
-{
-	if(const FRigVMTemplate* SuperTemplate = Super::GetTemplate())
-	{
-		return SuperTemplate;
-	}
-	
-	if(CachedTemplate == nullptr)
-	{
-		static TArray<const FRigVMTemplate*> TemplatesPerOpCode;
-		static TArray<TArray<FRigVMTemplateArgument>> ArgumentsPerOpCode;
-		if(TemplatesPerOpCode.IsEmpty() || ArgumentsPerOpCode.IsEmpty())
-		{
-			constexpr int32 NumArrayOpCodes = int32(ERigVMOpCode::LastArrayOpCode) - int32(ERigVMOpCode::FirstArrayOpCode) + 1;
-			TemplatesPerOpCode.SetNumZeroed(NumArrayOpCodes);
-			ArgumentsPerOpCode.SetNum(NumArrayOpCodes);
-		}
-
-		const int32 OpCodeIndex = int32(OpCode) - int32(ERigVMOpCode::FirstArrayOpCode);
-
-		if(TemplatesPerOpCode[OpCodeIndex])
-		{
-			CachedTemplate = TemplatesPerOpCode[OpCodeIndex]; 
-			return TemplatesPerOpCode[OpCodeIndex];
-		}
-
-		static const FName ArrayFName = *ArrayName;
-		static const FName NumFName = *NumName;
-		static const FName IndexFName = *IndexName;
-		static const FName ElementFName = *ElementName;
-		static const FName SuccessFName = *SuccessName;
-		static const FName OtherFName = *OtherName;
-		static const FName CloneFName = *CloneName;
-		static const FName CountFName = *CountName;
-		static const FName RatioFName = *RatioName;
-		static const FName ResultFName = *ResultName;
-		static const FName ContinueFName = *ContinueName;
-		
-		TArray<FRigVMTemplateArgument>& Arguments = ArgumentsPerOpCode[OpCodeIndex];
-		
-		if(Arguments.IsEmpty())
-		{
-			Arguments.Reserve(8);
-			
-			static const TArray<FRigVMTemplateArgument::ETypeCategory> SingleCategory = {FRigVMTemplateArgument::ETypeCategory_SingleAnyValue};
-			static const TArray<FRigVMTemplateArgument::ETypeCategory> ArrayCategory = {FRigVMTemplateArgument::ETypeCategory_ArrayAnyValue};
-
-			switch(OpCode)
-			{
-				case ERigVMOpCode::ArrayReset:
-				case ERigVMOpCode::ArrayReverse:
-				{
-					Arguments.Emplace(ArrayFName, ERigVMPinDirection::IO, ArrayCategory);
-					break;
-				}
-				case ERigVMOpCode::ArrayGetNum:
-				{
-					Arguments.Emplace(ArrayFName, ERigVMPinDirection::Input, ArrayCategory);
-					Arguments.Emplace(NumFName, ERigVMPinDirection::Output, RigVMTypeUtils::TypeIndex::Int32);
-					break;
-				} 
-				case ERigVMOpCode::ArraySetNum:
-				{
-					Arguments.Emplace(ArrayFName, ERigVMPinDirection::IO, ArrayCategory);
-					Arguments.Emplace(NumFName, ERigVMPinDirection::Input, RigVMTypeUtils::TypeIndex::Int32);
-					break;
-				}
-				case ERigVMOpCode::ArrayGetAtIndex:
-				{
-					Arguments.Emplace(ArrayFName, ERigVMPinDirection::Input, ArrayCategory);
-					Arguments.Emplace(IndexFName, ERigVMPinDirection::Input, RigVMTypeUtils::TypeIndex::Int32);
-					Arguments.Emplace(ElementFName, ERigVMPinDirection::Output, SingleCategory);
-					break;
-				}  
-				case ERigVMOpCode::ArraySetAtIndex:
-				case ERigVMOpCode::ArrayInsert:
-				{
-					Arguments.Emplace(ArrayFName, ERigVMPinDirection::IO, ArrayCategory);
-					Arguments.Emplace(IndexFName, ERigVMPinDirection::Input, RigVMTypeUtils::TypeIndex::Int32);
-					Arguments.Emplace(ElementFName, ERigVMPinDirection::Input, SingleCategory);
-					break;
-				}
-				case ERigVMOpCode::ArrayAdd:
-				{
-					Arguments.Emplace(ArrayFName, ERigVMPinDirection::IO, ArrayCategory);
-					Arguments.Emplace(ElementFName, ERigVMPinDirection::Input, SingleCategory);
-					Arguments.Emplace(IndexFName, ERigVMPinDirection::Output, RigVMTypeUtils::TypeIndex::Int32);
-					break;
-				}
-				case ERigVMOpCode::ArrayRemove:
-				{
-					Arguments.Emplace(ArrayFName, ERigVMPinDirection::IO, ArrayCategory);
-					Arguments.Emplace(IndexFName, ERigVMPinDirection::Input, RigVMTypeUtils::TypeIndex::Int32);
-					break;
-				}
-				case ERigVMOpCode::ArrayFind:
-				{
-					Arguments.Emplace(ArrayFName, ERigVMPinDirection::Input, ArrayCategory);
-					Arguments.Emplace(ElementFName, ERigVMPinDirection::Input, SingleCategory);
-					Arguments.Emplace(IndexFName, ERigVMPinDirection::Output, RigVMTypeUtils::TypeIndex::Int32);
-					Arguments.Emplace(SuccessFName, ERigVMPinDirection::Output, RigVMTypeUtils::TypeIndex::Bool);
-					break;
-				}
-				case ERigVMOpCode::ArrayAppend:
-				case ERigVMOpCode::ArrayUnion:
-				{
-					Arguments.Emplace(ArrayFName, ERigVMPinDirection::IO, ArrayCategory);
-					Arguments.Emplace(OtherFName, ERigVMPinDirection::Input, ArrayCategory);
-					break;
-				}
-				case ERigVMOpCode::ArrayClone:
-				{
-					Arguments.Emplace(ArrayFName, ERigVMPinDirection::Input, ArrayCategory);
-					Arguments.Emplace(CloneFName, ERigVMPinDirection::Output, ArrayCategory);
-					break;    				
-				}
-				case ERigVMOpCode::ArrayIterator:
-				{
-					Arguments.Emplace(ArrayFName, ERigVMPinDirection::Input, ArrayCategory);
-					Arguments.Emplace(ElementFName, ERigVMPinDirection::Output, SingleCategory);
-					Arguments.Emplace(IndexFName, ERigVMPinDirection::Output, RigVMTypeUtils::TypeIndex::Int32);
-					Arguments.Emplace(CountFName, ERigVMPinDirection::Output, RigVMTypeUtils::TypeIndex::Int32);
-					Arguments.Emplace(RatioFName, ERigVMPinDirection::Output, RigVMTypeUtils::TypeIndex::Float);
-					Arguments.Emplace(ContinueFName, ERigVMPinDirection::Hidden, RigVMTypeUtils::TypeIndex::Bool);
-					break;
-				}
-				case ERigVMOpCode::ArrayDifference:
-				case ERigVMOpCode::ArrayIntersection:
-				{
-					Arguments.Emplace(ArrayFName, ERigVMPinDirection::Input, ArrayCategory);
-					Arguments.Emplace(OtherFName, ERigVMPinDirection::Input, ArrayCategory);
-					Arguments.Emplace(ResultFName, ERigVMPinDirection::Output, ArrayCategory);
-					break;
-				}
-				default:
-				{
-					ensure(false);
-					break;
-				}
-			}
-		}
-
-		const FRigVMTemplateArgumentType ExecuteType(GetGraph()->GetExecuteContextStruct());
-		const int32 ExecuteTypeIndex = FRigVMRegistry::Get().FindOrAddType(ExecuteType);
-		const FString TemplateName = StaticEnum<ERigVMOpCode>()->GetDisplayNameTextByValue((int64)OpCode).ToString();
-
-		FRigVMTemplateDelegates Delegates;
-		Delegates.NewArgumentTypeDelegate = 
-			FRigVMTemplate_NewArgumentTypeDelegate::CreateLambda([OpCodeIndex, ExecuteTypeIndex](const FRigVMTemplate*, const FName& InArgumentName, int32 InTypeIndex)
-			{
-				FRigVMTemplateTypeMap Types;
-				
-				const ERigVMOpCode CurrentOpCode = (ERigVMOpCode)(OpCodeIndex + int32(ERigVMOpCode::FirstArrayOpCode));
-				const FRigVMRegistry& Registry = FRigVMRegistry::Get();
-
-				switch(CurrentOpCode)
-				{
-					case ERigVMOpCode::ArrayReset:
-					case ERigVMOpCode::ArrayReverse:
-					{
-						if(InArgumentName == ArrayFName)
-						{
-							Types.Add(ArrayFName, InTypeIndex);
-						}
-						break;
-					}
-					case ERigVMOpCode::ArrayGetNum:
-					{
-						if(InArgumentName == ArrayFName)
-						{
-							Types.Add(ArrayFName, InTypeIndex);
-							Types.Add(NumFName, RigVMTypeUtils::TypeIndex::Int32);
-						}
-						break;
-					} 
-					case ERigVMOpCode::ArraySetNum:
-					{
-						if(InArgumentName == ArrayFName)
-						{
-							Types.Add(ArrayFName, InTypeIndex);
-							Types.Add(NumFName, RigVMTypeUtils::TypeIndex::Int32);
-						}
-						break;
-					}
-					case ERigVMOpCode::ArrayGetAtIndex:
-					{
-						Types.Add(IndexFName, RigVMTypeUtils::TypeIndex::Int32);
-						if(InArgumentName == ArrayFName)
-						{
-							Types.Add(ArrayFName, InTypeIndex);
-							Types.Add(ElementFName, Registry.GetBaseTypeFromArrayTypeIndex(InTypeIndex));
-						}
-						else if(InArgumentName == ElementFName)
-						{
-							Types.Add(ArrayFName, Registry.GetArrayTypeFromBaseTypeIndex(InTypeIndex));
-							Types.Add(ElementFName, InTypeIndex);
-						}
-						break;
-					}  
-					case ERigVMOpCode::ArraySetAtIndex:
-					case ERigVMOpCode::ArrayInsert:
-					case ERigVMOpCode::ArrayAdd:
-					{
-						Types.Add(IndexFName, RigVMTypeUtils::TypeIndex::Int32);
-						if(InArgumentName == ArrayFName)
-						{
-							Types.Add(ArrayFName, InTypeIndex);
-							Types.Add(ElementFName, Registry.GetBaseTypeFromArrayTypeIndex(InTypeIndex));
-						}
-						else if(InArgumentName == ElementFName)
-						{
-							Types.Add(ArrayFName, Registry.GetArrayTypeFromBaseTypeIndex(InTypeIndex));
-							Types.Add(ElementFName, InTypeIndex);
-						}
-						break;
-					}
-					case ERigVMOpCode::ArrayRemove:
-					{
-						if(InArgumentName == ArrayFName)
-						{
-							Types.Add(ArrayFName, InTypeIndex);
-							Types.Add(IndexFName, RigVMTypeUtils::TypeIndex::Int32);
-						}
-						break;
-					}
-					case ERigVMOpCode::ArrayFind:
-					{
-						Types.Add(IndexFName, RigVMTypeUtils::TypeIndex::Int32);
-						Types.Add(SuccessFName, RigVMTypeUtils::TypeIndex::Bool);
-						if(InArgumentName == ArrayFName)
-						{
-							Types.Add(ArrayFName, InTypeIndex);
-							Types.Add(ElementFName, Registry.GetBaseTypeFromArrayTypeIndex(InTypeIndex));
-						}
-						else if(InArgumentName == ElementFName)
-						{
-							Types.Add(ArrayFName, Registry.GetArrayTypeFromBaseTypeIndex(InTypeIndex));
-							Types.Add(ElementFName, InTypeIndex);
-						}
-						break;
-					}
-					case ERigVMOpCode::ArrayAppend:
-					case ERigVMOpCode::ArrayUnion:
-					{
-						if(InArgumentName == ArrayFName || InArgumentName == OtherFName)
-						{
-							Types.Add(ArrayFName, InTypeIndex);
-							Types.Add(OtherFName, InTypeIndex);
-						}
-						break;
-					}
-					case ERigVMOpCode::ArrayClone:
-					{
-						if(InArgumentName == ArrayFName || InArgumentName == CloneFName)
-						{
-							Types.Add(ArrayFName, InTypeIndex);
-							Types.Add(CloneFName, InTypeIndex);
-						}
-						break;    				
-					}
-					case ERigVMOpCode::ArrayIterator:
-					{
-						Types.Add(ContinueFName, RigVMTypeUtils::TypeIndex::Bool);
-						Types.Add(IndexFName, RigVMTypeUtils::TypeIndex::Int32);
-						Types.Add(CountFName, RigVMTypeUtils::TypeIndex::Int32);
-						Types.Add(RatioFName, RigVMTypeUtils::TypeIndex::Float);
-						if(InArgumentName == ArrayFName)
-						{
-							Types.Add(ArrayFName, InTypeIndex);
-							Types.Add(ElementFName, Registry.GetBaseTypeFromArrayTypeIndex(InTypeIndex));
-						}
-						else if(InArgumentName == ElementFName)
-						{
-							Types.Add(ArrayFName, Registry.GetArrayTypeFromBaseTypeIndex(InTypeIndex));
-							Types.Add(ElementFName, InTypeIndex);
-						}
-						break;
-					}
-					case ERigVMOpCode::ArrayDifference:
-					case ERigVMOpCode::ArrayIntersection:
-					{
-						if(InArgumentName == ArrayFName || InArgumentName == OtherFName || InArgumentName == ResultFName)
-						{
-							Types.Add(ArrayFName, InTypeIndex);
-							Types.Add(OtherFName, InTypeIndex);
-							Types.Add(ResultFName, InTypeIndex);
-						}
-						break;
-					}
-					default:
-					{
-						ensure(false);
-						break;
-					}
-				}
-
-				return Types;
-			});
-
-		TemplatesPerOpCode[OpCodeIndex] = CachedTemplate = FRigVMRegistry::Get().GetOrAddTemplateFromArguments(*TemplateName, Arguments, Delegates);
-	}
-	return CachedTemplate;
-}
-
-const TArray<FName>& URigVMArrayNode::GetControlFlowBlocks() const
-{
-	if(GetOpCode() == ERigVMOpCode::ArrayIterator)
-	{
-		static const TArray<FName> LoopBlocks = {FRigVMStruct::ExecuteContextName, FRigVMStruct::ForLoopCompletedPinName};
-		return LoopBlocks;
-	}
-	return Super::GetControlFlowBlocks();
-}
-
-FText URigVMArrayNode::GetToolTipTextForPin(const URigVMPin* InPin) const
-{
-	switch(OpCode)
-	{
-		case ERigVMOpCode::ArrayReset:
-		{
-			if(InPin->GetName() == URigVMArrayNode::ArrayName)
-			{
-				return FText::FromString(TEXT("The array to be cleared."));
-			}
-			break;
-		}
-		case ERigVMOpCode::ArrayGetNum:
-		{
-			if(InPin->GetName() == URigVMArrayNode::ArrayName)
-			{
-				return FText::FromString(TEXT("The array to retrieve the size for."));
-			}
-			else if(InPin->GetName() == URigVMArrayNode::NumName)
-			{
-				return FText::FromString(TEXT("The size of the input array."));
-			}
-			break;
-		} 
-		case ERigVMOpCode::ArraySetNum:
-		{
-			if(InPin->GetName() == URigVMArrayNode::ArrayName)
-			{
-				return FText::FromString(TEXT("The array to set the size for."));
-			}
-			else if(InPin->GetName() == URigVMArrayNode::NumName)
-			{
-				return FText::FromString(TEXT("The new size of the array."));
-			}
-			break;
-		}
-		case ERigVMOpCode::ArrayGetAtIndex:
-		{
-			if(InPin->GetName() == URigVMArrayNode::ArrayName)
-			{
-				return FText::FromString(TEXT("The array to retrieve an element from."));
-			}
-			else if(InPin->GetName() == URigVMArrayNode::IndexName)
-			{
-				return FText::FromString(TEXT("The index of the element to retrieve."));
-			}
-			else if(InPin->GetName() == URigVMArrayNode::ElementName)
-			{
-				return FText::FromString(TEXT("The element at the given index."));
-			}
-			break;
-		}  
-		case ERigVMOpCode::ArraySetAtIndex:
-		{
-			if(InPin->GetName() == URigVMArrayNode::ArrayName)
-			{
-				return FText::FromString(TEXT("The array to set an element for."));
-			}
-			else if(InPin->GetName() == URigVMArrayNode::IndexName)
-			{
-				return FText::FromString(TEXT("The index of the element to set."));
-			}
-			else if(InPin->GetName() == URigVMArrayNode::ElementName)
-			{
-				return FText::FromString(TEXT("The new value for element to set."));
-			}
-			break;
-		}
-		case ERigVMOpCode::ArrayAdd:
-		{
-			if(InPin->GetName() == URigVMArrayNode::ArrayName)
-			{
-				return FText::FromString(TEXT("The array to add an element to."));
-			}
-			else if(InPin->GetName() == URigVMArrayNode::ElementName)
-			{
-				return FText::FromString(TEXT("The element to add to the array."));
-			}
-			else if(InPin->GetName() == URigVMArrayNode::IndexName)
-			{
-				return FText::FromString(TEXT("The index of the newly added element."));
-			}
-			break;
-		}
-		case ERigVMOpCode::ArrayInsert:
-		{
-			if(InPin->GetName() == URigVMArrayNode::ArrayName)
-			{
-				return FText::FromString(TEXT("The array to insert an element into."));
-			}
-			else if(InPin->GetName() == URigVMArrayNode::IndexName)
-			{
-				return FText::FromString(TEXT("The index at which to insert the element."));
-			}
-			else if(InPin->GetName() == URigVMArrayNode::ElementName)
-			{
-				return FText::FromString(TEXT("The element to insert into the array."));
-			}
-			break;
-		}
-		case ERigVMOpCode::ArrayRemove:
-		{
-			if(InPin->GetName() == URigVMArrayNode::ArrayName)
-			{
-				return FText::FromString(TEXT("The array to remove an element from."));
-			}
-			else if(InPin->GetName() == URigVMArrayNode::IndexName)
-			{
-				return FText::FromString(TEXT("The index at which to remove the element."));
-			}
-			break;
-		}
-		case ERigVMOpCode::ArrayFind:
-		{
-			if(InPin->GetName() == URigVMArrayNode::ArrayName)
-			{
-				return FText::FromString(TEXT("The array to search within."));
-			}
-			else if(InPin->GetName() == URigVMArrayNode::ElementName)
-			{
-				return FText::FromString(TEXT("The element to look for."));
-			}
-			else if(InPin->GetName() == URigVMArrayNode::IndexName)
-			{
-				return FText::FromString(TEXT("The index of the found element (or -1)."));
-			}
-			else if(InPin->GetName() == URigVMArrayNode::SuccessName)
-			{
-				return FText::FromString(TEXT("True if the element has been found."));
-			}
-			break;
-		}
-		case ERigVMOpCode::ArrayAppend:
-		{
-			if(InPin->GetName() == URigVMArrayNode::ArrayName)
-			{
-				return FText::FromString(TEXT("The array to append the other array to."));
-			}
-			else if(InPin->GetName() == URigVMArrayNode::OtherName)
-			{
-				return FText::FromString(TEXT("The second array to append to the first one."));
-			}
-			break;
-		}
-		case ERigVMOpCode::ArrayClone:
-		{
-			if(InPin->GetName() == URigVMArrayNode::ArrayName)
-			{
-				return FText::FromString(TEXT("The array to clone"));
-			}
-			else if(InPin->GetName() == URigVMArrayNode::CloneName)
-			{
-				return FText::FromString(TEXT("The duplicate of the input array."));
-			}
-			break;
-		}
-		case ERigVMOpCode::ArrayIterator:
-		{
-			if(InPin->GetName() == URigVMArrayNode::ArrayName)
-			{
-				return FText::FromString(TEXT("The array to iterate over"));
-			}
-			else if(InPin->GetName() == URigVMArrayNode::IndexName)
-			{
-				return FText::FromString(TEXT("The index of the current element."));
-			}
-			else if(InPin->GetName() == URigVMArrayNode::CountName)
-			{
-				return FText::FromString(TEXT("The count of all elements in the array."));
-			}
-			else if(InPin->GetName() == URigVMArrayNode::ElementName)
-			{
-				return FText::FromString(TEXT("The current element."));
-			}
-			else if(InPin->GetName() == URigVMArrayNode::RatioName)
-			{
-				return FText::FromString(TEXT("A float ratio from 0.0 (first element) to 1.0 (last element)."));
-			}
-			else if(InPin->GetName() == URigVMArrayNode::CompletedName)
-			{
-				return FText::FromString(TEXT("The execute block to run once the loop has completed."));
-			}
-			break;
-		}
-		case ERigVMOpCode::ArrayUnion:
-		{
-			if(InPin->GetName() == URigVMArrayNode::ArrayName)
-			{
-				return FText::FromString(TEXT("The array to merge the other array with."));
-			}
-			else if(InPin->GetName() == URigVMArrayNode::OtherName)
-			{
-				return FText::FromString(TEXT("The second array to merge to the first one."));
-			}
-			break;
-		}
-		case ERigVMOpCode::ArrayDifference:
-		{
-			if(InPin->GetName() == URigVMArrayNode::ArrayName)
-			{
-				return FText::FromString(TEXT("The first array to compare the other array with."));
-			}
-			else if(InPin->GetName() == URigVMArrayNode::OtherName)
-			{
-				return FText::FromString(TEXT("The second array to compare the other array with."));
-			}
-			else if(InPin->GetName() == URigVMArrayNode::ResultName)
-			{
-				return FText::FromString(TEXT("The resulting difference array."));
-			}
-			break;
-		}
-		case ERigVMOpCode::ArrayIntersection:
-		{
-			if(InPin->GetName() == URigVMArrayNode::ArrayName)
-			{
-				return FText::FromString(TEXT("The first array to compare the other array with."));
-			}
-			else if(InPin->GetName() == URigVMArrayNode::OtherName)
-			{
-				return FText::FromString(TEXT("The second array to compare the other array with."));
-			}
-			else if(InPin->GetName() == URigVMArrayNode::ResultName)
-			{
-				return FText::FromString(TEXT("The resulting intersecting array."));
-			}
-			break;
-		}
-		case ERigVMOpCode::ArrayReverse:
-		{
-			if(InPin->GetName() == URigVMArrayNode::ArrayName)
-			{
-				return FText::FromString(TEXT("The array reverse."));
-			}
-			break;
-		}
-		default:
-		{
-			ensure(false);
-			break;
-		}
-	}
-	return Super::GetToolTipTextForPin(InPin);
-}
-
-FLinearColor URigVMArrayNode::GetNodeColor() const
-{
-	return FLinearColor::White;
-}
-
-ERigVMOpCode URigVMArrayNode::GetOpCode() const
+ERigVMOpCode UDEPRECATED_RigVMArrayNode::GetOpCode() const
 {
 	return OpCode;
 }
 
-FString URigVMArrayNode::GetCPPType() const
+FString UDEPRECATED_RigVMArrayNode::GetCPPType() const
 {
-	URigVMPin* ArrayPin = FindPin(ArrayName);
+	const URigVMPin* ArrayPin = FindPin(TEXT("Array"));
 	if (ArrayPin == nullptr)
 	{
 		return FString();
@@ -824,13 +26,12 @@ FString URigVMArrayNode::GetCPPType() const
 	return RigVMTypeUtils::BaseTypeFromArrayType(ArrayPin->GetCPPType());
 }
 
-UObject* URigVMArrayNode::GetCPPTypeObject() const
+UObject* UDEPRECATED_RigVMArrayNode::GetCPPTypeObject() const
 {
-	URigVMPin* ArrayPin = FindPin(ArrayName);
+	const URigVMPin* ArrayPin = FindPin(TEXT("Array"));
 	if (ArrayPin == nullptr)
 	{
 		return nullptr;
 	}
 	return ArrayPin->GetCPPTypeObject();
 }
-
diff --git a/Engine/Source/Developer/RigVMDeveloper/Private/RigVMModel/RigVMClient.cpp b/Engine/Source/Developer/RigVMDeveloper/Private/RigVMModel/RigVMClient.cpp
--- a/Engine/Source/Developer/RigVMDeveloper/Private/RigVMModel/RigVMClient.cpp
+++ b/Engine/Source/Developer/RigVMDeveloper/Private/RigVMModel/RigVMClient.cpp
@@ -671,6 +671,7 @@ void FRigVMClient::PatchModelsOnLoad()
 		Controller->PatchDispatchNodesOnLoad();
 		Controller->PatchBranchNodesOnLoad();
 		Controller->PatchIfSelectNodesOnLoad();
+		Controller->PatchArrayNodesOnLoad();
 	}
 }
 
diff --git a/Engine/Source/Developer/RigVMDeveloper/Private/RigVMModel/RigVMController.cpp b/Engine/Source/Developer/RigVMDeveloper/Private/RigVMModel/RigVMController.cpp
--- a/Engine/Source/Developer/RigVMDeveloper/Private/RigVMModel/RigVMController.cpp
+++ b/Engine/Source/Developer/RigVMDeveloper/Private/RigVMModel/RigVMController.cpp
@@ -26,8 +26,10 @@
 #include "Engine/UserDefinedStruct.h"
 #include "RigVMFunctions/RigVMDispatch_If.h"
 #include "RigVMFunctions/RigVMDispatch_Select.h"
+#include "RigVMFunctions/RigVMDispatch_Array.h"
 #include "RigVMModel/RigVMClient.h"
 #include "RigVMModel/Nodes/RigVMBranchNode.h"
+#include "RigVMModel/Nodes/RigVMArrayNode.h"
 
 #include UE_INLINE_GENERATED_CPP_BY_NAME(RigVMController)
 
@@ -616,31 +618,6 @@ TArray<FString> URigVMController::GetAddNodePythonCommands(URigVMNode* Node) con
 				*RigVMPythonUtils::Vector2DToPythonString(RerouteNode->GetPosition()),
 				*NodeName));
 	}
-	else if (const URigVMArrayNode* ArrayNode = Cast<URigVMArrayNode>(Node))
-	{
-		// add_array_node(opcode, cpp_type, cpp_type_object, position=[0.0, 0.0], node_name='', undo=True)
-		if (ArrayNode->GetCPPTypeObject())
-		{
-			static constexpr TCHAR ArrayNodeFormat[] = TEXT("blueprint.get_controller_by_name('%s').add_array_node_from_object_path(%s, '%s', '%s', %s, '%s')");
-			Commands.Add(FString::Printf(ArrayNodeFormat,
-					*GraphName,
-					*RigVMPythonUtils::EnumValueToPythonString<ERigVMOpCode>((int64)ArrayNode->GetOpCode()),
-					*ArrayNode->GetCPPType(),
-					*ArrayNode->GetCPPTypeObject()->GetPathName(),
-					*RigVMPythonUtils::Vector2DToPythonString(ArrayNode->GetPosition()),
-					*NodeName));	
-		}
-		else
-		{
-			static constexpr TCHAR ArrayNodeFormat[] = TEXT("blueprint.get_controller_by_name('%s').add_array_node(%s, '%s', None, %s, '%s')");
-			Commands.Add(FString::Printf(ArrayNodeFormat,
-					*GraphName,
-					*RigVMPythonUtils::EnumValueToPythonString<ERigVMOpCode>((int64)ArrayNode->GetOpCode()),
-					*ArrayNode->GetCPPType(),
-					*RigVMPythonUtils::Vector2DToPythonString(ArrayNode->GetPosition()),
-					*NodeName));	
-		}
-	}
 	else if (const URigVMEnumNode* EnumNode = Cast<URigVMEnumNode>(Node))
 	{
 		// add_enum_node(cpp_type_object_path, position=[0.0, 0.0], node_name='', undo=True)
@@ -1898,6 +1875,27 @@ bool URigVMController::UnresolveTemplateNodes(const TArray<URigVMNode*>& InNodes
 					URigVMNode* OtherNode = SourcePin->GetNode() == Node ? TargetPin->GetNode() : SourcePin->GetNode();
 					if (!InNodes.Contains(OtherNode))
 					{
+						const URigVMPin* PinOnNode = SourcePin->GetNode() == Node ? SourcePin : TargetPin;
+						if(PinOnNode->IsExecuteContext())
+						{
+							continue;
+						}
+						
+						if (const URigVMTemplateNode* TemplateNode = Cast<URigVMTemplateNode>(Node))
+						{
+							if(const FRigVMTemplate* Template = TemplateNode->GetTemplate())
+							{
+								const URigVMPin* RootPin = PinOnNode->GetRootPin();
+								if(const FRigVMTemplateArgument* Argument = Template->FindArgument(RootPin->GetFName()))
+								{
+									if(Argument->IsSingleton())
+									{
+										continue;
+									}
+								}
+							}
+						}
+						
 						BreakLink(SourcePin, TargetPin, bSetupUndoRedo);
 					}			
 				}
@@ -9329,6 +9327,12 @@ bool URigVMController::AddLink(URigVMPin* OutputPin, URigVMPin* InputPin, bool b
 
 							if(SourcePin && TargetPin)
 							{
+								if(SourcePin->IsExecuteContext())
+								{
+									NumLinksRemaining--;
+									continue;
+								}
+								
 								const URigVMPin* PinOnTemplateNode = SourcePin->GetNode() == TemplateNode ? SourcePin : TargetPin;
 								if(PinOnTemplateNode->IsRootPin() && PinOnTemplateNode != InPin)
 								{
@@ -13218,7 +13222,7 @@ URigVMEnumNode* URigVMController::AddEnumNode(const FName& InCPPTypeObjectPath,
 	return Node;
 }
 
-URigVMArrayNode* URigVMController::AddArrayNode(ERigVMOpCode InOpCode, const FString& InCPPType,
+URigVMNode* URigVMController::AddArrayNode(ERigVMOpCode InOpCode, const FString& InCPPType,
 	UObject* InCPPTypeObject, const FVector2D& InPosition, const FString& InNodeName, bool bSetupUndoRedo,
 	bool bPrintPythonCommand)
 {
@@ -13232,311 +13236,84 @@ URigVMArrayNode* URigVMController::AddArrayNode(ERigVMOpCode InOpCode, const FSt
 		return nullptr;
 	}
 
-	// validate the op code
-	bool bIsMutable = false;
-	switch(InOpCode)
-	{
-		case ERigVMOpCode::ArrayReset:
-		case ERigVMOpCode::ArrayGetNum: 
-		case ERigVMOpCode::ArraySetNum:
-		case ERigVMOpCode::ArrayGetAtIndex:  
-		case ERigVMOpCode::ArraySetAtIndex:
-		case ERigVMOpCode::ArrayAdd:
-		case ERigVMOpCode::ArrayInsert:
-		case ERigVMOpCode::ArrayRemove:
-		case ERigVMOpCode::ArrayFind:
-		case ERigVMOpCode::ArrayAppend:
-		case ERigVMOpCode::ArrayClone:
-		case ERigVMOpCode::ArrayIterator:
-		case ERigVMOpCode::ArrayUnion:
-		case ERigVMOpCode::ArrayDifference:
-		case ERigVMOpCode::ArrayIntersection:
-		case ERigVMOpCode::ArrayReverse:
-		{
-			break;
-		}
-		default:
-		{
-			ReportErrorf(TEXT("OpCode '%s' is not valid for Array Node."), *StaticEnum<ERigVMOpCode>()->GetNameStringByValue((int64)InOpCode));
-			return nullptr;
-		}
-	}
-
-	URigVMGraph* Graph = GetGraph();
-	check(Graph);
-
-	if (Graph->IsA<URigVMFunctionLibrary>())
+	const FString CPPType = RigVMTypeUtils::IsArrayType(InCPPType) ? RigVMTypeUtils::BaseTypeFromArrayType(InCPPType) : InCPPType;
+	const TRigVMTypeIndex& ElementTypeIndex = FRigVMRegistry::Get().FindOrAddType({*CPPType, InCPPTypeObject});
+	if(ElementTypeIndex == INDEX_NONE)
 	{
-		ReportError(TEXT("Cannot add array nodes to function library graphs."));
 		return nullptr;
 	}
+	const TRigVMTypeIndex& ArrayTypeIndex = FRigVMRegistry::Get().GetArrayTypeFromBaseTypeIndex(ElementTypeIndex);
 
-	FString CPPType = InCPPType;
-	if(RigVMTypeUtils::IsArrayType(CPPType))
+	const FName FactoryName = FRigVMDispatch_ArrayBase::GetFactoryNameForOpCode(InOpCode);
+	if(FactoryName.IsNone())
 	{
-		CPPType = RigVMTypeUtils::BaseTypeFromArrayType(CPPType);
+		ReportErrorf(TEXT("OpCode '%s' is not valid for Array Node."), *StaticEnum<ERigVMOpCode>()->GetNameStringByValue((int64)InOpCode));
+		return nullptr;
 	}
 
-	if (InCPPTypeObject == nullptr)
+	const FRigVMDispatchFactory* Factory = FRigVMRegistry::Get().FindDispatchFactory(FactoryName);
+	if(Factory == nullptr)
 	{
-		InCPPTypeObject = URigVMCompiler::GetScriptStructForCPPType(CPPType);
+		ReportErrorf(TEXT("Cannot find array dispatch '%s'."), *FactoryName.ToString());
+		return nullptr;
 	}
-	if (InCPPTypeObject == nullptr)
+
+	if(bSetupUndoRedo)
 	{
-		InCPPTypeObject = URigVMPin::FindObjectFromCPPTypeObjectPath<UObject>(CPPType);
+		OpenUndoBracket(TEXT("Add Array Node"));
 	}
-
-	CPPType = RigVMTypeUtils::PostProcessCPPType(CPPType, InCPPTypeObject);
 	
-	const FString Name = GetValidNodeName(InNodeName.IsEmpty() ? FString(TEXT("ArrayNode")) : InNodeName);
-	URigVMArrayNode* Node = NewObject<URigVMArrayNode>(Graph, *Name);
-	Node->Position = InPosition;
-	Node->OpCode = InOpCode;
+	const FRigVMTemplate* Template = Factory->GetTemplate();
+	URigVMTemplateNode* Node = AddTemplateNode(
+		Template->GetNotation(),
+		InPosition, 
+		InNodeName, 
+		bSetupUndoRedo, 
+		bPrintPythonCommand);
 
-	struct Local
+	if(!FRigVMRegistry::Get().IsWildCardType(ElementTypeIndex))
 	{
-		static URigVMPin* AddPin(URigVMController* InController, URigVMNode* InNode, const FName& InName, ERigVMPinDirection InDirection, bool bIsArray, const FString& InCPPType, UObject* InCPPTypeObject)
+		FName ArgumentNameToResolve = NAME_None;
+		TRigVMTypeIndex TypeIndex = INDEX_NONE;
+		for(int32 Index = 0; Index < Template->NumArguments(); Index++)
 		{
-			URigVMPin* Pin = NewObject<URigVMPin>(InNode, InName);
-			Pin->CPPType = InCPPType;
-			Pin->CPPTypeObject = InCPPTypeObject;
-			if(Pin->CPPTypeObject)
+			const FRigVMTemplateArgument* Argument = Template->GetArgument(Index);
+			if(Argument->IsSingleton())
 			{
-				Pin->CPPTypeObjectPath = *Pin->CPPTypeObject->GetPathName();
+				continue;
 			}
-			if(bIsArray && !RigVMTypeUtils::IsArrayType(Pin->CPPType))
+			if(Argument->GetArrayType() == FRigVMTemplateArgument::EArrayType_SingleValue)
 			{
-				Pin->CPPType = RigVMTypeUtils::ArrayTypeFromBaseType(*Pin->CPPType);
-			}
-			Pin->Direction = InDirection;
-			Pin->bIsDynamicArray = bIsArray;
-			AddNodePin(InNode, Pin);
-
-			if(Pin->Direction != ERigVMPinDirection::Hidden && !bIsArray && !Pin->IsExecuteContext())
-			{
-				if(UScriptStruct* Struct = Cast<UScriptStruct>(Pin->CPPTypeObject))
-				{
-					FString DefaultValue;
-					InController->CreateDefaultValueForStructIfRequired(Pin->GetScriptStruct(), DefaultValue);
-					{
-						TGuardValue<bool> SuspendNotifications(InController->bSuspendNotifications, true);
-						InController->AddPinsForStruct(Struct, InNode, Pin, InDirection, DefaultValue, true);
-					}
-				}
+				ArgumentNameToResolve = Argument->GetName();
+				TypeIndex = ElementTypeIndex;
+				break;
 			}
-
-			return Pin;
-		}
-		
-		static URigVMPin* AddExecutePin(URigVMController* InController, URigVMNode* InNode, ERigVMPinDirection InDirection = ERigVMPinDirection::IO, const FName& InName = NAME_None)
-		{
-			const FName PinName = InName.IsNone() ? FRigVMStruct::ExecuteContextName : InName;
-			UScriptStruct* ExecuteContextStruct = FRigVMExecuteContext::StaticStruct();
-			URigVMPin* Pin = AddPin(InController, InNode, PinName, InDirection, false, FString::Printf(TEXT("F%s"), *ExecuteContextStruct->GetName()), ExecuteContextStruct);
-			if(PinName == FRigVMStruct::ExecuteContextName)
+			if(Argument->GetArrayType() == FRigVMTemplateArgument::EArrayType_ArrayValue)
 			{
-				Pin->DisplayName = FRigVMStruct::ExecuteName;
+				ArgumentNameToResolve = Argument->GetName();
+				TypeIndex = ArrayTypeIndex;
+				break;
 			}
-			return Pin;
-		}
-
-		static URigVMPin* AddArrayPin(URigVMController* InController, URigVMNode* InNode, ERigVMPinDirection InDirection, const FString& InCPPType, UObject* InCPPTypeObject, const FName& InName = NAME_None)
-		{
-			const FName PinName = InName.IsNone() ? *URigVMArrayNode::ArrayName : InName;
-			return AddPin(InController, InNode, PinName, InDirection, true, InCPPType, InCPPTypeObject);  
-		}
-
-		static URigVMPin* AddElementPin(URigVMController* InController, URigVMNode* InNode, ERigVMPinDirection InDirection, const FString& InCPPType, UObject* InCPPTypeObject)
-		{
-			return AddPin(InController, InNode, *URigVMArrayNode::ElementName, InDirection, false, InCPPType, InCPPTypeObject);  
-		}
-
-		static URigVMPin* AddIndexPin(URigVMController* InController, URigVMNode* InNode, ERigVMPinDirection InDirection)
-		{
-			return AddPin(InController, InNode, *URigVMArrayNode::IndexName, InDirection, false, RigVMTypeUtils::Int32Type, nullptr);  
-		}
-
-		static URigVMPin* AddNumPin(URigVMController* InController, URigVMNode* InNode, ERigVMPinDirection InDirection)
-		{
-			return AddPin(InController, InNode, *URigVMArrayNode::NumName, InDirection, false, RigVMTypeUtils::Int32Type, nullptr);  
-		}
-
-		static URigVMPin* AddCountPin(URigVMController* InController, URigVMNode* InNode, ERigVMPinDirection InDirection)
-		{
-			return AddPin(InController, InNode, *URigVMArrayNode::CountName, InDirection, false, RigVMTypeUtils::Int32Type, nullptr);  
-		}
-
-		static URigVMPin* AddRatioPin(URigVMController* InController, URigVMNode* InNode)
-		{
-			return AddPin(InController, InNode, *URigVMArrayNode::RatioName, ERigVMPinDirection::Output, false, RigVMTypeUtils::FloatType, nullptr);  
-		}
-
-		static URigVMPin* AddContinuePin(URigVMController* InController, URigVMNode* InNode)
-		{
-			return AddPin(InController, InNode, *URigVMArrayNode::ContinueName, ERigVMPinDirection::Hidden, false, RigVMTypeUtils::BoolType, nullptr);
-		}
-
-		static URigVMPin* AddSuccessPin(URigVMController* InController, URigVMNode* InNode)
-		{
-			return AddPin(InController, InNode, *URigVMArrayNode::SuccessName, ERigVMPinDirection::Output, false, RigVMTypeUtils::BoolType, nullptr);  
 		}
-	};
 
-	switch(InOpCode)
-	{
-		case ERigVMOpCode::ArrayReset:
-		case ERigVMOpCode::ArrayReverse:
-		{
-			Local::AddExecutePin(this, Node);
-			Local::AddArrayPin(this, Node, ERigVMPinDirection::IO, CPPType, InCPPTypeObject);
-			break;
-		}
-		case ERigVMOpCode::ArrayGetNum:
-		{
-			Local::AddArrayPin(this, Node, ERigVMPinDirection::Input, CPPType, InCPPTypeObject);
-			Local::AddNumPin(this, Node, ERigVMPinDirection::Output);
-			break;
-		} 
-		case ERigVMOpCode::ArraySetNum:
-		{
-			Local::AddExecutePin(this, Node);
-			Local::AddArrayPin(this, Node, ERigVMPinDirection::IO, CPPType, InCPPTypeObject);
-			Local::AddNumPin(this, Node, ERigVMPinDirection::Input);
-			break;
-		}
-		case ERigVMOpCode::ArrayGetAtIndex:
-		{
-			Local::AddArrayPin(this, Node, ERigVMPinDirection::Input, CPPType, InCPPTypeObject);
-			Local::AddIndexPin(this, Node, ERigVMPinDirection::Input);
-			Local::AddElementPin(this, Node, ERigVMPinDirection::Output, CPPType, InCPPTypeObject);
-			break;
-		}
-		case ERigVMOpCode::ArraySetAtIndex:
-		case ERigVMOpCode::ArrayInsert:
-		{
-			Local::AddExecutePin(this, Node);
-			Local::AddArrayPin(this, Node, ERigVMPinDirection::IO, CPPType, InCPPTypeObject);
-			Local::AddIndexPin(this, Node, ERigVMPinDirection::Input);
-			Local::AddElementPin(this, Node, ERigVMPinDirection::Input, CPPType, InCPPTypeObject);
-			break;
-		}
-		case ERigVMOpCode::ArrayAdd:
-		{
-			Local::AddExecutePin(this, Node);
-			Local::AddArrayPin(this, Node, ERigVMPinDirection::IO, CPPType, InCPPTypeObject);
-			Local::AddElementPin(this, Node, ERigVMPinDirection::Input, CPPType, InCPPTypeObject);
-			Local::AddIndexPin(this, Node, ERigVMPinDirection::Output);
-			break;
-		}
-		case ERigVMOpCode::ArrayFind:
+		if(!ArgumentNameToResolve.IsNone() && TypeIndex != INDEX_NONE)
 		{
-			Local::AddArrayPin(this, Node, ERigVMPinDirection::Input, CPPType, InCPPTypeObject);
-			Local::AddElementPin(this, Node, ERigVMPinDirection::Input, CPPType, InCPPTypeObject);
-			Local::AddIndexPin(this, Node, ERigVMPinDirection::Output);
-			Local::AddSuccessPin(this, Node);
-			break;
-		}
-		case ERigVMOpCode::ArrayRemove:
-		{
-			Local::AddExecutePin(this, Node);
-			Local::AddArrayPin(this, Node, ERigVMPinDirection::IO, CPPType, InCPPTypeObject);
-			Local::AddIndexPin(this, Node, ERigVMPinDirection::Input);
-			break;
-		}
-		case ERigVMOpCode::ArrayAppend:
-		case ERigVMOpCode::ArrayUnion:
-		{
-			Local::AddExecutePin(this, Node);
-			Local::AddArrayPin(this, Node, ERigVMPinDirection::IO, CPPType, InCPPTypeObject);
-			Local::AddArrayPin(this, Node, ERigVMPinDirection::Input, CPPType, InCPPTypeObject, *URigVMArrayNode::OtherName);
-			break;
-		}
-		case ERigVMOpCode::ArrayClone:
-		{
-			Local::AddArrayPin(this, Node, ERigVMPinDirection::Input, CPPType, InCPPTypeObject);
-			Local::AddArrayPin(this, Node, ERigVMPinDirection::Output, CPPType, InCPPTypeObject, *URigVMArrayNode::CloneName);
-			break;
-		}
-		case ERigVMOpCode::ArrayIterator:
-		{
-			Local::AddExecutePin(this, Node);
-			Local::AddArrayPin(this, Node, ERigVMPinDirection::Input, CPPType, InCPPTypeObject);
-			Local::AddElementPin(this, Node, ERigVMPinDirection::Output, CPPType, InCPPTypeObject);
-			Local::AddIndexPin(this, Node, ERigVMPinDirection::Output);
-			Local::AddCountPin(this, Node, ERigVMPinDirection::Output);
-			Local::AddRatioPin(this, Node);
-			Local::AddContinuePin(this, Node);
-			Local::AddExecutePin(this, Node, ERigVMPinDirection::Output, *URigVMArrayNode::CompletedName);
-			break;
-		}
-		case ERigVMOpCode::ArrayDifference:
-		case ERigVMOpCode::ArrayIntersection:
-		{
-			Local::AddArrayPin(this, Node, ERigVMPinDirection::Input, CPPType, InCPPTypeObject);
-			Local::AddArrayPin(this, Node, ERigVMPinDirection::Input, CPPType, InCPPTypeObject, *URigVMArrayNode::OtherName);
-			Local::AddArrayPin(this, Node, ERigVMPinDirection::Output, CPPType, InCPPTypeObject, *URigVMArrayNode::ResultName);
-			break;
-		}
-		default:
-		{
-			checkNoEntry();
-		}
-	}
-
-	Graph->Nodes.Add(Node);
-
-	if (!bSuspendNotifications)
-	{
-		Graph->MarkPackageDirty();
-	}
-
-	FRigVMControllerCompileBracketScope CompileScope(this);
-	FRigVMAddArrayNodeAction Action;
-	if (bSetupUndoRedo)
-	{
-		Action = FRigVMAddArrayNodeAction(Node);
-		Action.Title = FString::Printf(TEXT("Add %s Array Node"), *Node->GetNodeTitle());
-		ActionStack->BeginAction(Action);
-	}
-
-	Notify(ERigVMGraphNotifType::NodeAdded, Node);
-
-	if (bSetupUndoRedo)
-	{
-		ActionStack->EndAction(Action);
-	}
-
-	Node->InitializeFilteredPermutations();
-	if (InCPPType != RigVMTypeUtils::GetWildCardCPPType() && InCPPType != RigVMTypeUtils::GetWildCardArrayCPPType())
-	{
-		URigVMPin* ArrayPin = Node->FindPin(URigVMArrayNode::ArrayName);
-		PrepareTemplatePinForType(ArrayPin, {ArrayPin->GetTypeIndex()}, bSetupUndoRedo);
-		TArray<int32> FilterPermutations = Node->GetFilteredPermutationsIndices();
-		if (FilterPermutations.Num() == 1)
-		{
-			const TArray<FRigVMTemplatePreferredType> NewPreferredPermutationTypes = Node->GetPreferredTypesForPermutation(FilterPermutations[0]);
-			if (bSetupUndoRedo)
+			if(const URigVMPin* Pin = Node->FindPin(ArgumentNameToResolve.ToString()))
 			{
-				ActionStack->AddAction(FRigVMSetPreferredTemplatePermutationsAction(Node, NewPreferredPermutationTypes));
+				ResolveWildCardPin(Pin->GetPinPath(), TypeIndex, bSetupUndoRedo, bPrintPythonCommand);
 			}
-			Node->PreferredPermutationPairs = NewPreferredPermutationTypes;
 		}
 	}
 
-	if (bPrintPythonCommand)
+	if(bSetupUndoRedo)
 	{
-		TArray<FString> Commands = GetAddNodePythonCommands(Node);
-		for (const FString& Command : Commands)
-		{
-			RigVMPythonUtils::Print(GetGraphOuterName(), 
-				FString::Printf(TEXT("%s"), *Command));
-		}
+		CloseUndoBracket();
 	}
 
 	return Node;
 }
 
-URigVMArrayNode* URigVMController::AddArrayNodeFromObjectPath(ERigVMOpCode InOpCode, const FString& InCPPType,
+URigVMNode* URigVMController::AddArrayNodeFromObjectPath(ERigVMOpCode InOpCode, const FString& InCPPType,
 	const FString& InCPPTypeObjectPath, const FVector2D& InPosition, const FString& InNodeName, bool bSetupUndoRedo,
 	bool bPrintPythonCommand)
 {
@@ -14264,40 +14041,51 @@ void URigVMController::AddPinsForTemplate(const FRigVMTemplate* InTemplate, cons
 	{
 		DispatchContext = DispatchNode->GetDispatchContext();
 	}
-	
-	for (int32 ArgIndex = 0; ArgIndex < InTemplate->NumExecuteArguments(DispatchContext); ArgIndex++)
-    {
-        const FRigVMExecuteArgument* Arg = InTemplate->GetExecuteArgument(ArgIndex, DispatchContext);
-        URigVMPin* Pin = NewObject<URigVMPin>(InNode, Arg->Name);
-        const FRigVMTemplateArgumentType Type = Registry.GetType(Arg->TypeIndex);
-        
-        Pin->CPPType = Type.CPPType.ToString();
-        Pin->CPPTypeObject = Type.CPPTypeObject;
-        if (Pin->CPPTypeObject)
-        {
-        	Pin->CPPTypeObjectPath = *Pin->CPPTypeObject->GetPathName();
-        }
-        Pin->Direction = Arg->Direction;
-        Pin->LastKnownTypeIndex = Arg->TypeIndex;
-        Pin->LastKnownCPPType = Pin->CPPType;
-
-        AddNodePin(InNode, Pin);
 
-		if(Registry.IsArrayType(Arg->TypeIndex))
+	auto AddExecutePins = [InTemplate, InNode, &Registry, &DispatchContext, this](ERigVMPinDirection InPinDirection)
+	{
+		for (int32 ArgIndex = 0; ArgIndex < InTemplate->NumExecuteArguments(DispatchContext); ArgIndex++)
 		{
-			if(const URigVMDispatchNode* DispatchNode = Cast<URigVMDispatchNode>(Pin->GetNode()))
+			const FRigVMExecuteArgument* Arg = InTemplate->GetExecuteArgument(ArgIndex, DispatchContext);
+			if(Arg->Direction != InPinDirection)
 			{
-				if(const FRigVMDispatchFactory* Factory = DispatchNode->GetFactory())
+				continue;
+			}
+			
+			URigVMPin* Pin = NewObject<URigVMPin>(InNode, Arg->Name);
+			const FRigVMTemplateArgumentType Type = Registry.GetType(Arg->TypeIndex);
+	        
+			Pin->CPPType = Type.CPPType.ToString();
+			Pin->CPPTypeObject = Type.CPPTypeObject;
+			if (Pin->CPPTypeObject)
+			{
+				Pin->CPPTypeObjectPath = *Pin->CPPTypeObject->GetPathName();
+			}
+			Pin->Direction = Arg->Direction;
+			Pin->LastKnownTypeIndex = Arg->TypeIndex;
+			Pin->LastKnownCPPType = Pin->CPPType;
+
+			AddNodePin(InNode, Pin);
+
+			if(Registry.IsArrayType(Arg->TypeIndex))
+			{
+				if(const URigVMDispatchNode* DispatchNode = Cast<URigVMDispatchNode>(Pin->GetNode()))
 				{
-					const FString DefaultValue =  Factory->GetArgumentDefaultValue(Pin->GetFName(), Arg->TypeIndex);
-					if(!DefaultValue.IsEmpty())
+					if(const FRigVMDispatchFactory* Factory = DispatchNode->GetFactory())
 					{
-						SetPinDefaultValue(Pin, DefaultValue, true, false, false);
+						const FString DefaultValue =  Factory->GetArgumentDefaultValue(Pin->GetFName(), Arg->TypeIndex);
+						if(!DefaultValue.IsEmpty())
+						{
+							SetPinDefaultValue(Pin, DefaultValue, true, false, false);
+						}
 					}
 				}
 			}
 		}
-    }
+	};
+
+	AddExecutePins(ERigVMPinDirection::IO);
+	AddExecutePins(ERigVMPinDirection::Input);
 
 	for (int32 ArgIndex = 0; ArgIndex < InTemplate->NumArguments(); ArgIndex++)
 	{
@@ -14349,6 +14137,8 @@ void URigVMController::AddPinsForTemplate(const FRigVMTemplate* InTemplate, cons
 			}
 		}
 	}
+
+	AddExecutePins(ERigVMPinDirection::Output);
 }
 
 void URigVMController::ConfigurePinFromProperty(FProperty* InProperty, URigVMPin* InOutPin, ERigVMPinDirection InPinDirection)
@@ -15077,7 +14867,6 @@ void URigVMController::RepopulatePinsOnNode(URigVMNode* InNode, bool bFollowCore
 	URigVMCollapseNode* CollapseNode = Cast<URigVMCollapseNode>(InNode);
 	URigVMFunctionReferenceNode* FunctionRefNode = Cast<URigVMFunctionReferenceNode>(InNode);
 	URigVMVariableNode* VariableNode = Cast<URigVMVariableNode>(InNode);
-	URigVMArrayNode* ArrayNode = Cast<URigVMArrayNode>(InNode);
 	URigVMDispatchNode* DispatchNode = Cast<URigVMDispatchNode>(InNode);
 
 	FScopeLock Lock(&PinPathCoreRedirectorsLock);
@@ -15433,40 +15222,6 @@ void URigVMController::RepopulatePinsOnNode(URigVMNode* InNode, bool bFollowCore
 			ApplyPinStates(InNode, ReferencedPinStates);
 		}
 	}
-	else if (ArrayNode)
-	{
-		// PinStates were already saved earlier, and will be applied later, no need to do it here
-	 
-		// since we are replacing existing pins, their names have to be saved
-		// and assigned to the new pins only after we removed old ones	
-		TArray<TTuple<URigVMPin*, FName>> NewRootPinInfos;
-		for (URigVMPin* RootPin : InNode->Pins)
-		{
-			URigVMPin* NewRootPin = NewObject<URigVMPin>(InNode);
-			ConfigurePinFromPin(NewRootPin, RootPin, true);
-			EnsurePinValidity(NewRootPin, false);
-		
-			NewRootPinInfos.Add(TTuple<URigVMPin*, FName>(NewRootPin, RootPin->GetFName()));
-		}
-		
-		RemovePinsDuringRepopulate(InNode, InNode->Pins, bSetupOrphanedPins);
-		
-		for (TTuple<URigVMPin*, FName> NewRootPinInfo : NewRootPinInfos)
-		{
-			RenameObject(NewRootPinInfo.Key, *NewRootPinInfo.Value.ToString(), InNode);
-			AddNodePin(InNode, NewRootPinInfo.Key);
-		}
-		
-		for (URigVMPin* Pin : InNode->Pins)
-		{
-			if (Pin->IsStruct())
-			{
-				TGuardValue<bool> SuspendNotifications(bSuspendNotifications, true);
-				AddPinsForStruct(Pin->GetScriptStruct(), InNode, Pin, Pin->GetDirection(), FString(), false);
-			}
-			Notify(ERigVMGraphNotifType::PinAdded, Pin);
-		}	
-	}
 	else
 	{
 		return;
@@ -19770,6 +19525,35 @@ void URigVMController::PatchIfSelectNodesOnLoad()
 	}
 }
 
+void URigVMController::PatchArrayNodesOnLoad()
+{
+	if (const URigVMGraph* Graph = GetGraph())
+	{
+		TArray<URigVMNode*> ArrayNodes = Graph->GetNodes().FilterByPredicate([](URigVMNode* Node)
+		{
+			return Node->IsA<UDEPRECATED_RigVMArrayNode>();
+		});
+
+		for(URigVMNode* ModelNode : ArrayNodes)
+		{
+			UDEPRECATED_RigVMArrayNode* ArrayNode = CastChecked<UDEPRECATED_RigVMArrayNode>(ModelNode);
+			TArray<TPair<FString, FString>> LinkedPaths = GetLinkedPinPaths(ArrayNode);
+			const FVector2D NodePosition = ArrayNode->GetPosition();
+			const FString NodeName = ArrayNode->GetName();
+			const FString CPPType = ArrayNode->GetCPPType();
+			UObject* CPPTypeObject = ArrayNode->GetCPPTypeObject();
+			const ERigVMOpCode OpCode = ArrayNode->GetOpCode();
+			TMap<FString, FPinState> PinStates = GetPinStates(ArrayNode, true);
+			
+			RemoveNode(ArrayNode, false, true, false, false);
+
+			URigVMNode* NewNode = AddArrayNode(OpCode, CPPType, CPPTypeObject, NodePosition, NodeName, false, false);
+			ApplyPinStates(NewNode, PinStates, {}, false);
+			RestoreLinkedPaths(LinkedPaths, {}, {}, false);
+		}
+	}
+}
+
 void URigVMController::PostDuplicateHost(const FString& InOldPathName, const FString& InNewPathName)
 {
 	if (const URigVMGraph* Graph = GetGraph())
diff --git a/Engine/Source/Developer/RigVMDeveloper/Private/RigVMModel/RigVMControllerActions.cpp b/Engine/Source/Developer/RigVMDeveloper/Private/RigVMModel/RigVMControllerActions.cpp
--- a/Engine/Source/Developer/RigVMDeveloper/Private/RigVMModel/RigVMControllerActions.cpp
+++ b/Engine/Source/Developer/RigVMDeveloper/Private/RigVMModel/RigVMControllerActions.cpp
@@ -622,18 +622,6 @@ FRigVMRemoveNodeAction::FRigVMRemoveNodeAction(URigVMNode* InNode, URigVMControl
 	{
 		InverseAction.AddAction(FRigVMAddEnumNodeAction(EnumNode), InController);
 	}
-	else if (URigVMArrayNode* ArrayNode = Cast<URigVMArrayNode>(InNode))
-	{
-		InverseAction.AddAction(FRigVMAddArrayNodeAction(ArrayNode), InController);
-		for (URigVMPin* Pin : ArrayNode->GetPins())
-		{
-			if (Pin->GetDirection() == ERigVMPinDirection::Input ||
-				Pin->GetDirection() == ERigVMPinDirection::Visible)
-			{
-				InverseAction.AddAction(FRigVMSetPinDefaultValueAction(Pin, Pin->GetDefaultValue()), InController);
-			}
-		}
-	}
 	else if (URigVMLibraryNode* LibraryNode = Cast<URigVMLibraryNode>(InNode))
 	{
 		InverseAction.AddAction(FRigVMImportNodeFromTextAction(LibraryNode, InController), InController);
@@ -2172,48 +2160,6 @@ bool FRigVMChangeLocalVariableDefaultValueAction::Redo(URigVMController* InContr
 	return FRigVMBaseAction::Redo(InController);
 }
 
-FRigVMAddArrayNodeAction::FRigVMAddArrayNodeAction()
-	: OpCode(ERigVMOpCode::Invalid)
-	, CPPType()
-	, CPPTypeObjectPath()
-	, Position(FVector2D::ZeroVector)
-	, NodePath()
-{
-}
-
-FRigVMAddArrayNodeAction::FRigVMAddArrayNodeAction(URigVMArrayNode* InNode)
-	: OpCode(InNode->GetOpCode())
-	, CPPType(InNode->GetCPPType())
-	, CPPTypeObjectPath()
-	, Position(InNode->GetPosition())
-	, NodePath(InNode->GetNodePath())
-{
-	if (InNode->GetCPPTypeObject())
-	{
-		CPPTypeObjectPath = InNode->GetCPPTypeObject()->GetPathName();
-	}
-}
-
-bool FRigVMAddArrayNodeAction::Undo(URigVMController* InController)
-{
-	if (!FRigVMBaseAction::Undo(InController))
-	{
-		return false;
-	}
-	return InController->RemoveNodeByName(*NodePath, false);
-}
-
-bool FRigVMAddArrayNodeAction::Redo(URigVMController* InController)
-{
-#if WITH_EDITOR
-	if (URigVMArrayNode* Node = InController->AddArrayNodeFromObjectPath(OpCode, CPPType, CPPTypeObjectPath, Position, NodePath, false))
-	{
-		return FRigVMBaseAction::Redo(InController);
-	}
-#endif
-	return false;
-}
-
 FRigVMPromoteNodeAction::FRigVMPromoteNodeAction()
 	: LibraryNodePath()
 	, FunctionDefinitionPath()
diff --git a/Engine/Source/Developer/RigVMDeveloper/Public/RigVMCompiler/RigVMAST.h b/Engine/Source/Developer/RigVMDeveloper/Public/RigVMCompiler/RigVMAST.h
--- a/Engine/Source/Developer/RigVMDeveloper/Public/RigVMCompiler/RigVMAST.h
+++ b/Engine/Source/Developer/RigVMDeveloper/Public/RigVMCompiler/RigVMAST.h
@@ -26,7 +26,6 @@ class FRigVMAssignExprAST;
 class FRigVMCopyExprAST;
 class FRigVMCachedValueExprAST;
 class FRigVMExitExprAST;
-class FRigVMArrayExprAST;
 
 class URigVMPin;
 class URigVMLink;
@@ -91,7 +90,6 @@ public:
 		Copy,
 		CachedValue,
 		Exit,
-		Array,
 		Invalid
 	};
 
@@ -315,8 +313,7 @@ inline const FRigVMNodeExprAST* FRigVMExprAST::To() const
 		IsA(EType::InvokeEntry) ||
 		IsA(EType::CallExtern) ||
 		IsA(EType::InlineFunction) ||
-		IsA(EType::NoOp) ||
-		IsA(EType::Array)
+		IsA(EType::NoOp)
 	);
 	return (const FRigVMNodeExprAST*)this;
 }
@@ -453,17 +450,6 @@ inline const FRigVMExitExprAST* FRigVMExprAST::To() const
 	return (const FRigVMExitExprAST*)this;
 }
 
-// specialized cast for type checking
-// for a Array / FRigVMArrayExprAST expression
-// will raise if types are not compatible
-// @return this expression cast to FRigVMArrayExprAST 
-template<>
-inline const FRigVMArrayExprAST* FRigVMExprAST::To() const
-{
-	ensure(IsA(EType::Array));
-	return (const FRigVMArrayExprAST*)this;
-}
-
 /*
  * An abstract syntax tree block expression represents a sequence
  * of child expressions to be executed in order.
@@ -1076,41 +1062,6 @@ private:
 	friend class FRigVMParserAST;
 };
 
-/*
-* An abstract syntax tree array expression that can be used to represent
-* any array operation, such as Add, GetNum etc
-* In C++ the array op is the definition: Value[Index] etc
-*/
-class RIGVMDEVELOPER_API FRigVMArrayExprAST : public FRigVMNodeExprAST
-{
-public:
-
-	// virtual destructor
-	virtual ~FRigVMArrayExprAST() {}
-
-	// disable copy constructor
-	FRigVMArrayExprAST(const FRigVMEntryExprAST&) = delete;
-
-	// overload of the type checking mechanism
-	virtual bool IsA(EType InType) const override
-	{
-		return InType == EType::Array;
-	};
-
-protected:
-
-	// default constructor (protected so that only parser can access it)
-	FRigVMArrayExprAST(const FRigVMASTProxy& InProxy)
-		: FRigVMNodeExprAST(EType::Array, InProxy)
-	{
-	}
-
-private:
-
-	friend class FRigVMParserAST;
-};
-
-
 /*
  * The settings to apply during the parse of the abstract syntax tree.
  * The folding settings can affect the performance of the parse dramatically.
diff --git a/Engine/Source/Developer/RigVMDeveloper/Public/RigVMCompiler/RigVMCompiler.h b/Engine/Source/Developer/RigVMDeveloper/Public/RigVMCompiler/RigVMCompiler.h
--- a/Engine/Source/Developer/RigVMDeveloper/Public/RigVMCompiler/RigVMCompiler.h
+++ b/Engine/Source/Developer/RigVMDeveloper/Public/RigVMCompiler/RigVMCompiler.h
@@ -215,7 +215,6 @@ private:
 	void TraverseCopy(const FRigVMCopyExprAST* InExpr, FRigVMCompilerWorkData& WorkData);
 	void TraverseCachedValue(const FRigVMCachedValueExprAST* InExpr, FRigVMCompilerWorkData& WorkData);
 	void TraverseExit(const FRigVMExitExprAST* InExpr, FRigVMCompilerWorkData& WorkData);
-	void TraverseArray(const FRigVMArrayExprAST* InExpr, FRigVMCompilerWorkData& WorkData);
 	void TraverseInvokeEntry(const FRigVMInvokeEntryExprAST* InExpr, FRigVMCompilerWorkData& WorkData);
 
 	void AddCopyOperator(
diff --git a/Engine/Source/Developer/RigVMDeveloper/Public/RigVMModel/Nodes/RigVMArrayNode.h b/Engine/Source/Developer/RigVMDeveloper/Public/RigVMModel/Nodes/RigVMArrayNode.h
--- a/Engine/Source/Developer/RigVMDeveloper/Public/RigVMModel/Nodes/RigVMArrayNode.h
+++ b/Engine/Source/Developer/RigVMDeveloper/Public/RigVMModel/Nodes/RigVMArrayNode.h
@@ -3,43 +3,21 @@
 #pragma once
 
 #include "RigVMTemplateNode.h"
-#include "RigVMCore/RigVMExternalVariable.h"
 #include "RigVMArrayNode.generated.h"
 
 /**
  * The Array Node represents one of a series available
  * array operations such as SetNum, GetAtIndex etc.
  */
-UCLASS(BlueprintType)
-class RIGVMDEVELOPER_API URigVMArrayNode : public URigVMTemplateNode
+UCLASS(BlueprintType, Deprecated)
+class RIGVMDEVELOPER_API UDEPRECATED_RigVMArrayNode : public URigVMTemplateNode
 {
 	GENERATED_BODY()
 
 public:
 
 	// Default constructor
-	URigVMArrayNode();
-
-	// Override of node title
-	virtual FString GetNodeTitle() const override;
-
-	// static version
-	static FString GetNodeTitle(ERigVMOpCode InOpCode);
-
-	// Override of the node tooltip
-	virtual FText GetToolTipText() const override;
-
-	// Override of the node color
-	virtual FLinearColor GetNodeColor() const override;
-
-    // Override of the template notation
-	virtual FName GetNotation() const override;
-	
-	// Override of the template
-	virtual const FRigVMTemplate* GetTemplate() const override;
-
-	// Override of singleton
-	virtual bool IsSingleton() const override { return false; }
+	UDEPRECATED_RigVMArrayNode();
 
 	// Returns the op code of this node
 	UFUNCTION(BlueprintCallable, Category = RigVMArrayNode)
@@ -53,34 +31,11 @@ public:
 	UFUNCTION(BlueprintCallable, Category = RigVMArrayNode)
 	UObject* GetCPPTypeObject() const;
 
-	// Override of node color
-	virtual bool IsDefinedAsVarying() const override { return true; }
-
-	// Override control flow
-	virtual const TArray<FName>& GetControlFlowBlocks() const override;
-
-protected:
-
-	virtual FText GetToolTipTextForPin(const URigVMPin* InPin) const override;
-	
 private:
 
 	UPROPERTY()
 	ERigVMOpCode OpCode;
 
-	static const FString ArrayName;
-	static const FString NumName;
-	static const FString IndexName;
-	static const FString ElementName;
-	static const FString SuccessName;
-	static const FString OtherName;
-	static const FString CloneName;
-	static const FString CountName;
-	static const FString RatioName;
-	static const FString ResultName;
-	static const FString ContinueName;
-	static const FString CompletedName;
-
 	friend class URigVMController;
 	friend class URigVMCompiler;
 	friend class FRigVMVarExprAST;
diff --git a/Engine/Source/Developer/RigVMDeveloper/Public/RigVMModel/RigVMController.h b/Engine/Source/Developer/RigVMDeveloper/Public/RigVMModel/RigVMController.h
--- a/Engine/Source/Developer/RigVMDeveloper/Public/RigVMModel/RigVMController.h
+++ b/Engine/Source/Developer/RigVMDeveloper/Public/RigVMModel/RigVMController.h
@@ -15,7 +15,6 @@
 #include "RigVMModel/Nodes/RigVMEnumNode.h"
 #include "RigVMModel/Nodes/RigVMCollapseNode.h"
 #include "RigVMModel/Nodes/RigVMFunctionReferenceNode.h"
-#include "RigVMModel/Nodes/RigVMArrayNode.h"
 #include "RigVMModel/Nodes/RigVMInvokeEntryNode.h"
 #include "RigVMModel/RigVMBuildData.h"
 #include "RigVMCore/RigVMUserWorkflow.h"
@@ -411,12 +410,12 @@ public:
 	// Adds a Array Node to the edited Graph.
 	// This causes a NodeAdded modified event.
 	UFUNCTION(BlueprintCallable, Category = RigVMController)
-	URigVMArrayNode* AddArrayNode(ERigVMOpCode InOpCode, const FString& InCPPType, UObject* InCPPTypeObject, const FVector2D& InPosition = FVector2D::ZeroVector, const FString& InNodeName = TEXT(""), bool bSetupUndoRedo = true, bool bPrintPythonCommand = false);
+	URigVMNode* AddArrayNode(ERigVMOpCode InOpCode, const FString& InCPPType, UObject* InCPPTypeObject, const FVector2D& InPosition = FVector2D::ZeroVector, const FString& InNodeName = TEXT(""), bool bSetupUndoRedo = true, bool bPrintPythonCommand = false);
 
 	// Adds a Array Node to the edited Graph given a struct object path name.
 	// This causes a NodeAdded modified event.
 	UFUNCTION(BlueprintCallable, Category = RigVMController)
-	URigVMArrayNode* AddArrayNodeFromObjectPath(ERigVMOpCode InOpCode, const FString& InCPPType, const FString& InCPPTypeObjectPath, const FVector2D& InPosition = FVector2D::ZeroVector, const FString& InNodeName = TEXT(""), bool bSetupUndoRedo = true, bool bPrintPythonCommand = false);
+	URigVMNode* AddArrayNodeFromObjectPath(ERigVMOpCode InOpCode, const FString& InCPPType, const FString& InCPPTypeObjectPath, const FVector2D& InPosition = FVector2D::ZeroVector, const FString& InNodeName = TEXT(""), bool bSetupUndoRedo = true, bool bPrintPythonCommand = false);
 
 	// Adds an entry invocation node
 	// This causes a NodeAdded modified event.
@@ -1172,6 +1171,7 @@ protected:
 	void PatchDispatchNodesOnLoad();
 	void PatchBranchNodesOnLoad();
 	void PatchIfSelectNodesOnLoad();
+	void PatchArrayNodesOnLoad();
 
 	// work to do after a duplication of the host asset
 	void PostDuplicateHost(const FString& InOldPathName, const FString& InNewPathName);
diff --git a/Engine/Source/Developer/RigVMDeveloper/Public/RigVMModel/RigVMControllerActions.h b/Engine/Source/Developer/RigVMDeveloper/Public/RigVMModel/RigVMControllerActions.h
--- a/Engine/Source/Developer/RigVMDeveloper/Public/RigVMModel/RigVMControllerActions.h
+++ b/Engine/Source/Developer/RigVMDeveloper/Public/RigVMModel/RigVMControllerActions.h
@@ -1684,39 +1684,6 @@ public:
 	FString DefaultValue;
 };
 
-/**
-* An action adding a array node to the graph.
-*/
-USTRUCT()
-struct FRigVMAddArrayNodeAction : public FRigVMBaseAction
-{
-	GENERATED_BODY()
-
-public:
-
-	FRigVMAddArrayNodeAction();
-	FRigVMAddArrayNodeAction(URigVMArrayNode* InNode);
-	virtual ~FRigVMAddArrayNodeAction() {};
-	virtual UScriptStruct* GetScriptStruct() const override { return FRigVMAddArrayNodeAction::StaticStruct(); }
-	virtual bool Undo(URigVMController* InController) override;
-	virtual bool Redo(URigVMController* InController) override;
-
-	UPROPERTY()
-	ERigVMOpCode OpCode;
-	
-	UPROPERTY()
-	FString CPPType;
-
-	UPROPERTY()
-	FString CPPTypeObjectPath;
-
-	UPROPERTY()
-	FVector2D Position;
-
-	UPROPERTY()
-	FString NodePath;
-};
-
 /**
  * An action to promote a function to collapse node or vice versa
  */
diff --git a/Engine/Source/Runtime/RigVM/Private/RigVMCore/RigVM.cpp b/Engine/Source/Runtime/RigVM/Private/RigVMCore/RigVM.cpp
--- a/Engine/Source/Runtime/RigVM/Private/RigVMCore/RigVM.cpp
+++ b/Engine/Source/Runtime/RigVM/Private/RigVMCore/RigVM.cpp
@@ -485,8 +485,6 @@ bool URigVM::ValidateAllOperandsDuringLoad()
 			case ERigVMOpCode::BoolTrue:
 			case ERigVMOpCode::Increment:
 			case ERigVMOpCode::Decrement:
-			case ERigVMOpCode::ArrayReset:
-			case ERigVMOpCode::ArrayReverse:
 			{
 				const FRigVMUnaryOp& Op = ByteCodeStorage.GetOpAt<FRigVMUnaryOp>(ByteCodeInstruction);
 				CheckOperandValidity(Op.Arg);
@@ -517,51 +515,12 @@ bool URigVM::ValidateAllOperandsDuringLoad()
 				break;
 			}
 			case ERigVMOpCode::BeginBlock:
-			case ERigVMOpCode::ArrayGetNum:
-			case ERigVMOpCode::ArraySetNum:
-			case ERigVMOpCode::ArrayAppend:
-			case ERigVMOpCode::ArrayClone:
-			case ERigVMOpCode::ArrayRemove:
-			case ERigVMOpCode::ArrayUnion:
 			{
 				const FRigVMBinaryOp& Op = ByteCodeStorage.GetOpAt<FRigVMBinaryOp>(ByteCodeInstruction);
 				CheckOperandValidity(Op.ArgA);
 				CheckOperandValidity(Op.ArgB);
 				break;
 			}
-			case ERigVMOpCode::ArrayAdd:
-			case ERigVMOpCode::ArrayGetAtIndex:
-			case ERigVMOpCode::ArraySetAtIndex:
-			case ERigVMOpCode::ArrayInsert:
-			case ERigVMOpCode::ArrayDifference:
-			case ERigVMOpCode::ArrayIntersection:
-			{
-				const FRigVMTernaryOp& Op = ByteCodeStorage.GetOpAt<FRigVMTernaryOp>(ByteCodeInstruction);
-				CheckOperandValidity(Op.ArgA);
-				CheckOperandValidity(Op.ArgB);
-				CheckOperandValidity(Op.ArgC);
-				break;
-			}
-			case ERigVMOpCode::ArrayFind:
-			{
-				const FRigVMQuaternaryOp& Op = ByteCodeStorage.GetOpAt<FRigVMQuaternaryOp>(ByteCodeInstruction);
-				CheckOperandValidity(Op.ArgA);
-				CheckOperandValidity(Op.ArgB);
-				CheckOperandValidity(Op.ArgC);
-				CheckOperandValidity(Op.ArgD);
-				break;
-			}
-			case ERigVMOpCode::ArrayIterator:
-			{
-				const FRigVMSenaryOp& Op = ByteCodeStorage.GetOpAt<FRigVMSenaryOp>(ByteCodeInstruction);
-				CheckOperandValidity(Op.ArgA);
-				CheckOperandValidity(Op.ArgB);
-				CheckOperandValidity(Op.ArgC);
-				CheckOperandValidity(Op.ArgD);
-				CheckOperandValidity(Op.ArgE);
-				CheckOperandValidity(Op.ArgF);
-				break;
-			}
 			case ERigVMOpCode::JumpToBranch:
 			{
 				const FRigVMJumpToBranchOp& Op = ByteCodeStorage.GetOpAt<FRigVMJumpToBranchOp>(ByteCodeInstruction);
@@ -1178,8 +1137,6 @@ void URigVM::CacheMemoryHandlesIfRequired(TArrayView<URigVMMemoryStorage*> InMem
 			case ERigVMOpCode::BoolTrue:
 			case ERigVMOpCode::Increment:
 			case ERigVMOpCode::Decrement:
-			case ERigVMOpCode::ArrayReset:
-			case ERigVMOpCode::ArrayReverse:
 				{
 					const FRigVMUnaryOp& Op = ByteCode.GetOpAt<FRigVMUnaryOp>(Instructions[InstructionIndex]);
 					InOpFunc(InHandleBaseIndex, {}, Op.Arg);
@@ -1231,51 +1188,12 @@ void URigVM::CacheMemoryHandlesIfRequired(TArrayView<URigVMMemoryStorage*> InMem
 					break;
 				}
 			case ERigVMOpCode::BeginBlock:
-			case ERigVMOpCode::ArrayGetNum:
-			case ERigVMOpCode::ArraySetNum:
-			case ERigVMOpCode::ArrayAppend:
-			case ERigVMOpCode::ArrayClone:
-			case ERigVMOpCode::ArrayRemove:
-			case ERigVMOpCode::ArrayUnion:
 				{
 					const FRigVMBinaryOp& Op = ByteCode.GetOpAt<FRigVMBinaryOp>(Instructions[InstructionIndex]);
 					InOpFunc(InHandleBaseIndex + 0, {}, Op.ArgA);
 					InOpFunc(InHandleBaseIndex + 1, {}, Op.ArgB);
 					break;
 				}
-			case ERigVMOpCode::ArrayAdd:
-			case ERigVMOpCode::ArrayGetAtIndex:
-			case ERigVMOpCode::ArraySetAtIndex:
-			case ERigVMOpCode::ArrayInsert:
-			case ERigVMOpCode::ArrayDifference:
-			case ERigVMOpCode::ArrayIntersection:
-				{
-					const FRigVMTernaryOp& Op = ByteCode.GetOpAt<FRigVMTernaryOp>(Instructions[InstructionIndex]);
-					InOpFunc(InHandleBaseIndex + 0, {}, Op.ArgA);
-					InOpFunc(InHandleBaseIndex + 1, {}, Op.ArgB);
-					InOpFunc(InHandleBaseIndex + 2, {}, Op.ArgC);
-					break;
-				}
-			case ERigVMOpCode::ArrayFind:
-				{
-					const FRigVMQuaternaryOp& Op = ByteCode.GetOpAt<FRigVMQuaternaryOp>(Instructions[InstructionIndex]);
-					InOpFunc(InHandleBaseIndex + 0, {}, Op.ArgA);
-					InOpFunc(InHandleBaseIndex + 1, {}, Op.ArgB);
-					InOpFunc(InHandleBaseIndex + 2, {}, Op.ArgC);
-					InOpFunc(InHandleBaseIndex + 3, {}, Op.ArgD);
-					break;
-				}
-			case ERigVMOpCode::ArrayIterator:
-				{
-					const FRigVMSenaryOp& Op = ByteCode.GetOpAt<FRigVMSenaryOp>(Instructions[InstructionIndex]);
-					InOpFunc(InHandleBaseIndex + 0, {}, Op.ArgA);
-					InOpFunc(InHandleBaseIndex + 1, {}, Op.ArgB);
-					InOpFunc(InHandleBaseIndex + 2, {}, Op.ArgC);
-					InOpFunc(InHandleBaseIndex + 3, {}, Op.ArgD);
-					InOpFunc(InHandleBaseIndex + 4, {}, Op.ArgE);
-					InOpFunc(InHandleBaseIndex + 5, {}, Op.ArgF);
-					break;
-				}
 			case ERigVMOpCode::EndBlock:
 				{
 					break;
@@ -2142,545 +2060,6 @@ ERigVMExecuteResult URigVM::ExecuteInstructions(int32 InFirstInstruction, int32
 				ContextPublicData.InstructionIndex++;
 				break;
 			}
-			case ERigVMOpCode::ArrayReset:
-			{
-				FRigVMMemoryHandle& ArrayHandle = CachedMemoryHandles[FirstHandleForInstruction[ContextPublicData.InstructionIndex]]; 					
-				FScriptArrayHelper ArrayHelper(CastFieldChecked<FArrayProperty>(ArrayHandle.GetProperty()), ArrayHandle.GetData());
-				ArrayHelper.Resize(0);
-
-				if(DebugMemoryStorageObject->Num() > 0)
-				{
-					const FRigVMUnaryOp& Op = ByteCode.GetOpAt<FRigVMUnaryOp>(Instruction);
-					CopyOperandForDebuggingIfNeeded(Op.Arg, ArrayHandle);
-				}
-
-				ContextPublicData.InstructionIndex++;
-				break;
-			}
-			case ERigVMOpCode::ArrayGetNum:
-			{
-				FRigVMMemoryHandle& ArrayHandle = CachedMemoryHandles[FirstHandleForInstruction[ContextPublicData.InstructionIndex]]; 					
-				FScriptArrayHelper ArrayHelper(CastFieldChecked<FArrayProperty>(ArrayHandle.GetProperty()), ArrayHandle.GetData());
-				int32& Count = (*((int32*)CachedMemoryHandles[FirstHandleForInstruction[ContextPublicData.InstructionIndex] + 1].GetData()));
-				Count = ArrayHelper.Num();
-
-				if(DebugMemoryStorageObject->Num() > 0)
-				{
-					const FRigVMBinaryOp& Op = ByteCode.GetOpAt<FRigVMBinaryOp>(Instruction);
-					CopyOperandForDebuggingIfNeeded(Op.ArgA, ArrayHandle);
-					CopyOperandForDebuggingIfNeeded(Op.ArgB, CachedMemoryHandles[FirstHandleForInstruction[ContextPublicData.InstructionIndex] + 1]);
-				}
-
-				ContextPublicData.InstructionIndex++;
-				break;
-			}
-			case ERigVMOpCode::ArraySetNum:
-			{
-				FRigVMMemoryHandle& ArrayHandle = CachedMemoryHandles[FirstHandleForInstruction[ContextPublicData.InstructionIndex]]; 					
-				FScriptArrayHelper ArrayHelper(CastFieldChecked<FArrayProperty>(ArrayHandle.GetProperty()), ArrayHandle.GetData());
-				const int32 Count = (*((int32*)CachedMemoryHandles[FirstHandleForInstruction[ContextPublicData.InstructionIndex] + 1].GetData()));
-				if(Context.IsValidArraySize(Count))
-				{
-					ArrayHelper.Resize(Count);
-				}
-
-				if(DebugMemoryStorageObject->Num() > 0)
-				{
-					const FRigVMBinaryOp& Op = ByteCode.GetOpAt<FRigVMBinaryOp>(Instruction);
-					CopyOperandForDebuggingIfNeeded(Op.ArgA, ArrayHandle);
-					CopyOperandForDebuggingIfNeeded(Op.ArgB, CachedMemoryHandles[FirstHandleForInstruction[ContextPublicData.InstructionIndex] + 1]);
-				}
-
-				ContextPublicData.InstructionIndex++;
-				break;
-			}
-			case ERigVMOpCode::ArrayAppend:
-			{
-				FRigVMMemoryHandle& ArrayHandle = CachedMemoryHandles[FirstHandleForInstruction[ContextPublicData.InstructionIndex]]; 					
-				FRigVMMemoryHandle& OtherArrayHandle = CachedMemoryHandles[FirstHandleForInstruction[ContextPublicData.InstructionIndex] + 1];
-				const FArrayProperty* ArrayProperty = CastFieldChecked<FArrayProperty>(ArrayHandle.GetProperty()); 
-				const FArrayProperty* OtherArrayProperty = CastFieldChecked<FArrayProperty>(OtherArrayHandle.GetProperty()); 
-
-				FScriptArrayHelper ArrayHelper(ArrayProperty, ArrayHandle.GetData());
-				FScriptArrayHelper OtherArrayHelper(OtherArrayProperty, OtherArrayHandle.GetData());
-
-				if(OtherArrayHelper.Num() > 0)
-				{
-					if(Context.IsValidArraySize(ArrayHelper.Num() + OtherArrayHelper.Num()))
-					{
-						const FProperty* TargetProperty = ArrayProperty->Inner;
-						const FProperty* SourceProperty = OtherArrayProperty->Inner;
-
-						int32 TargetIndex = ArrayHelper.AddValues(OtherArrayHelper.Num());
-						for(int32 SourceIndex = 0; SourceIndex < OtherArrayHelper.Num(); SourceIndex++, TargetIndex++)
-						{
-							uint8* TargetMemory = ArrayHelper.GetRawPtr(TargetIndex);
-							const uint8* SourceMemory = OtherArrayHelper.GetRawPtr(SourceIndex);
-							URigVMMemoryStorage::CopyProperty(TargetProperty, TargetMemory, SourceProperty, SourceMemory);
-						}
-					}
-				}
-
-				if(DebugMemoryStorageObject->Num() > 0)
-				{
-					const FRigVMBinaryOp& Op = ByteCode.GetOpAt<FRigVMBinaryOp>(Instruction);
-					CopyOperandForDebuggingIfNeeded(Op.ArgA, ArrayHandle);
-					CopyOperandForDebuggingIfNeeded(Op.ArgB, CachedMemoryHandles[FirstHandleForInstruction[ContextPublicData.InstructionIndex] + 1]);
-				}
-
-				ContextPublicData.InstructionIndex++;
-				break;
-			}
-			case ERigVMOpCode::ArrayClone:
-			{
-				FRigVMMemoryHandle& ArrayHandle = CachedMemoryHandles[FirstHandleForInstruction[ContextPublicData.InstructionIndex]]; 					
-				FRigVMMemoryHandle& ClonedArrayHandle = CachedMemoryHandles[FirstHandleForInstruction[ContextPublicData.InstructionIndex] + 1];
-				const FArrayProperty* ArrayProperty = CastFieldChecked<FArrayProperty>(ArrayHandle.GetProperty()); 
-				const FArrayProperty* ClonedArrayProperty = CastFieldChecked<FArrayProperty>(ClonedArrayHandle.GetProperty()); 
-				FScriptArrayHelper ArrayHelper(ArrayProperty, ArrayHandle.GetData());
-				FScriptArrayHelper ClonedArrayHelper(ClonedArrayProperty, ClonedArrayHandle.GetData());
-
-				CopyArray(ClonedArrayHelper, ClonedArrayHandle, ArrayHelper, ArrayHandle);
-					
-				if(DebugMemoryStorageObject->Num() > 0)
-				{
-					const FRigVMBinaryOp& Op = ByteCode.GetOpAt<FRigVMBinaryOp>(Instruction);
-					CopyOperandForDebuggingIfNeeded(Op.ArgA, ArrayHandle);
-					CopyOperandForDebuggingIfNeeded(Op.ArgB, ClonedArrayHandle);
-				}
-
-				ContextPublicData.InstructionIndex++;
-				break;
-			}
-			case ERigVMOpCode::ArrayGetAtIndex:
-			{
-				FRigVMMemoryHandle& ArrayHandle = CachedMemoryHandles[FirstHandleForInstruction[ContextPublicData.InstructionIndex]];
-				const FArrayProperty* ArrayProperty = CastFieldChecked<FArrayProperty>(ArrayHandle.GetProperty());
-				FScriptArrayHelper ArrayHelper(ArrayProperty, ArrayHandle.GetData());
-				int32 Index = (*((int32*)CachedMemoryHandles[FirstHandleForInstruction[ContextPublicData.InstructionIndex] + 1].GetData()));
-
-				if(Context.IsValidArrayIndex(Index, ArrayHelper))
-				{
-					FRigVMMemoryHandle& ElementHandle = CachedMemoryHandles[FirstHandleForInstruction[ContextPublicData.InstructionIndex] + 2];
-					uint8* TargetMemory = ElementHandle.GetData();
-					const uint8* SourceMemory = ArrayHelper.GetRawPtr(Index);
-					URigVMMemoryStorage::CopyProperty(ElementHandle.GetProperty(), TargetMemory, ArrayProperty->Inner, SourceMemory);
-				}
-
-				if(DebugMemoryStorageObject->Num() > 0)
-				{
-					const FRigVMTernaryOp& Op = ByteCode.GetOpAt<FRigVMTernaryOp>(Instruction);
-					CopyOperandForDebuggingIfNeeded(Op.ArgA, ArrayHandle);
-					CopyOperandForDebuggingIfNeeded(Op.ArgB, CachedMemoryHandles[FirstHandleForInstruction[ContextPublicData.InstructionIndex] + 1]);
-					CopyOperandForDebuggingIfNeeded(Op.ArgC, CachedMemoryHandles[FirstHandleForInstruction[ContextPublicData.InstructionIndex] + 2]);
-				}
-
-				ContextPublicData.InstructionIndex++;
-				break;
-			}
-			case ERigVMOpCode::ArraySetAtIndex:
-			{
-				FRigVMMemoryHandle& ArrayHandle = CachedMemoryHandles[FirstHandleForInstruction[ContextPublicData.InstructionIndex]];
-				const FArrayProperty* ArrayProperty = CastFieldChecked<FArrayProperty>(ArrayHandle.GetProperty());
-				FScriptArrayHelper ArrayHelper(ArrayProperty, ArrayHandle.GetData());
-				int32 Index = (*((int32*)CachedMemoryHandles[FirstHandleForInstruction[ContextPublicData.InstructionIndex] + 1].GetData()));
-				if(Context.IsValidArrayIndex(Index, ArrayHelper))
-				{
-					FRigVMMemoryHandle& ElementHandle = CachedMemoryHandles[FirstHandleForInstruction[ContextPublicData.InstructionIndex] + 2];
-					uint8* TargetMemory = ArrayHelper.GetRawPtr(Index);
-					const uint8* SourceMemory = ElementHandle.GetData();
-					URigVMMemoryStorage::CopyProperty(ArrayProperty->Inner, TargetMemory, ElementHandle.GetProperty(), SourceMemory);
-				}
-
-				if(DebugMemoryStorageObject->Num() > 0)
-				{
-					const FRigVMTernaryOp& Op = ByteCode.GetOpAt<FRigVMTernaryOp>(Instruction);
-					CopyOperandForDebuggingIfNeeded(Op.ArgA, ArrayHandle);
-					CopyOperandForDebuggingIfNeeded(Op.ArgB, CachedMemoryHandles[FirstHandleForInstruction[ContextPublicData.InstructionIndex] + 1]);
-					CopyOperandForDebuggingIfNeeded(Op.ArgC, CachedMemoryHandles[FirstHandleForInstruction[ContextPublicData.InstructionIndex] + 2]);
-				}
-
-				ContextPublicData.InstructionIndex++;
-				break;
-			}
-			case ERigVMOpCode::ArrayInsert:
-			{
-				FRigVMMemoryHandle& ArrayHandle = CachedMemoryHandles[FirstHandleForInstruction[ContextPublicData.InstructionIndex]];
-				const FArrayProperty* ArrayProperty = CastFieldChecked<FArrayProperty>(ArrayHandle.GetProperty());
-				FScriptArrayHelper ArrayHelper(ArrayProperty, ArrayHandle.GetData());
-				if(Context.IsValidArraySize(ArrayHelper.Num() + 1))
-				{
-					int32 Index = (*((int32*)CachedMemoryHandles[FirstHandleForInstruction[ContextPublicData.InstructionIndex] + 1].GetData()));
-
-					// we support wrapping the index around similar to python
-					if(Index < 0)
-					{
-						Index = ArrayHelper.Num() + Index;
-					}
-					
-					Index = FMath::Clamp<int32>(Index, 0, ArrayHelper.Num());
-					ArrayHelper.InsertValues(Index, 1);
-
-					FRigVMMemoryHandle& ElementHandle = CachedMemoryHandles[FirstHandleForInstruction[ContextPublicData.InstructionIndex] + 2];
-					uint8* TargetMemory = ArrayHelper.GetRawPtr(Index);
-					const uint8* SourceMemory = ElementHandle.GetData();
-					URigVMMemoryStorage::CopyProperty(ArrayProperty->Inner, TargetMemory, ElementHandle.GetProperty(), SourceMemory);
-				}
-
-				if(DebugMemoryStorageObject->Num() > 0)
-				{
-					const FRigVMTernaryOp& Op = ByteCode.GetOpAt<FRigVMTernaryOp>(Instruction);
-					CopyOperandForDebuggingIfNeeded(Op.ArgA, ArrayHandle);
-					CopyOperandForDebuggingIfNeeded(Op.ArgB, CachedMemoryHandles[FirstHandleForInstruction[ContextPublicData.InstructionIndex] + 1]);
-					CopyOperandForDebuggingIfNeeded(Op.ArgC, CachedMemoryHandles[FirstHandleForInstruction[ContextPublicData.InstructionIndex] + 2]);
-				}
-
-				ContextPublicData.InstructionIndex++;
-				break;
-			}
-			case ERigVMOpCode::ArrayRemove:
-			{
-				FRigVMMemoryHandle& ArrayHandle = CachedMemoryHandles[FirstHandleForInstruction[ContextPublicData.InstructionIndex]];
-				const FArrayProperty* ArrayProperty = CastFieldChecked<FArrayProperty>(ArrayHandle.GetProperty());
-				FScriptArrayHelper ArrayHelper(ArrayProperty, ArrayHandle.GetData());
-				int32 Index = (*((int32*)CachedMemoryHandles[FirstHandleForInstruction[ContextPublicData.InstructionIndex] + 1].GetData()));
-				if(Context.IsValidArrayIndex(Index, ArrayHelper))
-				{
-					ArrayHelper.RemoveValues(Index, 1);
-				}
-
-				if(DebugMemoryStorageObject->Num() > 0)
-				{
-					const FRigVMBinaryOp& Op = ByteCode.GetOpAt<FRigVMBinaryOp>(Instruction);
-					CopyOperandForDebuggingIfNeeded(Op.ArgA, ArrayHandle);
-					CopyOperandForDebuggingIfNeeded(Op.ArgB, CachedMemoryHandles[FirstHandleForInstruction[ContextPublicData.InstructionIndex] + 1]);
-				}
-
-				ContextPublicData.InstructionIndex++;
-				break;
-			}
-			case ERigVMOpCode::ArrayAdd:
-			{
-				FRigVMMemoryHandle& ArrayHandle = CachedMemoryHandles[FirstHandleForInstruction[ContextPublicData.InstructionIndex]];
-				const FArrayProperty* ArrayProperty = CastFieldChecked<FArrayProperty>(ArrayHandle.GetProperty());
-				FScriptArrayHelper ArrayHelper(ArrayProperty, ArrayHandle.GetData());
-				int32& Index = (*((int32*)CachedMemoryHandles[FirstHandleForInstruction[ContextPublicData.InstructionIndex] + 2].GetData()));
-				if(Context.IsValidArraySize(ArrayHelper.Num() + 1))
-				{
-					FRigVMMemoryHandle& ElementHandle = CachedMemoryHandles[FirstHandleForInstruction[ContextPublicData.InstructionIndex] + 1];
-					Index = ArrayHelper.AddValue();
-
-					uint8* TargetMemory = ArrayHelper.GetRawPtr(Index);
-					const uint8* SourceMemory = ElementHandle.GetData();
-					URigVMMemoryStorage::CopyProperty(ArrayProperty->Inner, TargetMemory, ElementHandle.GetProperty(), SourceMemory);
-				}
-				else
-				{
-					Index = INDEX_NONE;
-				}
-
-				if(DebugMemoryStorageObject->Num() > 0)
-				{
-					const FRigVMTernaryOp& Op = ByteCode.GetOpAt<FRigVMTernaryOp>(Instruction);
-					CopyOperandForDebuggingIfNeeded(Op.ArgA, ArrayHandle);
-					CopyOperandForDebuggingIfNeeded(Op.ArgB, CachedMemoryHandles[FirstHandleForInstruction[ContextPublicData.InstructionIndex] + 1]);
-					CopyOperandForDebuggingIfNeeded(Op.ArgC, CachedMemoryHandles[FirstHandleForInstruction[ContextPublicData.InstructionIndex] + 2]);
-				}
-
-				ContextPublicData.InstructionIndex++;
-				break;
-			}
-			case ERigVMOpCode::ArrayFind:
-			{
-				FRigVMMemoryHandle& ArrayHandle = CachedMemoryHandles[FirstHandleForInstruction[ContextPublicData.InstructionIndex]];
-				const FArrayProperty* ArrayProperty = CastFieldChecked<FArrayProperty>(ArrayHandle.GetProperty());
-				FScriptArrayHelper ArrayHelper(ArrayProperty, ArrayHandle.GetData());
-
-				FRigVMMemoryHandle& ElementHandle = CachedMemoryHandles[FirstHandleForInstruction[ContextPublicData.InstructionIndex] + 1];
-				int32& FoundIndex = (*((int32*)CachedMemoryHandles[FirstHandleForInstruction[ContextPublicData.InstructionIndex] + 2].GetData()));
-				bool& bFound = (*((bool*)CachedMemoryHandles[FirstHandleForInstruction[ContextPublicData.InstructionIndex] + 3].GetData()));
-
-				FoundIndex = INDEX_NONE;
-				bFound = false;
-
-				const FProperty* PropertyA = ElementHandle.GetProperty();
-				const FProperty* PropertyB = ArrayProperty->Inner;
-
-				if(PropertyA->SameType(PropertyB))
-				{
-					const uint8* MemoryA = ElementHandle.GetData();
-
-					for(int32 Index = 0; Index < ArrayHelper.Num(); Index++)
-					{
-						const uint8* MemoryB = ArrayHelper.GetRawPtr(Index);
-						if(PropertyA->Identical(MemoryA, MemoryB))
-						{
-							FoundIndex = Index;
-							bFound = true;
-							break;
-						}
-					}
-				}
-				else
-				{
-					static const TCHAR IncompatibleTypes[] = TEXT("Array('%s') doesn't support searching for element('%$s').");
-					ContextPublicData.Logf(EMessageSeverity::Error, IncompatibleTypes, *PropertyB->GetCPPType(), *PropertyA->GetCPPType());
-				}
-
-				if(DebugMemoryStorageObject->Num() > 0)
-				{
-					const FRigVMQuaternaryOp& Op = ByteCode.GetOpAt<FRigVMQuaternaryOp>(Instruction);
-					CopyOperandForDebuggingIfNeeded(Op.ArgA, ArrayHandle);
-					CopyOperandForDebuggingIfNeeded(Op.ArgB, CachedMemoryHandles[FirstHandleForInstruction[ContextPublicData.InstructionIndex] + 1]);
-					CopyOperandForDebuggingIfNeeded(Op.ArgC, CachedMemoryHandles[FirstHandleForInstruction[ContextPublicData.InstructionIndex] + 2]);
-					CopyOperandForDebuggingIfNeeded(Op.ArgD, CachedMemoryHandles[FirstHandleForInstruction[ContextPublicData.InstructionIndex] + 3]);
-				}
-
-				ContextPublicData.InstructionIndex++;
-				break;
-			}
-			case ERigVMOpCode::ArrayIterator:
-			{
-				FRigVMMemoryHandle& ArrayHandle = CachedMemoryHandles[FirstHandleForInstruction[ContextPublicData.InstructionIndex]];
-				const FArrayProperty* ArrayProperty = CastFieldChecked<FArrayProperty>(ArrayHandle.GetProperty());
-				FScriptArrayHelper ArrayHelper(ArrayProperty, ArrayHandle.GetData());
-
-				FRigVMMemoryHandle& ElementHandle = CachedMemoryHandles[FirstHandleForInstruction[ContextPublicData.InstructionIndex] + 1];
-				const int32& Index = (*((int32*)CachedMemoryHandles[FirstHandleForInstruction[ContextPublicData.InstructionIndex] + 2].GetData()));
-				int32& Count = (*((int32*)CachedMemoryHandles[FirstHandleForInstruction[ContextPublicData.InstructionIndex] + 3].GetData()));
-				float& Ratio = (*((float*)CachedMemoryHandles[FirstHandleForInstruction[ContextPublicData.InstructionIndex] + 4].GetData()));
-				bool& bContinue = (*((bool*)CachedMemoryHandles[FirstHandleForInstruction[ContextPublicData.InstructionIndex] + 5].GetData()));
-
-				Count = ArrayHelper.Num();
-				bContinue = Index >=0 && Index < Count;
-
-				if((Count <= 0) || !bContinue)
-				{
-					Ratio = 0.f;
-				}
-				else
-				{
-					Ratio = float(Index) / float(Count - 1);
-
-					uint8* TargetMemory = ElementHandle.GetData();
-					const uint8* SourceMemory = ArrayHelper.GetRawPtr(Index);
-					URigVMMemoryStorage::CopyProperty(ElementHandle.GetProperty(), TargetMemory, ArrayProperty->Inner, SourceMemory);
-
-					if(DebugMemoryStorageObject->Num() > 0)
-					{
-						const FRigVMSenaryOp& Op = ByteCode.GetOpAt<FRigVMSenaryOp>(Instruction);
-						CopyOperandForDebuggingIfNeeded(Op.ArgA, ArrayHandle); // array
-						CopyOperandForDebuggingIfNeeded(Op.ArgD, CachedMemoryHandles[FirstHandleForInstruction[ContextPublicData.InstructionIndex] + 3]); // count
-
-						Context.BeginSlice(Count, Index);
-						CopyOperandForDebuggingIfNeeded(Op.ArgB, CachedMemoryHandles[FirstHandleForInstruction[ContextPublicData.InstructionIndex] + 1]); // element
-						CopyOperandForDebuggingIfNeeded(Op.ArgC, CachedMemoryHandles[FirstHandleForInstruction[ContextPublicData.InstructionIndex] + 2]); // index
-						CopyOperandForDebuggingIfNeeded(Op.ArgE, CachedMemoryHandles[FirstHandleForInstruction[ContextPublicData.InstructionIndex] + 4]); // ratio
-						Context.EndSlice();
-					}
-				}
-
-				ContextPublicData.InstructionIndex++;
-				break;
-			}
-			case ERigVMOpCode::ArrayUnion:
-			case ERigVMOpCode::ArrayDifference:
-			case ERigVMOpCode::ArrayIntersection:
-			{
-				FRigVMMemoryHandle& ArrayHandleA = CachedMemoryHandles[FirstHandleForInstruction[ContextPublicData.InstructionIndex]]; 					
-				FRigVMMemoryHandle& ArrayHandleB = CachedMemoryHandles[FirstHandleForInstruction[ContextPublicData.InstructionIndex] + 1];
-				const FArrayProperty* ArrayPropertyA = CastFieldChecked<FArrayProperty>(ArrayHandleA.GetProperty()); 					
-				const FArrayProperty* ArrayPropertyB = CastFieldChecked<FArrayProperty>(ArrayHandleB.GetProperty()); 					
-				FScriptArrayHelper ArrayHelperA(ArrayPropertyA, ArrayHandleA.GetData());
-				FScriptArrayHelper ArrayHelperB(ArrayPropertyB, ArrayHandleB.GetData());
-				const FProperty* ElementPropertyA = ArrayPropertyA->Inner;
-				const FProperty* ElementPropertyB = ArrayPropertyB->Inner;
-
-				TMap<uint32, int32> HashA, HashB;
-				HashA.Reserve(ArrayHelperA.Num());
-				HashB.Reserve(ArrayHelperB.Num());
-
-				for(int32 Index = 0; Index < ArrayHelperA.Num(); Index++)
-				{
-					uint32 HashValue;
-					if (ElementPropertyA->PropertyFlags & CPF_HasGetValueTypeHash)
-					{
-						HashValue = ElementPropertyA->GetValueTypeHash(ArrayHelperA.GetRawPtr(Index));
-					}
-					else
-					{
-						FString Value;
-						ElementPropertyA->ExportTextItem_Direct(Value, ArrayHelperA.GetRawPtr(Index), nullptr, nullptr, PPF_None);
-						HashValue = TextKeyUtil::HashString(Value);
-					}
-					
-					if(!HashA.Contains(HashValue))
-					{
-						HashA.Add(HashValue, Index);
-					}
-				}
-				for(int32 Index = 0; Index < ArrayHelperB.Num(); Index++)
-				{
-					uint32 HashValue;
-					if (ElementPropertyB->PropertyFlags & CPF_HasGetValueTypeHash)
-					{
-						HashValue = ElementPropertyB->GetValueTypeHash(ArrayHelperB.GetRawPtr(Index));
-					}
-					else
-					{
-						FString Value;
-						ElementPropertyB->ExportTextItem_Direct(Value, ArrayHelperB.GetRawPtr(Index), nullptr, nullptr, PPF_None);
-						HashValue = TextKeyUtil::HashString(Value);
-					}
-					if(!HashB.Contains(HashValue))
-					{
-						HashB.Add(HashValue, Index);
-					}
-				}
-
-				if(Instruction.OpCode == ERigVMOpCode::ArrayUnion)
-				{
-					// copy the complete array to a temp storage
-					TArray<uint8, TAlignedHeapAllocator<16>> TempStorage;
-					const int32 NumElementsA = ArrayHelperA.Num();
-					TempStorage.AddZeroed(NumElementsA * ElementPropertyA->GetSize());
-					uint8* TempMemory = TempStorage.GetData();
-					for(int32 Index = 0; Index < NumElementsA; Index++)
-					{
-						ElementPropertyA->InitializeValue(TempMemory);
-						ElementPropertyA->CopyCompleteValue(TempMemory, ArrayHelperA.GetRawPtr(Index));
-						TempMemory += ElementPropertyA->GetSize();
-					}
-
-					ArrayHelperA.Resize(0);
-
-					for(const TPair<uint32, int32>& Pair : HashA)
-					{
-						int32 AddedIndex = ArrayHelperA.AddValue();
-						TempMemory = TempStorage.GetData() + Pair.Value * ElementPropertyA->GetSize();
-						
-						URigVMMemoryStorage::CopyProperty(
-							ElementPropertyA,
-							ArrayHelperA.GetRawPtr(AddedIndex),
-							ElementPropertyA,
-							TempMemory
-						);
-					}
-
-					TempMemory = TempStorage.GetData();
-					for(int32 Index = 0; Index < NumElementsA; Index++)
-					{
-						ElementPropertyA->DestroyValue(TempMemory);
-						TempMemory += ElementPropertyA->GetSize();
-					}
-
-					for(const TPair<uint32, int32>& Pair : HashB)
-					{
-						if(!HashA.Contains(Pair.Key))
-						{
-							int32 AddedIndex = ArrayHelperA.AddValue();
-							
-							URigVMMemoryStorage::CopyProperty(
-								ElementPropertyA,
-								ArrayHelperA.GetRawPtr(AddedIndex),
-								ElementPropertyB,
-								ArrayHelperB.GetRawPtr(Pair.Value)
-							);
-						}
-					}
-					
-					if(DebugMemoryStorageObject->Num() > 0)
-					{
-						const FRigVMBinaryOp& Op = ByteCode.GetOpAt<FRigVMBinaryOp>(Instruction);
-						CopyOperandForDebuggingIfNeeded(Op.ArgA, ArrayHandleA);
-						CopyOperandForDebuggingIfNeeded(Op.ArgB, ArrayHandleB);
-					}
-				}
-				else
-				{
-					FRigVMMemoryHandle& ResultArrayHandle = CachedMemoryHandles[FirstHandleForInstruction[ContextPublicData.InstructionIndex] + 2]; 					
-					FScriptArrayHelper ResultArrayHelper(CastFieldChecked<FArrayProperty>(ResultArrayHandle.GetProperty()), ResultArrayHandle.GetData());
-					const FArrayProperty* ResultArrayProperty = CastFieldChecked<FArrayProperty>(ResultArrayHandle.GetProperty());
-					const FProperty* ResultElementProperty = ResultArrayProperty->Inner;
-
-					ResultArrayHelper.Resize(0);
-					
-					if(Instruction.OpCode == ERigVMOpCode::ArrayDifference)
-					{
-						for(const TPair<uint32, int32>& Pair : HashA)
-						{
-							if(!HashB.Contains(Pair.Key))
-							{
-								int32 AddedIndex = ResultArrayHelper.AddValue();
-								URigVMMemoryStorage::CopyProperty(
-									ResultElementProperty,
-									ResultArrayHelper.GetRawPtr(AddedIndex),
-									ElementPropertyA,
-									ArrayHelperA.GetRawPtr(Pair.Value)
-								);
-							}
-						}
-						for(const TPair<uint32, int32>& Pair : HashB)
-						{
-							if(!HashA.Contains(Pair.Key))
-							{
-								int32 AddedIndex = ResultArrayHelper.AddValue();
-								URigVMMemoryStorage::CopyProperty(
-									ResultElementProperty,
-									ResultArrayHelper.GetRawPtr(AddedIndex),
-									ElementPropertyB,
-									ArrayHelperB.GetRawPtr(Pair.Value)
-								);
-							}
-						}
-					}
-					else // intersection
-					{
-						for(const TPair<uint32, int32>& Pair : HashA)
-						{
-							if(HashB.Contains(Pair.Key))
-							{
-								int32 AddedIndex = ResultArrayHelper.AddValue();
-								URigVMMemoryStorage::CopyProperty(
-									ResultElementProperty,
-									ResultArrayHelper.GetRawPtr(AddedIndex),
-									ElementPropertyA,
-									ArrayHelperA.GetRawPtr(Pair.Value)
-								);
-							}
-						}
-					}
-
-					if(DebugMemoryStorageObject->Num() > 0)
-					{
-						const FRigVMTernaryOp& Op = ByteCode.GetOpAt<FRigVMTernaryOp>(Instruction);
-						CopyOperandForDebuggingIfNeeded(Op.ArgA, ArrayHandleA);
-						CopyOperandForDebuggingIfNeeded(Op.ArgB, ArrayHandleB);
-						CopyOperandForDebuggingIfNeeded(Op.ArgC, ResultArrayHandle);
-					}
-				}
-
-				ContextPublicData.InstructionIndex++;
-				break;
-			}
-			case ERigVMOpCode::ArrayReverse:
-			{
-				FRigVMMemoryHandle& ArrayHandle = CachedMemoryHandles[FirstHandleForInstruction[ContextPublicData.InstructionIndex]]; 					
-				FScriptArrayHelper ArrayHelper(CastFieldChecked<FArrayProperty>(ArrayHandle.GetProperty()), ArrayHandle.GetData());
-				for(int32 A=0, B=ArrayHelper.Num()-1; A<B; A++, B--)
-				{
-					ArrayHelper.SwapValues(A, B);
-				}
-
-				if(DebugMemoryStorageObject->Num() > 0)
-				{
-					const FRigVMUnaryOp& Op = ByteCode.GetOpAt<FRigVMUnaryOp>(Instruction);
-					CopyOperandForDebuggingIfNeeded(Op.Arg, ArrayHandle);
-				}
-
-				ContextPublicData.InstructionIndex++;
-				break;
-			}
 			case ERigVMOpCode::InvokeEntry:
 			{
 				const FRigVMInvokeEntryOp& Op = ByteCode.GetOpAt<FRigVMInvokeEntryOp>(Instruction);
@@ -3103,120 +2482,6 @@ TArray<FString> URigVM::DumpByteCodeAsTextArray(const TArray<int32>& InInstructi
 				ResultLine = TEXT("End Block");
 				break;
 			}
-			case ERigVMOpCode::ArrayReset:
-			{
-				const FRigVMUnaryOp& Op = ByteCode.GetOpAt<FRigVMUnaryOp>(Instructions[InstructionIndex]);
-				ResultLine = FString::Printf(TEXT("Reset array %s"), *GetOperandLabel(Op.Arg, OperandFormatFunction));
-				break;
-			}
-			case ERigVMOpCode::ArrayGetNum:
-			{
-				const FRigVMBinaryOp& Op = ByteCode.GetOpAt<FRigVMBinaryOp>(Instructions[InstructionIndex]);
-				ResultLine = FString::Printf(TEXT("Get size of array %s and assign to %s"), *GetOperandLabel(Op.ArgA, OperandFormatFunction), *GetOperandLabel(Op.ArgB, OperandFormatFunction));
-				break;
-			} 
-			case ERigVMOpCode::ArraySetNum:
-			{
-				const FRigVMBinaryOp& Op = ByteCode.GetOpAt<FRigVMBinaryOp>(Instructions[InstructionIndex]);
-				ResultLine = FString::Printf(TEXT("Set size of array %s to %s"), *GetOperandLabel(Op.ArgA, OperandFormatFunction), *GetOperandLabel(Op.ArgB, OperandFormatFunction));
-				break;
-			}
-			case ERigVMOpCode::ArrayGetAtIndex:
-			{
-				const FRigVMTernaryOp& Op = ByteCode.GetOpAt<FRigVMTernaryOp>(Instructions[InstructionIndex]);
-				ResultLine = FString::Printf(TEXT("Get item of array %s at index %s and assign to %s"), *GetOperandLabel(Op.ArgA, OperandFormatFunction), *GetOperandLabel(Op.ArgB, OperandFormatFunction), *GetOperandLabel(Op.ArgC, OperandFormatFunction));
-				break;
-			}
-			case ERigVMOpCode::ArraySetAtIndex:
-			{
-				const FRigVMTernaryOp& Op = ByteCode.GetOpAt<FRigVMTernaryOp>(Instructions[InstructionIndex]);
-				ResultLine = FString::Printf(TEXT("Set item of array %s at index %s to %s"), *GetOperandLabel(Op.ArgA, OperandFormatFunction), *GetOperandLabel(Op.ArgB, OperandFormatFunction), *GetOperandLabel(Op.ArgC, OperandFormatFunction));
-				break;
-			}
-			case ERigVMOpCode::ArrayAdd:
-			{
-				const FRigVMTernaryOp& Op = ByteCode.GetOpAt<FRigVMTernaryOp>(Instructions[InstructionIndex]);
-				ResultLine = FString::Printf(TEXT("Add element %s to array %s and return index %s"),
-					*GetOperandLabel(Op.ArgB, OperandFormatFunction),
-					*GetOperandLabel(Op.ArgA, OperandFormatFunction),
-					*GetOperandLabel(Op.ArgC, OperandFormatFunction));
-				break;
-			}
-			case ERigVMOpCode::ArrayInsert:
-			{
-				const FRigVMTernaryOp& Op = ByteCode.GetOpAt<FRigVMTernaryOp>(Instructions[InstructionIndex]);
-				ResultLine = FString::Printf(TEXT("Insert element %s to array %s at index %s"), *GetOperandLabel(Op.ArgC, OperandFormatFunction), *GetOperandLabel(Op.ArgA, OperandFormatFunction), *GetOperandLabel(Op.ArgB, OperandFormatFunction));
-				break;
-			}
-			case ERigVMOpCode::ArrayRemove:
-			{
-				const FRigVMBinaryOp& Op = ByteCode.GetOpAt<FRigVMBinaryOp>(Instructions[InstructionIndex]);
-				ResultLine = FString::Printf(TEXT("Remove element at index %s from array %s"), *GetOperandLabel(Op.ArgB, OperandFormatFunction), *GetOperandLabel(Op.ArgA, OperandFormatFunction));
-				break;
-			}
-			case ERigVMOpCode::ArrayFind:
-			{
-				const FRigVMQuaternaryOp& Op = ByteCode.GetOpAt<FRigVMQuaternaryOp>(Instructions[InstructionIndex]);
-				ResultLine = FString::Printf(TEXT("Find element %s in array %s and returns index %s and if element was found %s"),
-					*GetOperandLabel(Op.ArgB, OperandFormatFunction),
-					*GetOperandLabel(Op.ArgA, OperandFormatFunction),
-					*GetOperandLabel(Op.ArgC, OperandFormatFunction),
-					*GetOperandLabel(Op.ArgD, OperandFormatFunction));
-				break;
-			}
-			case ERigVMOpCode::ArrayAppend:
-			{
-				const FRigVMBinaryOp& Op = ByteCode.GetOpAt<FRigVMBinaryOp>(Instructions[InstructionIndex]);
-				ResultLine = FString::Printf(TEXT("Append array %s to array %s"), *GetOperandLabel(Op.ArgB, OperandFormatFunction), *GetOperandLabel(Op.ArgA, OperandFormatFunction));
-				break;
-			}
-			case ERigVMOpCode::ArrayClone:
-			{
-				const FRigVMBinaryOp& Op = ByteCode.GetOpAt<FRigVMBinaryOp>(Instructions[InstructionIndex]);
-				ResultLine = FString::Printf(TEXT("Clone array %s to array %s"), *GetOperandLabel(Op.ArgA, OperandFormatFunction), *GetOperandLabel(Op.ArgB, OperandFormatFunction));
-				break;
-			}
-			case ERigVMOpCode::ArrayIterator:
-			{
-				const FRigVMSenaryOp& Op = ByteCode.GetOpAt<FRigVMSenaryOp>(Instructions[InstructionIndex]);
-				ResultLine = FString::Printf(TEXT("Iterate over array %s, with current element in %s, current index in %s, array count in %s and current ratio in %s"),
-					*GetOperandLabel(Op.ArgA, OperandFormatFunction),
-					*GetOperandLabel(Op.ArgB, OperandFormatFunction),
-					*GetOperandLabel(Op.ArgC, OperandFormatFunction),
-					*GetOperandLabel(Op.ArgD, OperandFormatFunction),
-					*GetOperandLabel(Op.ArgE, OperandFormatFunction));
-				break;
-			}
-			case ERigVMOpCode::ArrayUnion:
-			{
-				const FRigVMBinaryOp& Op = ByteCode.GetOpAt<FRigVMBinaryOp>(Instructions[InstructionIndex]);
-				ResultLine = FString::Printf(TEXT("Merge array %s and array %s"), *GetOperandLabel(Op.ArgA, OperandFormatFunction), *GetOperandLabel(Op.ArgB, OperandFormatFunction));
-				break;
-			}
-			case ERigVMOpCode::ArrayDifference:
-			{
-				const FRigVMTernaryOp& Op = ByteCode.GetOpAt<FRigVMTernaryOp>(Instructions[InstructionIndex]);
-				ResultLine = FString::Printf(TEXT("Create array %s from differences of array %s and array %s"),
-					*GetOperandLabel(Op.ArgC, OperandFormatFunction),
-					*GetOperandLabel(Op.ArgA, OperandFormatFunction),
-					*GetOperandLabel(Op.ArgB, OperandFormatFunction));
-				break;
-			}
-			case ERigVMOpCode::ArrayIntersection:
-			{
-				const FRigVMTernaryOp& Op = ByteCode.GetOpAt<FRigVMTernaryOp>(Instructions[InstructionIndex]);
-				ResultLine = FString::Printf(TEXT("Create array %s from intersection of array %s and array %s"),
-					*GetOperandLabel(Op.ArgC, OperandFormatFunction), 
-					*GetOperandLabel(Op.ArgA, OperandFormatFunction),
-					*GetOperandLabel(Op.ArgB, OperandFormatFunction));
-				break;
-			}
-			case ERigVMOpCode::ArrayReverse:
-			{
-				const FRigVMUnaryOp& Op = ByteCode.GetOpAt<FRigVMUnaryOp>(Instructions[InstructionIndex]);
-				ResultLine = FString::Printf(TEXT("Reverse array %s"), *GetOperandLabel(Op.Arg, OperandFormatFunction));
-				break;
-			}
 			case ERigVMOpCode::InvokeEntry:
 			{
 				const FRigVMInvokeEntryOp& Op = ByteCode.GetOpAt<FRigVMInvokeEntryOp>(Instructions[InstructionIndex]);
@@ -3472,26 +2737,6 @@ void URigVM::RefreshExternalPropertyPaths()
 	}
 }
 
-void URigVM::CopyArray(FScriptArrayHelper& TargetHelper, FRigVMMemoryHandle& TargetHandle,
-	FScriptArrayHelper& SourceHelper, FRigVMMemoryHandle& SourceHandle)
-{
-	const FArrayProperty* TargetArrayProperty = CastFieldChecked<FArrayProperty>(TargetHandle.GetProperty());
-	const FArrayProperty* SourceArrayProperty = CastFieldChecked<FArrayProperty>(SourceHandle.GetProperty());
-
-	TargetHelper.Resize(SourceHelper.Num());
-	if(SourceHelper.Num() > 0)
-	{
-		const FProperty* TargetProperty = TargetArrayProperty->Inner;
-		const FProperty* SourceProperty = SourceArrayProperty->Inner;
-		for(int32 ElementIndex = 0; ElementIndex < SourceHelper.Num(); ElementIndex++)
-		{
-			uint8* TargetMemory = TargetHelper.GetRawPtr(ElementIndex);
-			const uint8* SourceMemory = SourceHelper.GetRawPtr(ElementIndex);
-			URigVMMemoryStorage::CopyProperty(TargetProperty, TargetMemory, SourceProperty, SourceMemory);
-		}
-	}
-}
-
 void URigVM::SetupInstructionTracking(int32 InInstructionCount)
 {
 #if WITH_EDITOR
diff --git a/Engine/Source/Runtime/RigVM/Private/RigVMCore/RigVMByteCode.cpp b/Engine/Source/Runtime/RigVM/Private/RigVMCore/RigVMByteCode.cpp
--- a/Engine/Source/Runtime/RigVM/Private/RigVMCore/RigVMByteCode.cpp
+++ b/Engine/Source/Runtime/RigVM/Private/RigVMCore/RigVMByteCode.cpp
@@ -1410,8 +1410,6 @@ FString FRigVMByteCode::DumpToText() const
 			case ERigVMOpCode::BoolTrue:
 			case ERigVMOpCode::Increment:
 			case ERigVMOpCode::Decrement:
-			case ERigVMOpCode::ArrayReset:
-			case ERigVMOpCode::ArrayReverse:
 			{
 				const FRigVMUnaryOp& Op = GetOpAt<FRigVMUnaryOp>(Instruction.ByteCodeIndex);
 				FString ArgContent;
@@ -1459,12 +1457,6 @@ FString FRigVMByteCode::DumpToText() const
 				break;
 			}
 			case ERigVMOpCode::BeginBlock:
-			case ERigVMOpCode::ArrayGetNum:
-			case ERigVMOpCode::ArraySetNum:
-			case ERigVMOpCode::ArrayAppend:
-			case ERigVMOpCode::ArrayClone:
-			case ERigVMOpCode::ArrayRemove:
-			case ERigVMOpCode::ArrayUnion:
 			{
 				const FRigVMBinaryOp& Op = GetOpAt<FRigVMBinaryOp>(Instruction.ByteCodeIndex);
 				FString ArgA, ArgB;
@@ -1474,55 +1466,6 @@ FString FRigVMByteCode::DumpToText() const
 				Line += FString::Printf(TEXT(", ArgB %s"), *ArgB);
 				break;
 			}
-			case ERigVMOpCode::ArrayAdd:
-			case ERigVMOpCode::ArrayGetAtIndex:
-			case ERigVMOpCode::ArraySetAtIndex:
-			case ERigVMOpCode::ArrayInsert:
-			case ERigVMOpCode::ArrayDifference:
-			case ERigVMOpCode::ArrayIntersection:
-			{
-				const FRigVMTernaryOp& Op = GetOpAt<FRigVMTernaryOp>(Instruction.ByteCodeIndex);
-				FString ArgA, ArgB, ArgC;
-				FRigVMOperand::StaticStruct()->ExportText(ArgA, &Op.ArgA, &Op.ArgA, nullptr, PPF_None, nullptr);
-				FRigVMOperand::StaticStruct()->ExportText(ArgB, &Op.ArgB, &Op.ArgB, nullptr, PPF_None, nullptr);
-				FRigVMOperand::StaticStruct()->ExportText(ArgC, &Op.ArgC, &Op.ArgC, nullptr, PPF_None, nullptr);
-				Line += FString::Printf(TEXT(", ArgA %s"), *ArgA);
-				Line += FString::Printf(TEXT(", ArgB %s"), *ArgB);
-				Line += FString::Printf(TEXT(", ArgC %s"), *ArgC);
-				break;
-			}
-			case ERigVMOpCode::ArrayFind:
-			{
-				const FRigVMQuaternaryOp& Op = GetOpAt<FRigVMQuaternaryOp>(Instruction.ByteCodeIndex);
-				FString ArgA, ArgB, ArgC, ArgD;
-				FRigVMOperand::StaticStruct()->ExportText(ArgA, &Op.ArgA, &Op.ArgA, nullptr, PPF_None, nullptr);
-				FRigVMOperand::StaticStruct()->ExportText(ArgB, &Op.ArgB, &Op.ArgB, nullptr, PPF_None, nullptr);
-				FRigVMOperand::StaticStruct()->ExportText(ArgC, &Op.ArgC, &Op.ArgC, nullptr, PPF_None, nullptr);
-				FRigVMOperand::StaticStruct()->ExportText(ArgD, &Op.ArgD, &Op.ArgD, nullptr, PPF_None, nullptr);
-				Line += FString::Printf(TEXT(", ArgA %s"), *ArgA);
-				Line += FString::Printf(TEXT(", ArgB %s"), *ArgB);
-				Line += FString::Printf(TEXT(", ArgC %s"), *ArgC);
-				Line += FString::Printf(TEXT(", ArgD %s"), *ArgD);
-				break;
-			}
-			case ERigVMOpCode::ArrayIterator:
-			{
-				const FRigVMSenaryOp& Op = GetOpAt<FRigVMSenaryOp>(Instruction.ByteCodeIndex);
-				FString ArgA, ArgB, ArgC, ArgD, ArgE, ArgF;
-				FRigVMOperand::StaticStruct()->ExportText(ArgA, &Op.ArgA, &Op.ArgA, nullptr, PPF_None, nullptr);
-				FRigVMOperand::StaticStruct()->ExportText(ArgB, &Op.ArgB, &Op.ArgB, nullptr, PPF_None, nullptr);
-				FRigVMOperand::StaticStruct()->ExportText(ArgC, &Op.ArgC, &Op.ArgC, nullptr, PPF_None, nullptr);
-				FRigVMOperand::StaticStruct()->ExportText(ArgD, &Op.ArgD, &Op.ArgD, nullptr, PPF_None, nullptr);
-				FRigVMOperand::StaticStruct()->ExportText(ArgE, &Op.ArgE, &Op.ArgE, nullptr, PPF_None, nullptr);
-				FRigVMOperand::StaticStruct()->ExportText(ArgF, &Op.ArgF, &Op.ArgF, nullptr, PPF_None, nullptr);
-				Line += FString::Printf(TEXT(", ArgA %s"), *ArgA);
-				Line += FString::Printf(TEXT(", ArgB %s"), *ArgB);
-				Line += FString::Printf(TEXT(", ArgC %s"), *ArgC);
-				Line += FString::Printf(TEXT(", ArgD %s"), *ArgD);
-				Line += FString::Printf(TEXT(", ArgE %s"), *ArgE);
-				Line += FString::Printf(TEXT(", ArgF %s"), *ArgF);
-				break;
-			}
 			case ERigVMOpCode::InvokeEntry:
 			{
 				const FRigVMInvokeEntryOp& Op = GetOpAt<FRigVMInvokeEntryOp>(Instruction.ByteCodeIndex);
@@ -1584,86 +1527,6 @@ uint64 FRigVMByteCode::AddEndBlockOp()
 	return AddOp(Op);
 }
 
-uint64 FRigVMByteCode::AddArrayResetOp(FRigVMOperand InArrayArg)
-{
-	return AddOp(FRigVMUnaryOp(ERigVMOpCode::ArrayReset, InArrayArg));
-}
-
-uint64 FRigVMByteCode::AddArrayGetNumOp(FRigVMOperand InArrayArg, FRigVMOperand InNumArg)
-{
-	return AddOp(FRigVMBinaryOp(ERigVMOpCode::ArrayGetNum, InArrayArg, InNumArg));
-}
-
-uint64 FRigVMByteCode::AddArraySetNumOp(FRigVMOperand InArrayArg, FRigVMOperand InNumArg)
-{
-	return AddOp(FRigVMBinaryOp(ERigVMOpCode::ArraySetNum, InArrayArg, InNumArg));
-}
-
-uint64 FRigVMByteCode::AddArrayGetAtIndexOp(FRigVMOperand InArrayArg, FRigVMOperand InIndexArg, FRigVMOperand InElementArg)
-{
-	return AddOp(FRigVMTernaryOp(ERigVMOpCode::ArrayGetAtIndex, InArrayArg, InIndexArg, InElementArg));
-}
-
-uint64 FRigVMByteCode::AddArraySetAtIndexOp(FRigVMOperand InArrayArg, FRigVMOperand InIndexArg, FRigVMOperand InElementArg)
-{
-	return AddOp(FRigVMTernaryOp(ERigVMOpCode::ArraySetAtIndex, InArrayArg, InIndexArg, InElementArg));
-}
-
-uint64 FRigVMByteCode::AddArrayAddOp(FRigVMOperand InArrayArg, FRigVMOperand InElementArg, FRigVMOperand InIndexArg)
-{
-	return AddOp(FRigVMTernaryOp(ERigVMOpCode::ArrayAdd, InArrayArg, InElementArg, InIndexArg));	
-}
-
-uint64 FRigVMByteCode::AddArrayInsertOp(FRigVMOperand InArrayArg, FRigVMOperand InIndexArg, FRigVMOperand InElementArg)
-{
-	return AddOp(FRigVMTernaryOp(ERigVMOpCode::ArrayInsert, InArrayArg, InIndexArg, InElementArg));
-}
-
-uint64 FRigVMByteCode::AddArrayRemoveOp(FRigVMOperand InArrayArg, FRigVMOperand InIndexArg)
-{
-	return AddOp(FRigVMBinaryOp(ERigVMOpCode::ArrayRemove, InArrayArg, InIndexArg));
-}
-
-uint64 FRigVMByteCode::AddArrayFindOp(FRigVMOperand InArrayArg, FRigVMOperand InElementArg, FRigVMOperand InIndexArg, FRigVMOperand InSuccessArg)
-{
-	return AddOp(FRigVMQuaternaryOp(ERigVMOpCode::ArrayFind, InArrayArg, InElementArg, InIndexArg, InSuccessArg));	
-}
-
-uint64 FRigVMByteCode::AddArrayAppendOp(FRigVMOperand InArrayArg, FRigVMOperand InOtherArrayArg)
-{
-	return AddOp(FRigVMBinaryOp(ERigVMOpCode::ArrayAppend, InArrayArg, InOtherArrayArg));
-}
-
-uint64 FRigVMByteCode::AddArrayCloneOp(FRigVMOperand InArrayArg, FRigVMOperand InClonedArrayArg)
-{
-	return AddOp(FRigVMBinaryOp(ERigVMOpCode::ArrayClone, InArrayArg, InClonedArrayArg));
-}
-
-uint64 FRigVMByteCode::AddArrayIteratorOp(FRigVMOperand InArrayArg, FRigVMOperand InElementArg, FRigVMOperand InIndexArg, FRigVMOperand InCountArg, FRigVMOperand InRatioArg, FRigVMOperand InContinueArg)
-{
-	return AddOp(FRigVMSenaryOp(ERigVMOpCode::ArrayIterator, InArrayArg, InElementArg, InIndexArg, InCountArg, InRatioArg, InContinueArg));
-}
-
-uint64 FRigVMByteCode::AddArrayUnionOp(FRigVMOperand InArrayArg, FRigVMOperand InOtherArrayArg)
-{
-	return AddOp(FRigVMBinaryOp(ERigVMOpCode::ArrayUnion, InArrayArg, InOtherArrayArg));
-}
-	
-uint64 FRigVMByteCode::AddArrayDifferenceOp(FRigVMOperand InArrayArg, FRigVMOperand InOtherArrayArg, FRigVMOperand InResultArrayArg)
-{
-	return AddOp(FRigVMTernaryOp(ERigVMOpCode::ArrayDifference, InArrayArg, InOtherArrayArg, InResultArrayArg));
-}
-	
-uint64 FRigVMByteCode::AddArrayIntersectionOp(FRigVMOperand InArrayArg, FRigVMOperand InOtherArrayArg, FRigVMOperand InResultArrayArg)
-{
-	return AddOp(FRigVMTernaryOp(ERigVMOpCode::ArrayIntersection, InArrayArg, InOtherArrayArg, InResultArrayArg));
-}
-	
-uint64 FRigVMByteCode::AddArrayReverseOp(FRigVMOperand InArrayArg)
-{
-	return AddOp(FRigVMUnaryOp(ERigVMOpCode::ArrayReverse, InArrayArg));
-}
-
 uint64 FRigVMByteCode::AddInvokeEntryOp(const FName& InEntryName)
 {
 	return AddOp(FRigVMInvokeEntryOp(InEntryName));
diff --git a/Engine/Source/Runtime/RigVM/Private/RigVMCore/RigVMRegistry.cpp b/Engine/Source/Runtime/RigVM/Private/RigVMCore/RigVMRegistry.cpp
--- a/Engine/Source/Runtime/RigVM/Private/RigVMCore/RigVMRegistry.cpp
+++ b/Engine/Source/Runtime/RigVM/Private/RigVMCore/RigVMRegistry.cpp
@@ -1465,6 +1465,28 @@ FRigVMDispatchFactory* FRigVMRegistry::FindOrAddDispatchFactory(UScriptStruct* I
 	return (FRigVMDispatchFactory*)RegisterFactory(InFactoryStruct);
 }
 
+FString FRigVMRegistry::FindOrAddSingletonDispatchFunction(UScriptStruct* InFactoryStruct)
+{
+	if(const FRigVMDispatchFactory* Factory = FindOrAddDispatchFactory(InFactoryStruct))
+	{
+		if(Factory->IsSingleton())
+		{
+			if(const FRigVMTemplate* Template = Factory->GetTemplate())
+			{
+				// use the types for the first permutation - since we don't care
+				// for a singleton dispatch
+				const FRigVMTemplateTypeMap TypesForPrimaryPermutation = Template->GetTypesForPermutation(0);
+				const FString Name = Factory->GetPermutationName(TypesForPrimaryPermutation);
+				if(const FRigVMFunction* Function = FindFunction(*Name))
+				{
+					return Function->Name;
+				}
+			}
+		}
+	}
+	return FString();
+}
+
 const TArray<FRigVMDispatchFactory*>& FRigVMRegistry::GetFactories() const
 {
 	return Factories;
diff --git a/Engine/Source/Runtime/RigVM/Private/RigVMFunctions/RigVMDispatch_Array.cpp b/Engine/Source/Runtime/RigVM/Private/RigVMFunctions/RigVMDispatch_Array.cpp
new file mode 100644
--- /dev/null
+++ b/Engine/Source/Runtime/RigVM/Private/RigVMFunctions/RigVMDispatch_Array.cpp
@@ -0,0 +1,1389 @@
+// Copyright Epic Games, Inc. All Rights Reserved.
+
+#include "RigVMFunctions/RigVMDispatch_Array.h"
+#include "RigVMCore/RigVMRegistry.h"
+#include "RigVMCore/RigVMStruct.h"
+#include "RigVMStringUtils.h"
+
+#include UE_INLINE_GENERATED_CPP_BY_NAME(RigVMDispatch_Array)
+#define LOCTEXT_NAMESPACE "RigVMDispatch_Array"
+
+const FName& FRigVMDispatch_ArrayBase::ExecuteName = FRigVMStruct::ExecuteContextName;
+const FName FRigVMDispatch_ArrayBase::ArrayName = TEXT("Array");
+const FName FRigVMDispatch_ArrayBase::ValuesName = TEXT("Values");
+const FName FRigVMDispatch_ArrayBase::NumName = TEXT("Num");
+const FName FRigVMDispatch_ArrayBase::IndexName = TEXT("Index");
+const FName FRigVMDispatch_ArrayBase::ElementName = TEXT("Element");
+const FName FRigVMDispatch_ArrayBase::SuccessName = TEXT("Success");
+const FName FRigVMDispatch_ArrayBase::OtherName = TEXT("Other");
+const FName FRigVMDispatch_ArrayBase::CloneName = TEXT("Clone");
+const FName FRigVMDispatch_ArrayBase::CountName = TEXT("Count");
+const FName FRigVMDispatch_ArrayBase::RatioName = TEXT("Ratio");
+const FName FRigVMDispatch_ArrayBase::ResultName = TEXT("Result");
+const FName& FRigVMDispatch_ArrayBase::CompletedName = FRigVMStruct::ControlFlowCompletedName;
+
+////////////////////////////////////////////////////////////////////////////////////////////////////////////////////////
+///
+#if WITH_EDITOR
+
+UScriptStruct* FRigVMDispatch_ArrayBase::GetFactoryDispatchForOpCode(ERigVMOpCode InOpCode)
+{
+	switch(InOpCode)
+	{
+		case ERigVMOpCode::ArrayReset:
+		{
+			return FRigVMDispatch_ArrayReset::StaticStruct();
+		}
+		case ERigVMOpCode::ArrayGetNum:
+		{
+			return FRigVMDispatch_ArrayGetNum::StaticStruct();
+		} 
+		case ERigVMOpCode::ArraySetNum:
+		{
+			return FRigVMDispatch_ArraySetNum::StaticStruct();
+		}
+		case ERigVMOpCode::ArrayGetAtIndex:
+		{
+			return FRigVMDispatch_ArrayGetAtIndex::StaticStruct();
+		}  
+		case ERigVMOpCode::ArraySetAtIndex:
+		{
+			return FRigVMDispatch_ArraySetAtIndex::StaticStruct();
+		}
+		case ERigVMOpCode::ArrayAdd:
+		{
+			return FRigVMDispatch_ArrayAdd::StaticStruct();
+		}
+		case ERigVMOpCode::ArrayInsert:
+		{
+			return FRigVMDispatch_ArrayInsert::StaticStruct();
+		}
+		case ERigVMOpCode::ArrayRemove:
+		{
+			return FRigVMDispatch_ArrayRemove::StaticStruct();
+		}
+		case ERigVMOpCode::ArrayFind:
+		{
+			return FRigVMDispatch_ArrayFind::StaticStruct();
+		}
+		case ERigVMOpCode::ArrayAppend:
+		{
+			return FRigVMDispatch_ArrayAppend::StaticStruct();
+		}
+		case ERigVMOpCode::ArrayClone:
+		{
+			return FRigVMDispatch_ArrayClone::StaticStruct();
+		}
+		case ERigVMOpCode::ArrayIterator:
+		{
+			return FRigVMDispatch_ArrayIterator::StaticStruct();
+		}
+		case ERigVMOpCode::ArrayUnion:
+		{
+			return FRigVMDispatch_ArrayUnion::StaticStruct();
+		}
+		case ERigVMOpCode::ArrayDifference:
+		{
+			return FRigVMDispatch_ArrayDifference::StaticStruct();
+		}
+		case ERigVMOpCode::ArrayIntersection:
+		{
+			return FRigVMDispatch_ArrayIntersection::StaticStruct();
+		}
+		case ERigVMOpCode::ArrayReverse:
+		{
+			return FRigVMDispatch_ArrayReverse::StaticStruct();
+		}
+		default:
+		{
+			break;
+		}
+	}
+	return nullptr;
+}
+
+FName FRigVMDispatch_ArrayBase::GetFactoryNameForOpCode(ERigVMOpCode InOpCode)
+{
+	if(const UScriptStruct* FactoryStruct = GetFactoryDispatchForOpCode(InOpCode))
+	{
+		return *(DispatchPrefix + FactoryStruct->GetName());
+	}
+	return NAME_None;
+}
+
+FString FRigVMDispatch_ArrayBase::GetArgumentDefaultValue(const FName& InArgumentName,
+                                                          TRigVMTypeIndex InTypeIndex) const
+{
+	if(InArgumentName == NumName || 
+		InArgumentName == IndexName || 
+		InArgumentName == CountName)
+	{
+		static const FString ZeroValue = TEXT("0");
+		return ZeroValue;
+	}
+
+	return FRigVMDispatch_CoreBase::GetArgumentDefaultValue(InArgumentName, InTypeIndex);
+}
+
+#endif
+
+FRigVMTemplateArgument FRigVMDispatch_ArrayBase::CreateArgument(const FName& InName, ERigVMPinDirection InDirection)
+{
+	static const TArray<FRigVMTemplateArgument::ETypeCategory> ArrayCategories = {
+		FRigVMTemplateArgument::ETypeCategory_ArrayAnyValue
+	};
+
+	static const TArray<FRigVMTemplateArgument::ETypeCategory> ElementCategories = {
+		FRigVMTemplateArgument::ETypeCategory_SingleAnyValue
+	};
+
+	if(InName == ArrayName || InName == ValuesName || InName == OtherName || InName == CloneName || InName == ResultName)
+	{
+		return FRigVMTemplateArgument(InName, InDirection, ArrayCategories);
+	}
+	if(InName == ElementName)
+	{
+		return FRigVMTemplateArgument(InName, InDirection, ElementCategories);
+	}
+	if(InName == NumName || InName == IndexName || InName == CountName)
+	{
+		return FRigVMTemplateArgument(InName, InDirection, RigVMTypeUtils::TypeIndex::Int32);
+	}
+	if(InName == SuccessName)
+	{
+		return FRigVMTemplateArgument(InName, InDirection, RigVMTypeUtils::TypeIndex::Bool);
+	}
+	if(InName == RatioName)
+	{
+		return FRigVMTemplateArgument(InName, InDirection, RigVMTypeUtils::TypeIndex::Float);
+	}
+
+	checkNoEntry();
+	return FRigVMTemplateArgument();
+}
+
+TMap<uint32, int32> FRigVMDispatch_ArrayBase::GetArrayHash(FScriptArrayHelper& InArrayHelper, const FArrayProperty* InArrayProperty)
+{
+	TMap<uint32, int32> Hash;
+	Hash.Reserve(InArrayHelper.Num());
+
+	const FProperty* ElementProperty = InArrayProperty->Inner;
+	
+	for(int32 Index = 0; Index < InArrayHelper.Num(); Index++)
+	{
+		uint32 HashValue;
+		if (ElementProperty->PropertyFlags & CPF_HasGetValueTypeHash)
+		{
+			HashValue = ElementProperty->GetValueTypeHash(InArrayHelper.GetRawPtr(Index));
+		}
+		else
+		{
+			FString Value;
+			ElementProperty->ExportTextItem_Direct(Value, InArrayHelper.GetRawPtr(Index), nullptr, nullptr, PPF_None);
+			HashValue = TextKeyUtil::HashString(Value);
+		}
+					
+		if(!Hash.Contains(HashValue))
+		{
+			Hash.Add(HashValue, Index);
+		}
+	}
+
+	return Hash;
+}
+
+////////////////////////////////////////////////////////////////////////////////////////////////////////////////////////
+
+TArray<FRigVMExecuteArgument> FRigVMDispatch_ArrayBaseMutable::GetExecuteArguments_Impl(
+	const FRigVMDispatchContext& InContext) const
+{
+	return {
+			{ExecuteName, ERigVMPinDirection::IO}
+	};
+}
+
+////////////////////////////////////////////////////////////////////////////////////////////////////////////////////////
+
+TArray<FRigVMTemplateArgument> FRigVMDispatch_ArrayMake::GetArguments() const
+{
+	return {
+		CreateArgument(ValuesName, ERigVMPinDirection::Input),
+		CreateArgument(ArrayName, ERigVMPinDirection::Output)
+	};
+}
+
+FRigVMTemplateTypeMap FRigVMDispatch_ArrayMake::OnNewArgumentType(const FName& InArgumentName, TRigVMTypeIndex InTypeIndex) const
+{
+	check(InArgumentName == ValuesName || InArgumentName == ArrayName);
+	return {
+		{ValuesName, InTypeIndex},
+		{ArrayName, InTypeIndex}
+	};
+}
+
+#if WITH_EDITOR
+
+FText FRigVMDispatch_ArrayMake::GetNodeTooltip(const FRigVMTemplateTypeMap& InTypes) const
+{
+	return LOCTEXT("ArrayMakeToolTip", "Creates a new array from its elements.");
+}
+
+FText FRigVMDispatch_ArrayMake::GetArgumentTooltip(const FName& InArgumentName, TRigVMTypeIndex InTypeIndex) const
+{
+	if(ValuesName == ArrayName)
+	{
+		return LOCTEXT("ArrayMake_ValuesArgumentToolTip", "The elements of the array to create.");
+	}
+	if(InArgumentName == ArrayName)
+	{
+		return LOCTEXT("ArrayMake_ArrayArgumentToolTip", "The resulting array.");
+	}
+	return FText();
+}
+
+FString FRigVMDispatch_ArrayMake::GetArgumentMetaData(const FName& InArgumentName, const FName& InMetaDataKey) const
+{
+	if(InArgumentName == ValuesName && InMetaDataKey == FRigVMStruct::FixedSizeArrayMetaName)
+	{
+		return TrueString;
+	}
+	return FRigVMDispatch_ArrayBase::GetArgumentMetaData(InArgumentName, InMetaDataKey);
+}
+
+FString FRigVMDispatch_ArrayMake::GetArgumentDefaultValue(const FName& InArgumentName,
+	TRigVMTypeIndex InTypeIndex) const
+{
+	if(InArgumentName == ValuesName)
+	{
+		const TRigVMTypeIndex ElementTypeIndex = FRigVMRegistry::Get().GetBaseTypeFromArrayTypeIndex(InTypeIndex);
+		FString ElementDefaultValue = FRigVMDispatch_ArrayBase::GetArgumentDefaultValue(InArgumentName, ElementTypeIndex);
+		if(ElementDefaultValue.IsEmpty())
+		{
+			static const FString EmptyBracesString = TEXT("()");
+			ElementDefaultValue = EmptyBracesString;
+		}
+
+		// default to one element
+		return RigVMStringUtils::JoinDefaultValue({ElementDefaultValue});
+	}
+
+	return FRigVMDispatch_ArrayBase::GetArgumentDefaultValue(InArgumentName, InTypeIndex);
+}
+
+#endif
+
+FName FRigVMDispatch_ArrayMake::GetArgumentNameForOperandIndex(int32 InOperandIndex, int32 InTotalOperands) const
+{
+	if(InOperandIndex < InTotalOperands - 1)
+	{
+		return FRigVMBranchInfo::GetFixedArrayLabel(ValuesName, *FString::FromInt(InOperandIndex - 1));
+	}
+	return FRigVMDispatch_ArrayBase::GetArgumentNameForOperandIndex(1, 2);
+}
+
+void FRigVMDispatch_ArrayMake::Execute(FRigVMExtendedExecuteContext& InContext, FRigVMMemoryHandleArray Handles)
+{
+	const FArrayProperty* ArrayProperty = CastFieldChecked<FArrayProperty>(Handles.Last().GetProperty());
+	FScriptArrayHelper ArrayHelper(ArrayProperty, Handles.Last().GetData());
+
+	const int32 Num = Handles.Num() - 1;
+	if(InContext.IsValidArraySize(Num))
+	{
+		ArrayHelper.Resize(Num);
+
+		for(int32 Index = 0; Index < Num; Index++)
+		{
+			const uint8* SourceMemory = Handles[Index].GetData();
+			uint8* TargetMemory = ArrayHelper.GetRawPtr(Index);
+			URigVMMemoryStorage::CopyProperty(ArrayProperty->Inner, TargetMemory, Handles[Index].GetProperty(), SourceMemory);
+		}
+	}
+	else
+	{
+		ArrayHelper.Resize(0);
+	}
+}
+
+////////////////////////////////////////////////////////////////////////////////////////////////////////////////////////
+
+TArray<FRigVMTemplateArgument> FRigVMDispatch_ArrayReset::GetArguments() const
+{
+	return {
+		CreateArgument(ArrayName, ERigVMPinDirection::IO)
+	};
+}
+
+FRigVMTemplateTypeMap FRigVMDispatch_ArrayReset::OnNewArgumentType(const FName& InArgumentName, TRigVMTypeIndex InTypeIndex) const
+{
+	check(InArgumentName == ArrayName);
+	return {
+		{ArrayName, InTypeIndex}
+	};
+}
+
+void FRigVMDispatch_ArrayReset::Execute(FRigVMExtendedExecuteContext& InContext, FRigVMMemoryHandleArray Handles)
+{
+	const FArrayProperty* ArrayProperty = CastFieldChecked<FArrayProperty>(Handles[0].GetProperty());
+	FScriptArrayHelper ArrayHelper(ArrayProperty, Handles[0].GetData());
+	ArrayHelper.Resize(0);
+}
+
+#if WITH_EDITOR
+
+FText FRigVMDispatch_ArrayReset::GetNodeTooltip(const FRigVMTemplateTypeMap& InTypes) const
+{
+	return LOCTEXT("ArrayResetToolTip", "Removes all elements from an array.\nModifies the input array.");
+}
+
+FText FRigVMDispatch_ArrayReset::GetArgumentTooltip(const FName& InArgumentName, TRigVMTypeIndex InTypeIndex) const
+{
+	if(InArgumentName == ArrayName)
+	{
+		return LOCTEXT("ArrayReset_ArrayArgumentToolTip", "The array to be cleared.");
+	}
+	return FText();
+}
+
+#endif
+
+////////////////////////////////////////////////////////////////////////////////////////////////////////////////////////
+
+TArray<FRigVMTemplateArgument> FRigVMDispatch_ArrayGetNum::GetArguments() const
+{
+	return {
+		CreateArgument(ArrayName, ERigVMPinDirection::Input),
+		CreateArgument(NumName, ERigVMPinDirection::Output)
+	};
+}
+
+FRigVMTemplateTypeMap FRigVMDispatch_ArrayGetNum::OnNewArgumentType(const FName& InArgumentName, TRigVMTypeIndex InTypeIndex) const
+{
+	check(InArgumentName == ArrayName);
+	return {
+		{ArrayName, InTypeIndex},
+		{NumName, RigVMTypeUtils::TypeIndex::Int32}
+	};
+}
+
+void FRigVMDispatch_ArrayGetNum::Execute(FRigVMExtendedExecuteContext& InContext, FRigVMMemoryHandleArray Handles)
+{
+	const FArrayProperty* ArrayProperty = CastFieldChecked<FArrayProperty>(Handles[0].GetProperty());
+	const FScriptArrayHelper ArrayHelper(ArrayProperty, Handles[0].GetData());
+	int32* Num = (int32*)Handles[1].GetData();
+	*Num = ArrayHelper.Num();
+}
+
+#if WITH_EDITOR
+
+FText FRigVMDispatch_ArrayGetNum::GetNodeTooltip(const FRigVMTemplateTypeMap& InTypes) const
+{
+    return LOCTEXT("ArrayGetNumToolTip", "Returns the number of elements of an array");
+}
+
+FText FRigVMDispatch_ArrayGetNum::GetArgumentTooltip(const FName& InArgumentName, TRigVMTypeIndex InTypeIndex) const
+{
+	if(InArgumentName == ArrayName)
+	{
+		return LOCTEXT("ArrayGetNum_ArrayArgumentToolTip", "The array to retrieve the size for.");
+	}
+	if(InArgumentName == NumName)
+	{
+		return LOCTEXT("ArrayGetNum_NumArgumentToolTip", "The size of the input array.");
+	}
+	return FText();
+}
+
+#endif
+
+////////////////////////////////////////////////////////////////////////////////////////////////////////////////////////
+
+TArray<FRigVMTemplateArgument> FRigVMDispatch_ArraySetNum::GetArguments() const
+{
+	return {
+		CreateArgument(ArrayName, ERigVMPinDirection::IO),
+		CreateArgument(NumName, ERigVMPinDirection::Input)
+	};
+}
+
+FRigVMTemplateTypeMap FRigVMDispatch_ArraySetNum::OnNewArgumentType(const FName& InArgumentName, TRigVMTypeIndex InTypeIndex) const
+{
+	check(InArgumentName == ArrayName);
+	return {
+		{ArrayName, InTypeIndex},
+		{NumName, RigVMTypeUtils::TypeIndex::Int32}
+	};
+}
+
+void FRigVMDispatch_ArraySetNum::Execute(FRigVMExtendedExecuteContext& InContext, FRigVMMemoryHandleArray Handles)
+{
+	const FArrayProperty* ArrayProperty = CastFieldChecked<FArrayProperty>(Handles[0].GetProperty());
+	FScriptArrayHelper ArrayHelper(ArrayProperty, Handles[0].GetData());
+	const int32* Num = (int32*)Handles[1].GetData();
+	if(InContext.IsValidArraySize(*Num))
+	{
+		ArrayHelper.Resize(*Num);
+	}
+}
+
+#if WITH_EDITOR
+
+FText FRigVMDispatch_ArraySetNum::GetNodeTooltip(const FRigVMTemplateTypeMap& InTypes) const
+{
+	return LOCTEXT("ArraySetNumToolTip", "Sets the numbers of elements of an array.\nModifies the input array.");
+}
+
+FText FRigVMDispatch_ArraySetNum::GetArgumentTooltip(const FName& InArgumentName, TRigVMTypeIndex InTypeIndex) const
+{
+	if(InArgumentName == ArrayName)
+	{
+		return LOCTEXT("ArraySetNum_ArrayArgumentToolTip", "The array to set the size for.");
+	}
+	if(InArgumentName == NumName)
+	{
+		return LOCTEXT("ArraySetNum_NumArgumentToolTip", "The new size of the array.");
+	}
+	return FText();
+}
+
+#endif
+
+////////////////////////////////////////////////////////////////////////////////////////////////////////////////////////
+
+TArray<FRigVMTemplateArgument> FRigVMDispatch_ArrayGetAtIndex::GetArguments() const
+{
+	return {
+		CreateArgument(ArrayName, ERigVMPinDirection::Input),
+		CreateArgument(IndexName, ERigVMPinDirection::Input),
+		CreateArgument(ElementName, ERigVMPinDirection::Output)
+	};
+}
+
+FRigVMTemplateTypeMap FRigVMDispatch_ArrayGetAtIndex::OnNewArgumentType(const FName& InArgumentName, TRigVMTypeIndex InTypeIndex) const
+{
+	check(InArgumentName == ArrayName || InArgumentName == ElementName);
+
+	if(InArgumentName == ArrayName)
+	{
+		return {
+			{ArrayName, InTypeIndex},
+			{IndexName, RigVMTypeUtils::TypeIndex::Int32},
+			{ElementName, FRigVMRegistry::Get().GetBaseTypeFromArrayTypeIndex(InTypeIndex)}
+		};
+	}
+
+	return {
+		{ArrayName, FRigVMRegistry::Get().GetArrayTypeFromBaseTypeIndex(InTypeIndex)},
+		{IndexName, RigVMTypeUtils::TypeIndex::Int32},
+		{ElementName, InTypeIndex}
+	};
+}
+
+void FRigVMDispatch_ArrayGetAtIndex::Execute(FRigVMExtendedExecuteContext& InContext, FRigVMMemoryHandleArray Handles)
+{
+	const FArrayProperty* ArrayProperty = CastFieldChecked<FArrayProperty>(Handles[0].GetProperty());
+	FScriptArrayHelper ArrayHelper(ArrayProperty, Handles[0].GetData());
+	int32 Index = *(int32*)Handles[1].GetData();
+	uint8* TargetMemory = Handles[2].GetData();
+
+	if(InContext.IsValidArrayIndex(Index, ArrayHelper))
+	{
+		const uint8* SourceMemory = ArrayHelper.GetRawPtr(Index);
+		URigVMMemoryStorage::CopyProperty(Handles[2].GetProperty(), TargetMemory, ArrayProperty->Inner, SourceMemory);
+	}
+}
+
+#if WITH_EDITOR
+
+FText FRigVMDispatch_ArrayGetAtIndex::GetNodeTooltip(const FRigVMTemplateTypeMap& InTypes) const
+{
+	return LOCTEXT("ArrayGetAtIndexToolTip", "Returns an element of an array by index.");
+}
+
+FText FRigVMDispatch_ArrayGetAtIndex::GetArgumentTooltip(const FName& InArgumentName, TRigVMTypeIndex InTypeIndex) const
+{
+	if(InArgumentName == ArrayName)
+	{
+		return LOCTEXT("ArrayGetAtIndex_ArrayArgumentToolTip", "The array to retrieve an element from.");
+	}
+	if(InArgumentName == IndexName)
+	{
+		return LOCTEXT("ArrayGetAtIndex_IndexArgumentToolTip", "The index of the element to retrieve.");
+	}
+	if(InArgumentName == ElementName)
+	{
+		return LOCTEXT("ArrayGetAtIndex_ElementArgumentToolTip", "The element at the given index.");
+	}
+	return FText();
+}
+
+#endif
+
+////////////////////////////////////////////////////////////////////////////////////////////////////////////////////////
+
+TArray<FRigVMTemplateArgument> FRigVMDispatch_ArraySetAtIndex::GetArguments() const
+{
+	return {
+		CreateArgument(ArrayName, ERigVMPinDirection::IO),
+		CreateArgument(IndexName, ERigVMPinDirection::Input),
+		CreateArgument(ElementName, ERigVMPinDirection::Input)
+	};
+}
+
+FRigVMTemplateTypeMap FRigVMDispatch_ArraySetAtIndex::OnNewArgumentType(const FName& InArgumentName, TRigVMTypeIndex InTypeIndex) const
+{
+	check(InArgumentName == ArrayName || InArgumentName == ElementName);
+
+	if(InArgumentName == ArrayName)
+	{
+		return {
+			{ArrayName, InTypeIndex},
+			{IndexName, RigVMTypeUtils::TypeIndex::Int32},
+			{ElementName, FRigVMRegistry::Get().GetBaseTypeFromArrayTypeIndex(InTypeIndex)}
+		};
+	}
+
+	return {
+		{ArrayName, FRigVMRegistry::Get().GetArrayTypeFromBaseTypeIndex(InTypeIndex)},
+		{IndexName, RigVMTypeUtils::TypeIndex::Int32},
+		{ElementName, InTypeIndex}
+	};
+}
+
+void FRigVMDispatch_ArraySetAtIndex::Execute(FRigVMExtendedExecuteContext& InContext, FRigVMMemoryHandleArray Handles)
+{
+	const FArrayProperty* ArrayProperty = CastFieldChecked<FArrayProperty>(Handles[0].GetProperty());
+	FScriptArrayHelper ArrayHelper(ArrayProperty, Handles[0].GetData());
+	int32 Index = *(int32*)Handles[1].GetData();
+	const uint8* SourceMemory = Handles[2].GetData();
+
+	if(InContext.IsValidArrayIndex(Index, ArrayHelper))
+	{
+		uint8* TargetMemory = ArrayHelper.GetRawPtr(Index);
+		URigVMMemoryStorage::CopyProperty(ArrayProperty->Inner, TargetMemory, Handles[2].GetProperty(), SourceMemory);
+	}
+}
+
+#if WITH_EDITOR
+
+FText FRigVMDispatch_ArraySetAtIndex::GetNodeTooltip(const FRigVMTemplateTypeMap& InTypes) const
+{
+	return LOCTEXT("ArraySetAtIndexToolTip", "Sets an element of an array by index.\nModifies the input array.");
+}
+
+FText FRigVMDispatch_ArraySetAtIndex::GetArgumentTooltip(const FName& InArgumentName, TRigVMTypeIndex InTypeIndex) const
+{
+	if(InArgumentName == ArrayName)
+	{
+		return LOCTEXT("ArraySetAtIndex_ArrayArgumentToolTip", "The array to set an element for.");
+	}
+	if(InArgumentName == IndexName)
+	{
+		return LOCTEXT("ArraySetAtIndex_IndexArgumentToolTip", "The index of the element to set.");
+	}
+	if(InArgumentName == ElementName)
+	{
+		return LOCTEXT("ArraySetAtIndex_ElementArgumentToolTip", "The new value for element to set.");
+	}
+	return FText();
+}
+
+#endif
+
+////////////////////////////////////////////////////////////////////////////////////////////////////////////////////////
+
+TArray<FRigVMTemplateArgument> FRigVMDispatch_ArrayAdd::GetArguments() const
+{
+	return {
+		CreateArgument(ArrayName, ERigVMPinDirection::IO),
+		CreateArgument(ElementName, ERigVMPinDirection::Input),
+		CreateArgument(IndexName, ERigVMPinDirection::Output)
+	};
+}
+
+void FRigVMDispatch_ArrayAdd::Execute(FRigVMExtendedExecuteContext& InContext, FRigVMMemoryHandleArray Handles)
+{
+	const FArrayProperty* ArrayProperty = CastFieldChecked<FArrayProperty>(Handles[0].GetProperty());
+	FScriptArrayHelper ArrayHelper(ArrayProperty, Handles[0].GetData());
+	const uint8* SourceMemory = Handles[1].GetData();
+	int32& Index = *(int32*)Handles[2].GetData();
+
+	if(InContext.IsValidArraySize(ArrayHelper.Num() + 1))
+	{
+		Index = ArrayHelper.AddValue();
+		uint8* TargetMemory = ArrayHelper.GetRawPtr(Index);
+		URigVMMemoryStorage::CopyProperty(ArrayProperty->Inner, TargetMemory, Handles[1].GetProperty(), SourceMemory);
+	}
+	else
+	{
+		Index = INDEX_NONE;
+	}
+}
+
+#if WITH_EDITOR
+
+FText FRigVMDispatch_ArrayAdd::GetNodeTooltip(const FRigVMTemplateTypeMap& InTypes) const
+{
+	return LOCTEXT("ArrayAddToolTip", "Adds an element to an array and returns the new element's index.\nModifies the input array.");
+}
+
+FText FRigVMDispatch_ArrayAdd::GetArgumentTooltip(const FName& InArgumentName, TRigVMTypeIndex InTypeIndex) const
+{
+	if(InArgumentName == ArrayName)
+	{
+		return LOCTEXT("ArrayAdd_ArrayArgumentToolTip", "The array to add an element to.");
+	}
+	if(InArgumentName == ElementName)
+	{
+		return LOCTEXT("ArrayAdd_ElementArgumentToolTip", "The element to add to the array.");
+	}
+	if(InArgumentName == IndexName)
+	{
+		return LOCTEXT("ArrayAdd_IndexArgumentToolTip", "The index of the newly added element.");
+	}
+	return FText();
+}
+
+#endif
+
+////////////////////////////////////////////////////////////////////////////////////////////////////////////////////////
+
+void FRigVMDispatch_ArrayInsert::Execute(FRigVMExtendedExecuteContext& InContext, FRigVMMemoryHandleArray Handles)
+{
+	const FArrayProperty* ArrayProperty = CastFieldChecked<FArrayProperty>(Handles[0].GetProperty());
+	FScriptArrayHelper ArrayHelper(ArrayProperty, Handles[0].GetData());
+	int32 Index = *(int32*)Handles[1].GetData();
+	const uint8* SourceMemory = Handles[2].GetData();
+
+	if(InContext.IsValidArraySize(ArrayHelper.Num() + 1))
+	{
+		// we support wrapping the index around similar to python
+		if(Index < 0)
+		{
+			Index = ArrayHelper.Num() + Index;
+		}
+		Index = FMath::Clamp<int32>(Index, 0, ArrayHelper.Num());
+		
+		ArrayHelper.InsertValues(Index, 1);
+
+		uint8* TargetMemory = ArrayHelper.GetRawPtr(Index);
+		URigVMMemoryStorage::CopyProperty(ArrayProperty->Inner, TargetMemory, Handles[2].GetProperty(), SourceMemory);
+	}
+}
+
+#if WITH_EDITOR
+
+FText FRigVMDispatch_ArrayInsert::GetNodeTooltip(const FRigVMTemplateTypeMap& InTypes) const
+{
+	return LOCTEXT("ArrayInsertToolTip", "Inserts an element into an array at a given index.\nModifies the input array.");
+}
+
+FText FRigVMDispatch_ArrayInsert::GetArgumentTooltip(const FName& InArgumentName, TRigVMTypeIndex InTypeIndex) const
+{
+	if(InArgumentName == ArrayName)
+	{
+		return LOCTEXT("ArrayInsert_ArrayArgumentToolTip", "The array to insert an element into.");
+	}
+	if(InArgumentName == IndexName)
+	{
+		return LOCTEXT("ArrayInsert_IndexArgumentToolTip", "The index at which to insert the element.");
+	}
+	if(InArgumentName == ElementName)
+	{
+		return LOCTEXT("ArrayInsert_ElementArgumentToolTip", "The element to insert into the array.");
+	}
+	return FText();
+}
+
+#endif
+
+////////////////////////////////////////////////////////////////////////////////////////////////////////////////////////
+
+TArray<FRigVMTemplateArgument> FRigVMDispatch_ArrayRemove::GetArguments() const
+{
+	return {
+		CreateArgument(ArrayName, ERigVMPinDirection::IO),
+		CreateArgument(IndexName, ERigVMPinDirection::Input)
+	};
+}
+
+FRigVMTemplateTypeMap FRigVMDispatch_ArrayRemove::OnNewArgumentType(const FName& InArgumentName, TRigVMTypeIndex InTypeIndex) const
+{
+	check(InArgumentName == ArrayName);
+
+	return {
+		{ArrayName, InTypeIndex},
+		{IndexName, RigVMTypeUtils::TypeIndex::Int32}
+	};
+}
+
+void FRigVMDispatch_ArrayRemove::Execute(FRigVMExtendedExecuteContext& InContext, FRigVMMemoryHandleArray Handles)
+{
+	const FArrayProperty* ArrayProperty = CastFieldChecked<FArrayProperty>(Handles[0].GetProperty());
+	FScriptArrayHelper ArrayHelper(ArrayProperty, Handles[0].GetData());
+	int32 Index = *(int32*)Handles[1].GetData();
+
+	if(InContext.IsValidArrayIndex(Index, ArrayHelper))
+	{
+		ArrayHelper.RemoveValues(Index, 1);
+	}
+}
+
+#if WITH_EDITOR
+
+FText FRigVMDispatch_ArrayRemove::GetNodeTooltip(const FRigVMTemplateTypeMap& InTypes) const
+{
+	return LOCTEXT("ArrayRemoveToolTip", "Removes an element from an array by index.\nModifies the input array.");
+}
+
+FText FRigVMDispatch_ArrayRemove::GetArgumentTooltip(const FName& InArgumentName, TRigVMTypeIndex InTypeIndex) const
+{
+	if(InArgumentName == ArrayName)
+	{
+		return LOCTEXT("ArrayRemove_ArrayArgumentToolTip", "The array to remove an element from.");
+	}
+	if(InArgumentName == IndexName)
+	{
+		return LOCTEXT("ArrayRemove_IndexArgumentToolTip", "The index at which to remove the element.");
+	}
+	return FText();
+}
+
+#endif
+
+////////////////////////////////////////////////////////////////////////////////////////////////////////////////////////
+
+void FRigVMDispatch_ArrayReverse::Execute(FRigVMExtendedExecuteContext& InContext, FRigVMMemoryHandleArray Handles)
+{
+	const FArrayProperty* ArrayProperty = CastFieldChecked<FArrayProperty>(Handles[0].GetProperty());
+	FScriptArrayHelper ArrayHelper(ArrayProperty, Handles[0].GetData());
+	for(int32 A=0, B=ArrayHelper.Num()-1; A<B; A++, B--)
+	{
+		ArrayHelper.SwapValues(A, B);
+	}
+}
+
+#if WITH_EDITOR
+
+FText FRigVMDispatch_ArrayReverse::GetNodeTooltip(const FRigVMTemplateTypeMap& InTypes) const
+{
+	return LOCTEXT("ArrayReverseToolTip", "Reverses the order of the elements of an array.\nModifies the input array.");
+}
+
+FText FRigVMDispatch_ArrayReverse::GetArgumentTooltip(const FName& InArgumentName, TRigVMTypeIndex InTypeIndex) const
+{
+	if(InArgumentName == ArrayName)
+	{
+		return LOCTEXT("ArrayReverse_ArrayArgumentToolTip", "The array reverse.");
+	}
+	return FText();
+}
+
+#endif
+
+////////////////////////////////////////////////////////////////////////////////////////////////////////////////////////
+
+TArray<FRigVMTemplateArgument> FRigVMDispatch_ArrayFind::GetArguments() const
+{
+	return {
+		CreateArgument(ArrayName, ERigVMPinDirection::Input),
+		CreateArgument(ElementName, ERigVMPinDirection::Input),
+		CreateArgument(IndexName, ERigVMPinDirection::Output),
+		CreateArgument(SuccessName, ERigVMPinDirection::Output)
+	};
+}
+
+FRigVMTemplateTypeMap FRigVMDispatch_ArrayFind::OnNewArgumentType(const FName& InArgumentName, TRigVMTypeIndex InTypeIndex) const
+{
+	check(InArgumentName == ArrayName || InArgumentName == ElementName);
+
+	if(InArgumentName == ArrayName)
+	{
+		return {
+			{ArrayName, InTypeIndex},
+			{ElementName, FRigVMRegistry::Get().GetBaseTypeFromArrayTypeIndex(InTypeIndex)},
+			{IndexName, RigVMTypeUtils::TypeIndex::Int32},
+			{SuccessName, RigVMTypeUtils::TypeIndex::Bool}
+		};
+	}
+
+	return {
+		{ArrayName, FRigVMRegistry::Get().GetArrayTypeFromBaseTypeIndex(InTypeIndex)},
+		{ElementName, InTypeIndex},
+		{IndexName, RigVMTypeUtils::TypeIndex::Int32},
+		{SuccessName, RigVMTypeUtils::TypeIndex::Bool}
+	};
+}
+
+void FRigVMDispatch_ArrayFind::Execute(FRigVMExtendedExecuteContext& InContext, FRigVMMemoryHandleArray Handles)
+{
+	const FArrayProperty* ArrayProperty = CastFieldChecked<FArrayProperty>(Handles[0].GetProperty());
+	FScriptArrayHelper ArrayHelper(ArrayProperty, Handles[0].GetData());
+	int32& Index = *(int32*)Handles[2].GetData();
+	bool& bSuccess = *(bool*)Handles[3].GetData();
+
+	Index = INDEX_NONE;
+	bSuccess = false;
+
+	const FProperty* PropertyA = Handles[1].GetProperty();
+	const FProperty* PropertyB = ArrayProperty->Inner;
+
+	if(PropertyA->SameType(PropertyB))
+	{
+		const uint8* MemoryA = Handles[1].GetData();
+
+		for(int32 ElementIndex = 0; ElementIndex < ArrayHelper.Num(); ElementIndex++)
+		{
+			const uint8* MemoryB = ArrayHelper.GetRawPtr(ElementIndex);
+			if(PropertyA->Identical(MemoryA, MemoryB))
+			{
+				Index = ElementIndex;
+				bSuccess = true;
+				break;
+			}
+		}
+	}
+	else
+	{
+		static constexpr TCHAR IncompatibleTypes[] = TEXT("Array('%s') doesn't support searching for element('%$s').");
+		InContext.GetPublicData<>().Logf(EMessageSeverity::Error, IncompatibleTypes, *PropertyB->GetCPPType(), *PropertyA->GetCPPType());
+	}
+}
+
+#if WITH_EDITOR
+
+FText FRigVMDispatch_ArrayFind::GetNodeTooltip(const FRigVMTemplateTypeMap& InTypes) const
+{
+	return LOCTEXT("ArrayFindToolTip", "Searchs a potential element in an array and returns its index.");
+}
+
+FText FRigVMDispatch_ArrayFind::GetArgumentTooltip(const FName& InArgumentName, TRigVMTypeIndex InTypeIndex) const
+{
+	if(InArgumentName == ArrayName)
+	{
+		return LOCTEXT("ArrayFind_ArrayArgumentToolTip", "The array to search within.");
+	}
+	if(InArgumentName == ElementName)
+	{
+		return LOCTEXT("ArrayFind_ElementArgumentToolTip", "The element to look for.");
+	}
+	if(InArgumentName == IndexName)
+	{
+		return LOCTEXT("ArrayFind_IndexArgumentToolTip", "The index of the found element (or -1).");
+	}
+	if(InArgumentName == SuccessName)
+	{
+		return LOCTEXT("ArrayFind_SuccessArgumentToolTip", "True if the element has been found.");
+	}
+	return FText();
+}
+
+#endif
+
+////////////////////////////////////////////////////////////////////////////////////////////////////////////////////////
+
+TArray<FRigVMTemplateArgument> FRigVMDispatch_ArrayAppend::GetArguments() const
+{
+	return {
+		CreateArgument(ArrayName, ERigVMPinDirection::IO),
+		CreateArgument(OtherName, ERigVMPinDirection::Input)
+	};
+}
+
+FRigVMTemplateTypeMap FRigVMDispatch_ArrayAppend::OnNewArgumentType(const FName& InArgumentName, TRigVMTypeIndex InTypeIndex) const
+{
+	check(InArgumentName == ArrayName || InArgumentName == OtherName);
+
+	return {
+		{ArrayName, InTypeIndex},
+		{OtherName, InTypeIndex}
+	};
+}
+
+void FRigVMDispatch_ArrayAppend::Execute(FRigVMExtendedExecuteContext& InContext, FRigVMMemoryHandleArray Handles)
+{
+	const FArrayProperty* OtherProperty = CastFieldChecked<FArrayProperty>(Handles[1].GetProperty());
+	FScriptArrayHelper OtherHelper(OtherProperty, Handles[1].GetData());
+
+	if(OtherHelper.Num() > 0)
+	{
+		const FArrayProperty* ArrayProperty = CastFieldChecked<FArrayProperty>(Handles[0].GetProperty());
+		FScriptArrayHelper ArrayHelper(ArrayProperty, Handles[0].GetData());
+
+		if(InContext.IsValidArraySize(ArrayHelper.Num() + OtherHelper.Num()))
+		{
+			const FProperty* TargetProperty = ArrayProperty->Inner;
+			const FProperty* SourceProperty = OtherProperty->Inner;
+
+			int32 TargetIndex = ArrayHelper.AddValues(OtherHelper.Num());
+			for(int32 SourceIndex = 0; SourceIndex < OtherHelper.Num(); SourceIndex++, TargetIndex++)
+			{
+				uint8* TargetMemory = ArrayHelper.GetRawPtr(TargetIndex);
+				const uint8* SourceMemory = OtherHelper.GetRawPtr(SourceIndex);
+				URigVMMemoryStorage::CopyProperty(TargetProperty, TargetMemory, SourceProperty, SourceMemory);
+			}
+		}
+	}
+}
+
+#if WITH_EDITOR
+
+FText FRigVMDispatch_ArrayAppend::GetNodeTooltip(const FRigVMTemplateTypeMap& InTypes) const
+{
+	return LOCTEXT("ArrayAppendToolTip", "Appends the another array to the main one.\nModifies the input array.");
+}
+
+FText FRigVMDispatch_ArrayAppend::GetArgumentTooltip(const FName& InArgumentName, TRigVMTypeIndex InTypeIndex) const
+{
+	if(InArgumentName == ArrayName)
+	{
+		return LOCTEXT("ArrayAppend_ArrayArgumentToolTip", "The array to append the other array to.");
+	}
+	if(InArgumentName == OtherName)
+	{
+		return LOCTEXT("ArrayAppend_OtherArgumentToolTip", "The second array to append to the first one.");
+	}
+	return FText();
+}
+
+#endif
+
+////////////////////////////////////////////////////////////////////////////////////////////////////////////////////////
+
+TArray<FRigVMTemplateArgument> FRigVMDispatch_ArrayClone::GetArguments() const
+{
+	return {
+		CreateArgument(ArrayName, ERigVMPinDirection::Input),
+		CreateArgument(CloneName, ERigVMPinDirection::Output)
+	};
+}
+
+FRigVMTemplateTypeMap FRigVMDispatch_ArrayClone::OnNewArgumentType(const FName& InArgumentName, TRigVMTypeIndex InTypeIndex) const
+{
+	check(InArgumentName == ArrayName || InArgumentName == CloneName);
+
+	return {
+		{ArrayName, InTypeIndex},
+		{CloneName, InTypeIndex}
+	};
+}
+
+void FRigVMDispatch_ArrayClone::Execute(FRigVMExtendedExecuteContext& InContext, FRigVMMemoryHandleArray Handles)
+{
+	const FArrayProperty* ArrayProperty = CastFieldChecked<FArrayProperty>(Handles[0].GetProperty());
+	FScriptArrayHelper ArrayHelper(ArrayProperty, Handles[0].GetData());
+	const FArrayProperty* CloneProperty = CastFieldChecked<FArrayProperty>(Handles[1].GetProperty());
+	FScriptArrayHelper CloneHelper(CloneProperty, Handles[1].GetData());
+
+	CloneHelper.Resize(ArrayHelper.Num());
+	if(ArrayHelper.Num() > 0)
+	{
+		const FProperty* TargetProperty = CloneProperty->Inner;
+		const FProperty* SourceProperty = ArrayProperty->Inner;
+		for(int32 ElementIndex = 0; ElementIndex < ArrayHelper.Num(); ElementIndex++)
+		{
+			uint8* TargetMemory = CloneHelper.GetRawPtr(ElementIndex);
+			const uint8* SourceMemory = ArrayHelper.GetRawPtr(ElementIndex);
+			URigVMMemoryStorage::CopyProperty(TargetProperty, TargetMemory, SourceProperty, SourceMemory);
+		}
+	}
+}
+
+#if WITH_EDITOR
+
+FText FRigVMDispatch_ArrayClone::GetNodeTooltip(const FRigVMTemplateTypeMap& InTypes) const
+{
+	return LOCTEXT("ArrayCloneToolTip", "Clones an array and returns a duplicate.");
+}
+
+FText FRigVMDispatch_ArrayClone::GetArgumentTooltip(const FName& InArgumentName, TRigVMTypeIndex InTypeIndex) const
+{
+	if(InArgumentName == ArrayName)
+	{
+		return LOCTEXT("ArrayClone_ArrayArgumentToolTip", "The array to clone");
+	}
+	if(InArgumentName == CloneName)
+	{
+		return LOCTEXT("ArrayClone_CloneArgumentToolTip", "The duplicate of the input array.");
+	}
+	return FText();
+}
+
+#endif
+
+////////////////////////////////////////////////////////////////////////////////////////////////////////////////////////
+
+void FRigVMDispatch_ArrayUnion::Execute(FRigVMExtendedExecuteContext& InContext, FRigVMMemoryHandleArray Handles)
+{
+	const FArrayProperty* OtherProperty = CastFieldChecked<FArrayProperty>(Handles[1].GetProperty());
+	FScriptArrayHelper OtherHelper(OtherProperty, Handles[1].GetData());
+
+	if(OtherHelper.Num() > 0)
+	{
+		const FArrayProperty* ArrayProperty = CastFieldChecked<FArrayProperty>(Handles[0].GetProperty());
+		FScriptArrayHelper ArrayHelper(ArrayProperty, Handles[0].GetData());
+
+		const FProperty* ArrayElementProperty = ArrayProperty->Inner;
+		const FProperty* OtherElementProperty = OtherProperty->Inner;
+
+		const TMap<uint32, int32> HashA = GetArrayHash(ArrayHelper, ArrayProperty);
+		const TMap<uint32, int32> HashB = GetArrayHash(OtherHelper, OtherProperty);
+
+		// copy the complete array to a temp storage
+		TArray<uint8, TAlignedHeapAllocator<16>> TempStorage;
+		const int32 NumElementsA = ArrayHelper.Num();
+		TempStorage.AddZeroed(NumElementsA * ArrayElementProperty->GetSize());
+		uint8* TempMemory = TempStorage.GetData();
+		for(int32 Index = 0; Index < NumElementsA; Index++)
+		{
+			ArrayElementProperty->InitializeValue(TempMemory);
+			ArrayElementProperty->CopyCompleteValue(TempMemory, ArrayHelper.GetRawPtr(Index));
+			TempMemory += ArrayElementProperty->GetSize();
+		}
+
+		ArrayHelper.Resize(0);
+
+		for(const TPair<uint32, int32>& Pair : HashA)
+		{
+			const int32 AddedIndex = ArrayHelper.AddValue();
+			TempMemory = TempStorage.GetData() + Pair.Value * ArrayElementProperty->GetSize();
+						
+			URigVMMemoryStorage::CopyProperty(
+				ArrayElementProperty,
+				ArrayHelper.GetRawPtr(AddedIndex),
+				ArrayElementProperty,
+				TempMemory
+			);
+		}
+
+		TempMemory = TempStorage.GetData();
+		for(int32 Index = 0; Index < NumElementsA; Index++)
+		{
+			ArrayElementProperty->DestroyValue(TempMemory);
+			TempMemory += ArrayElementProperty->GetSize();
+		}
+
+		for(const TPair<uint32, int32>& Pair : HashB)
+		{
+			if(!HashA.Contains(Pair.Key))
+			{
+				const int32 AddedIndex = ArrayHelper.AddValue();
+				URigVMMemoryStorage::CopyProperty(
+					ArrayElementProperty,
+					ArrayHelper.GetRawPtr(AddedIndex),
+					OtherElementProperty,
+					OtherHelper.GetRawPtr(Pair.Value)
+				);
+			}
+		}
+	}
+}
+
+#if WITH_EDITOR
+
+FText FRigVMDispatch_ArrayUnion::GetNodeTooltip(const FRigVMTemplateTypeMap& InTypes) const
+{
+	return LOCTEXT("ArrayUnionToolTip", "Merges two arrays while ensuring unique elements.\nModifies the input array.");
+}
+
+FText FRigVMDispatch_ArrayUnion::GetArgumentTooltip(const FName& InArgumentName, TRigVMTypeIndex InTypeIndex) const
+{
+	if(InArgumentName == ArrayName)
+	{
+		return LOCTEXT("ArrayUnion_ArrayArgumentToolTip", "The array to merge the other array with.");
+	}
+	if(InArgumentName == OtherName)
+	{
+		return LOCTEXT("ArrayUnion_OtherArgumentToolTip", "The second array to merge to the first one.");
+	}
+	return FText();
+}
+
+#endif
+
+////////////////////////////////////////////////////////////////////////////////////////////////////////////////////////
+
+TArray<FRigVMTemplateArgument> FRigVMDispatch_ArrayDifference::GetArguments() const
+{
+	return {
+		CreateArgument(ArrayName, ERigVMPinDirection::Input),
+		CreateArgument(OtherName, ERigVMPinDirection::Input),
+		CreateArgument(ResultName, ERigVMPinDirection::Output)
+	};
+}
+
+FRigVMTemplateTypeMap FRigVMDispatch_ArrayDifference::OnNewArgumentType(const FName& InArgumentName, TRigVMTypeIndex InTypeIndex) const
+{
+	check(InArgumentName == ArrayName || InArgumentName == OtherName || InArgumentName == ResultName);
+
+	return {
+		{ArrayName, InTypeIndex},
+		{OtherName, InTypeIndex},
+		{ResultName, InTypeIndex}
+	};
+}
+
+void FRigVMDispatch_ArrayDifference::Execute(FRigVMExtendedExecuteContext& InContext, FRigVMMemoryHandleArray Handles)
+{
+	const FArrayProperty* ArrayProperty = CastFieldChecked<FArrayProperty>(Handles[0].GetProperty());
+	FScriptArrayHelper ArrayHelper(ArrayProperty, Handles[0].GetData());
+	const FArrayProperty* OtherProperty = CastFieldChecked<FArrayProperty>(Handles[1].GetProperty());
+	FScriptArrayHelper OtherHelper(OtherProperty, Handles[1].GetData());
+	const FArrayProperty* ResultProperty = CastFieldChecked<FArrayProperty>(Handles[2].GetProperty());
+	FScriptArrayHelper ResultHelper(ResultProperty, Handles[2].GetData());
+
+	const FProperty* ArrayElementProperty = ArrayProperty->Inner;
+	const FProperty* OtherElementProperty = OtherProperty->Inner;
+	const FProperty* ResultElementProperty = ResultProperty->Inner;
+
+	ResultHelper.Resize(0);
+
+	const TMap<uint32, int32> HashA = GetArrayHash(ArrayHelper, ArrayProperty);
+	const TMap<uint32, int32> HashB = GetArrayHash(OtherHelper, OtherProperty);
+
+	for(const TPair<uint32, int32>& Pair : HashA)
+	{
+		if(!HashB.Contains(Pair.Key))
+		{
+			const int32 AddedIndex = ResultHelper.AddValue();
+			URigVMMemoryStorage::CopyProperty(
+				ResultElementProperty,
+				ResultHelper.GetRawPtr(AddedIndex),
+				ArrayElementProperty,
+				ArrayHelper.GetRawPtr(Pair.Value)
+			);
+		}
+	}
+	for(const TPair<uint32, int32>& Pair : HashB)
+	{
+		if(!HashA.Contains(Pair.Key))
+		{
+			const int32 AddedIndex = ResultHelper.AddValue();
+			URigVMMemoryStorage::CopyProperty(
+				ResultElementProperty,
+				ResultHelper.GetRawPtr(AddedIndex),
+				OtherElementProperty,
+				OtherHelper.GetRawPtr(Pair.Value)
+			);
+		}
+	}
+}
+
+#if WITH_EDITOR
+
+FText FRigVMDispatch_ArrayDifference::GetNodeTooltip(const FRigVMTemplateTypeMap& InTypes) const
+{
+	return LOCTEXT("ArrayDifferenceToolTip", "Creates a new array containing the difference between the two input arrays.\nDifference here means elements that are only contained in either A or B.");
+}
+
+FText FRigVMDispatch_ArrayDifference::GetArgumentTooltip(const FName& InArgumentName, TRigVMTypeIndex InTypeIndex) const
+{
+	if(InArgumentName == ArrayName)
+	{
+		return LOCTEXT("ArrayDifference_ArrayArgumentToolTip", "The first array to compare the other array with.");
+	}
+	if(InArgumentName == OtherName)
+	{
+		return LOCTEXT("ArrayDifference_OtherArgumentToolTip", "The second array to compare the other array with.");
+	}
+	if(InArgumentName == ResultName)
+	{
+		return LOCTEXT("ArrayDifference_ResultArgumentToolTip", "The resulting difference array.");
+	}
+	return FText();
+}
+
+#endif
+
+////////////////////////////////////////////////////////////////////////////////////////////////////////////////////////
+
+void FRigVMDispatch_ArrayIntersection::Execute(FRigVMExtendedExecuteContext& InContext, FRigVMMemoryHandleArray Handles)
+{
+	const FArrayProperty* ArrayProperty = CastFieldChecked<FArrayProperty>(Handles[0].GetProperty());
+	FScriptArrayHelper ArrayHelper(ArrayProperty, Handles[0].GetData());
+	const FArrayProperty* OtherProperty = CastFieldChecked<FArrayProperty>(Handles[1].GetProperty());
+	FScriptArrayHelper OtherHelper(OtherProperty, Handles[1].GetData());
+	const FArrayProperty* ResultProperty = CastFieldChecked<FArrayProperty>(Handles[2].GetProperty());
+	FScriptArrayHelper ResultHelper(ResultProperty, Handles[2].GetData());
+
+	const FProperty* ArrayElementProperty = ArrayProperty->Inner;
+	const FProperty* ResultElementProperty = ResultProperty->Inner;
+
+	ResultHelper.Resize(0);
+
+	const TMap<uint32, int32> HashA = GetArrayHash(ArrayHelper, ArrayProperty);
+	const TMap<uint32, int32> HashB = GetArrayHash(OtherHelper, OtherProperty);
+
+	for(const TPair<uint32, int32>& Pair : HashA)
+	{
+		if(HashB.Contains(Pair.Key))
+		{
+			const int32 AddedIndex = ResultHelper.AddValue();
+			URigVMMemoryStorage::CopyProperty(
+				ResultElementProperty,
+				ResultHelper.GetRawPtr(AddedIndex),
+				ArrayElementProperty,
+				ArrayHelper.GetRawPtr(Pair.Value)
+			);
+		}
+	}
+}
+
+#if WITH_EDITOR
+
+FText FRigVMDispatch_ArrayIntersection::GetNodeTooltip(const FRigVMTemplateTypeMap& InTypes) const
+{
+	return LOCTEXT("ArrayIntersectionToolTip", "Creates a new array containing the intersection between the two input arrays.\nDifference here means elements that are contained in both A and B.");
+}
+
+FText FRigVMDispatch_ArrayIntersection::GetArgumentTooltip(const FName& InArgumentName, TRigVMTypeIndex InTypeIndex) const
+{
+	if(InArgumentName == ArrayName)
+	{
+		return LOCTEXT("ArrayIntersection_ArrayArgumentToolTip", "The first array to compare the other array with.");
+	}
+	if(InArgumentName == OtherName)
+	{
+		return LOCTEXT("ArrayIntersection_OtherArgumentToolTip", "The second array to compare the other array with.");
+	}
+	if(InArgumentName == ResultName)
+	{
+		return LOCTEXT("ArrayIntersection_ResultArgumentToolTip", "The resulting intersecting array.");
+	}
+	return FText();
+}
+
+#endif
+
+////////////////////////////////////////////////////////////////////////////////////////////////////////////////////////
+
+TArray<FRigVMTemplateArgument> FRigVMDispatch_ArrayIterator::GetArguments() const
+{
+	return {
+		CreateArgument(ArrayName, ERigVMPinDirection::Input),
+		CreateArgument(ElementName, ERigVMPinDirection::Output),
+		CreateArgument(IndexName, ERigVMPinDirection::Output),
+		CreateArgument(CountName, ERigVMPinDirection::Output),
+		CreateArgument(RatioName, ERigVMPinDirection::Output),
+		{FRigVMStruct::ControlFlowBlockToRunName, ERigVMPinDirection::Hidden, RigVMTypeUtils::TypeIndex::FName}
+	};
+}
+
+TArray<FRigVMExecuteArgument> FRigVMDispatch_ArrayIterator::GetExecuteArguments_Impl(
+	const FRigVMDispatchContext& InContext) const
+{
+	TArray<FRigVMExecuteArgument> ExecuteArguments = FRigVMDispatch_ArrayBaseMutable::GetExecuteArguments_Impl(InContext);
+	ExecuteArguments.Emplace(CompletedName, ERigVMPinDirection::Output);
+	return ExecuteArguments;
+}
+
+const TArray<FName>& FRigVMDispatch_ArrayIterator::GetControlFlowBlocks_Impl(
+	const FRigVMDispatchContext& InContext) const
+{
+	static const TArray<FName> Blocks = {ExecuteName, CompletedName};
+	return Blocks;
+}
+
+const bool FRigVMDispatch_ArrayIterator::IsControlFlowBlockSliced(const FName& InBlockName) const
+{
+	return InBlockName == ExecuteName;
+}
+
+FRigVMTemplateTypeMap FRigVMDispatch_ArrayIterator::OnNewArgumentType(const FName& InArgumentName,
+	TRigVMTypeIndex InTypeIndex) const
+{
+	check(InArgumentName == ArrayName || InArgumentName == ElementName);
+
+	if(InArgumentName == ArrayName)
+	{
+		return {
+			{ArrayName, InTypeIndex},
+			{ElementName, FRigVMRegistry::Get().GetBaseTypeFromArrayTypeIndex(InTypeIndex)},
+			{IndexName, RigVMTypeUtils::TypeIndex::Int32},
+			{CountName, RigVMTypeUtils::TypeIndex::Int32},
+			{RatioName, RigVMTypeUtils::TypeIndex::Float},
+			{FRigVMStruct::ControlFlowBlockToRunName, RigVMTypeUtils::TypeIndex::FName}
+		};
+	}
+
+	return {
+		{ArrayName, FRigVMRegistry::Get().GetArrayTypeFromBaseTypeIndex(InTypeIndex)},
+		{ElementName, InTypeIndex},
+		{IndexName, RigVMTypeUtils::TypeIndex::Int32},
+		{CountName, RigVMTypeUtils::TypeIndex::Int32},
+		{RatioName, RigVMTypeUtils::TypeIndex::Float},
+		{FRigVMStruct::ControlFlowBlockToRunName, RigVMTypeUtils::TypeIndex::FName}
+	};}
+
+void FRigVMDispatch_ArrayIterator::Execute(FRigVMExtendedExecuteContext& InContext, FRigVMMemoryHandleArray Handles)
+{
+	const FArrayProperty* ArrayProperty = CastFieldChecked<FArrayProperty>(Handles[0].GetProperty());
+	FScriptArrayHelper ArrayHelper(ArrayProperty, Handles[0].GetData());
+	int32& Index = *(int32*)Handles[2].GetData();
+	int32& Count = *(int32*)Handles[3].GetData();
+	float& Ratio = *(float*)Handles[4].GetData();
+	FName& Block = *(FName*)Handles[5].GetData();
+
+	if(Block.IsNone())
+	{
+		Block = ExecuteName;
+		Count = ArrayHelper.Num();
+		Index = -1;;
+		Ratio = 0.f;
+	}
+
+	Index++;
+	if(Count > 0)
+	{
+		Ratio = float(Index) / float(Count - 1);
+	}
+
+	if(Index >= Count)
+	{
+		Block = CompletedName;
+		return;
+	}
+	
+	uint8* TargetMemory = Handles[1].GetData();
+	const uint8* SourceMemory = ArrayHelper.GetRawPtr(Index);
+	URigVMMemoryStorage::CopyProperty(Handles[1].GetProperty(), TargetMemory, ArrayProperty->Inner, SourceMemory);
+}
+
+#if WITH_EDITOR
+
+FText FRigVMDispatch_ArrayIterator::GetNodeTooltip(const FRigVMTemplateTypeMap& InTypes) const
+{
+	return LOCTEXT("ArrayIteratorToolTip", "Loops over the elements in an array.");
+}
+
+FText FRigVMDispatch_ArrayIterator::GetArgumentTooltip(const FName& InArgumentName, TRigVMTypeIndex InTypeIndex) const
+{
+	if(InArgumentName == ArrayName)
+	{
+		return LOCTEXT("ArrayIterator_ArrayArgumentToolTip", "The array to iterate over");
+	}
+	if(InArgumentName == IndexName)
+	{
+		return LOCTEXT("ArrayIterator_IndexArgumentToolTip", "The index of the current element.");
+	}
+	if(InArgumentName == CountName)
+	{
+		return LOCTEXT("ArrayIterator_CountArgumentToolTip", "The count of all elements in the array.");
+	}
+	if(InArgumentName == ElementName)
+	{
+		return LOCTEXT("ArrayIterator_ElementArgumentToolTip", "The current element.");
+	}
+	if(InArgumentName == RatioName)
+	{
+		return LOCTEXT("ArrayIterator_RatioArgumentToolTip", "A float ratio from 0.0 (first element) to 1.0 (last element).");
+	}
+	if(InArgumentName == CompletedName)
+	{
+		return LOCTEXT("ArrayIterator_CompletedArgumentToolTip", "The execute block to run once the loop has completed.");
+	}
+	return FText();
+}
+
+#endif
+
+#undef LOCTEXT_NAMESPACE
diff --git a/Engine/Source/Runtime/RigVM/Public/RigVMCore/RigVM.h b/Engine/Source/Runtime/RigVM/Public/RigVMCore/RigVM.h
--- a/Engine/Source/Runtime/RigVM/Public/RigVMCore/RigVM.h
+++ b/Engine/Source/Runtime/RigVM/Public/RigVMCore/RigVM.h
@@ -812,7 +812,6 @@ private:
 
 	FRigVMCopyOp GetCopyOpForOperands(const FRigVMOperand& InSource, const FRigVMOperand& InTarget);
 	void RefreshExternalPropertyPaths();
-	static void CopyArray(FScriptArrayHelper& TargetHelper, FRigVMMemoryHandle& TargetHandle, FScriptArrayHelper& SourceHelper, FRigVMMemoryHandle& SourceHandle);
 	
 	TMap<FRigVMOperand, TArray<FRigVMOperand>> OperandToDebugRegisters;
 
diff --git a/Engine/Source/Runtime/RigVM/Public/RigVMCore/RigVMByteCode.h b/Engine/Source/Runtime/RigVM/Public/RigVMCore/RigVMByteCode.h
--- a/Engine/Source/Runtime/RigVM/Public/RigVMCore/RigVMByteCode.h
+++ b/Engine/Source/Runtime/RigVM/Public/RigVMCore/RigVMByteCode.h
@@ -182,22 +182,22 @@ enum class ERigVMOpCode : uint8
 	Exit, // exit the execution loop
 	BeginBlock, // begins a new memory slice / block
 	EndBlock, // ends the last memory slice / block
-	ArrayReset, // clears an array and resets its content
-	ArrayGetNum, // reads and returns the size of an array (binary op, in array, out int32) 
-	ArraySetNum, // resizes an array (binary op, in out array, in int32)
-	ArrayGetAtIndex, // returns an array element by index (ternary op, in array, in int32, out element)  
-	ArraySetAtIndex, // sets an array element by index (ternary op, in out array, in int32, in element)
-	ArrayAdd, // adds an element to an array (ternary op, in out array, in element, out int32 index)
-	ArrayInsert, // inserts an element to an array (ternary op, in out array, in int32, in element)
-	ArrayRemove, // removes an element from an array (binary op, in out array, in inindex)
-	ArrayFind, // finds and returns the index of an element (quaternery op, in array, in element, out int32 index, out bool success)
-	ArrayAppend, // appends an array to another (binary op, in out array, in array)
-	ArrayClone, // clones an array (binary op, in array, out array)
-	ArrayIterator, // iterates over an array (senary op, in array, out element, out index, out count, out ratio, out continue)
-	ArrayUnion, // merges two arrays while avoiding duplicates (binary op, in out array, in other array)
-	ArrayDifference, // returns a new array containing elements only found in one array (ternary op, in array, in array, out result)
-	ArrayIntersection, // returns a new array containing elements found in both of the input arrays (ternary op, in array, in array, out result)
-	ArrayReverse, // returns the reverse of the input array (unary op, in out array)
+	ArrayReset, // (DEPRECATED) clears an array and resets its content
+	ArrayGetNum, // (DEPRECATED) reads and returns the size of an array (binary op, in array, out int32) 
+	ArraySetNum, // (DEPRECATED) resizes an array (binary op, in out array, in int32)
+	ArrayGetAtIndex, // (DEPRECATED) returns an array element by index (ternary op, in array, in int32, out element)  
+	ArraySetAtIndex, // (DEPRECATED) sets an array element by index (ternary op, in out array, in int32, in element)
+	ArrayAdd, // (DEPRECATED) adds an element to an array (ternary op, in out array, in element, out int32 index)
+	ArrayInsert, // (DEPRECATED) inserts an element to an array (ternary op, in out array, in int32, in element)
+	ArrayRemove, // (DEPRECATED) removes an element from an array (binary op, in out array, in inindex)
+	ArrayFind, // (DEPRECATED) finds and returns the index of an element (quaternery op, in array, in element, out int32 index, out bool success)
+	ArrayAppend, // (DEPRECATED) appends an array to another (binary op, in out array, in array)
+	ArrayClone, // (DEPRECATED) clones an array (binary op, in array, out array)
+	ArrayIterator, // (DEPRECATED) iterates over an array (senary op, in array, out element, out index, out count, out ratio, out continue)
+	ArrayUnion, // (DEPRECATED) merges two arrays while avoiding duplicates (binary op, in out array, in other array)
+	ArrayDifference, // (DEPRECATED) returns a new array containing elements only found in one array (ternary op, in array, in array, out result)
+	ArrayIntersection, // (DEPRECATED) returns a new array containing elements found in both of the input arrays (ternary op, in array, in array, out result)
+	ArrayReverse, // (DEPRECATED) returns the reverse of the input array (unary op, in out array)
 	InvokeEntry, // invokes an entry from the entry list
 	JumpToBranch, // jumps to a branch based on a name operand
 	Invalid,
@@ -286,8 +286,6 @@ struct RIGVM_API FRigVMUnaryOp : public FRigVMBaseOp
 			uint8(InOpCode) == uint8(ERigVMOpCode::JumpForwardIf) ||
 			uint8(InOpCode) == uint8(ERigVMOpCode::JumpBackwardIf) ||
 			uint8(InOpCode) == uint8(ERigVMOpCode::ChangeType) ||
-			uint8(InOpCode) == uint8(ERigVMOpCode::ArrayReset) ||
-			uint8(InOpCode) == uint8(ERigVMOpCode::ArrayReverse) ||
 			uint8(InOpCode) == uint8(ERigVMOpCode::JumpToBranch)
 		);
 	}
@@ -325,15 +323,6 @@ struct RIGVM_API FRigVMBinaryOp : public FRigVMBaseOp
 		, ArgA(InArgA)
 		, ArgB(InArgB)
 	{
-		ensure(
-			uint8(InOpCode) == uint8(ERigVMOpCode::BeginBlock) ||
-			uint8(InOpCode) == uint8(ERigVMOpCode::ArrayGetNum) ||
-			uint8(InOpCode) == uint8(ERigVMOpCode::ArraySetNum) ||
-			uint8(InOpCode) == uint8(ERigVMOpCode::ArrayAppend) ||
-			uint8(InOpCode) == uint8(ERigVMOpCode::ArrayClone) ||
-			uint8(InOpCode) == uint8(ERigVMOpCode::ArrayRemove) ||
-			uint8(InOpCode) == uint8(ERigVMOpCode::ArrayUnion)
-		);
 	}
 
 	FRigVMOperand ArgA;
@@ -375,14 +364,6 @@ struct RIGVM_API FRigVMTernaryOp : public FRigVMBaseOp
 		, ArgB(InArgB)
 		, ArgC(InArgC)
 	{
-		ensure(
-			uint8(InOpCode) == uint8(ERigVMOpCode::ArrayAdd) ||
-			uint8(InOpCode) == uint8(ERigVMOpCode::ArrayGetAtIndex) ||
-			uint8(InOpCode) == uint8(ERigVMOpCode::ArraySetAtIndex) ||
-			uint8(InOpCode) == uint8(ERigVMOpCode::ArrayInsert) ||
-			uint8(InOpCode) == uint8(ERigVMOpCode::ArrayDifference) ||
-			uint8(InOpCode) == uint8(ERigVMOpCode::ArrayIntersection)
-		);
 	}
 
 	FRigVMOperand ArgA;
@@ -533,9 +514,6 @@ struct RIGVM_API FRigVMSenaryOp : public FRigVMBaseOp
 		, ArgE(InArgE)
 		, ArgF(InArgF)
 	{
-		ensure(
-			uint8(InOpCode) == uint8(ERigVMOpCode::ArrayIterator)
-		);
 	}
 
 	FRigVMOperand ArgA;
@@ -1047,54 +1025,6 @@ public:
 	// adds an operator to end the last memory slice
 	uint64 AddEndBlockOp();
 
-	// adds an array reset operator
-	uint64 AddArrayResetOp(FRigVMOperand InArrayArg);
-
-	// adds an array get num operator
-	uint64 AddArrayGetNumOp(FRigVMOperand InArrayArg, FRigVMOperand InNumArg);
-
-	// adds an array set num operator
-	uint64 AddArraySetNumOp(FRigVMOperand InArrayArg, FRigVMOperand InNumArg);
-
-	// adds an array get at index operator
-	uint64 AddArrayGetAtIndexOp(FRigVMOperand InArrayArg, FRigVMOperand InIndexArg, FRigVMOperand InElementArg);
-
-	// adds an array set at index operator
-	uint64 AddArraySetAtIndexOp(FRigVMOperand InArrayArg, FRigVMOperand InIndexArg, FRigVMOperand InElementArg);
-
-	// adds an array add operator
-	uint64 AddArrayAddOp(FRigVMOperand InArrayArg, FRigVMOperand InElementArg, FRigVMOperand InIndexArg);
-
-	// adds an array insert operator
-	uint64 AddArrayInsertOp(FRigVMOperand InArrayArg, FRigVMOperand InIndexArg, FRigVMOperand InElementArg);
-
-	// adds an array remove operator
-	uint64 AddArrayRemoveOp(FRigVMOperand InArrayArg, FRigVMOperand InIndexArg);
-
-	// adds an array find operator
-	uint64 AddArrayFindOp(FRigVMOperand InArrayArg, FRigVMOperand InElementArg, FRigVMOperand InIndexArg, FRigVMOperand InSuccessArg);
-
-	// adds an array append operator
-	uint64 AddArrayAppendOp(FRigVMOperand InArrayArg, FRigVMOperand InOtherArrayArg);
-
-	// adds an array clone operator
-	uint64 AddArrayCloneOp(FRigVMOperand InArrayArg, FRigVMOperand InClonedArrayArg);
-
-	// adds an array iterator operator
-	uint64 AddArrayIteratorOp(FRigVMOperand InArrayArg, FRigVMOperand InElementArg, FRigVMOperand InIndexArg, FRigVMOperand InCountArg, FRigVMOperand InRatioArg, FRigVMOperand InContinueArg);
-
-	// adds an array union operator
-	uint64 AddArrayUnionOp(FRigVMOperand InArrayArg, FRigVMOperand InOtherArrayArg);
-	
-	// adds an array difference operator
-	uint64 AddArrayDifferenceOp(FRigVMOperand InArrayArg, FRigVMOperand InOtherArrayArg, FRigVMOperand InResultArrayArg);
-	
-	// adds an array intersection operator
-	uint64 AddArrayIntersectionOp(FRigVMOperand InArrayArg, FRigVMOperand InOtherArrayArg, FRigVMOperand InResultArrayArg);
-	
-	// adds an array reverse operator
-	uint64 AddArrayReverseOp(FRigVMOperand InArrayArg);
-
 	// adds an invoke entry operator
 	uint64 AddInvokeEntryOp(const FName& InEntryName);
 
diff --git a/Engine/Source/Runtime/RigVM/Public/RigVMCore/RigVMDispatchFactory.h b/Engine/Source/Runtime/RigVM/Public/RigVMCore/RigVMDispatchFactory.h
--- a/Engine/Source/Runtime/RigVM/Public/RigVMCore/RigVMDispatchFactory.h
+++ b/Engine/Source/Runtime/RigVM/Public/RigVMCore/RigVMDispatchFactory.h
@@ -155,6 +155,9 @@ public:
 	// returns the name of the permutation for a given set of types
 	FString GetPermutationName(const FRigVMTemplateTypeMap& InTypes) const;
 
+	// returns true if the dispatch uses the same function ptr for all permutations
+	virtual bool IsSingleton() const { return false; } 
+
 protected:
 
 	// returns the name of the permutation for a given set of types
diff --git a/Engine/Source/Runtime/RigVM/Public/RigVMCore/RigVMRegistry.h b/Engine/Source/Runtime/RigVM/Public/RigVMCore/RigVMRegistry.h
--- a/Engine/Source/Runtime/RigVM/Public/RigVMCore/RigVMRegistry.h
+++ b/Engine/Source/Runtime/RigVM/Public/RigVMCore/RigVMRegistry.h
@@ -228,6 +228,29 @@ public:
 	// Returns a dispatch factory given its static struct (or nullptr)
 	FRigVMDispatchFactory* FindOrAddDispatchFactory(UScriptStruct* InFactoryStruct);
 
+	// Returns a dispatch factory given its static struct (or nullptr)
+	template <
+		typename T,
+		typename TEnableIf<TModels<CRigVMUStruct, T>::Value>::Type * = nullptr
+	>
+	FRigVMDispatchFactory* FindOrAddDispatchFactory()
+	{
+		return FindOrAddDispatchFactory(T::StaticStruct());
+	}
+
+	// Returns a dispatch factory's singleton function name if that exists
+	FString FindOrAddSingletonDispatchFunction(UScriptStruct* InFactoryStruct);
+
+	// Returns a dispatch factory's singleton function name if that exists
+	template <
+		typename T,
+		typename TEnableIf<TModels<CRigVMUStruct, T>::Value>::Type * = nullptr
+	>
+	FString FindOrAddSingletonDispatchFunction()
+	{
+		return FindOrAddSingletonDispatchFunction(T::StaticStruct());
+	}
+
 	// Returns all dispatch factories
 	const TArray<FRigVMDispatchFactory*>& GetFactories() const;
 
diff --git a/Engine/Source/Runtime/RigVM/Public/RigVMFunctions/RigVMDispatch_Array.h b/Engine/Source/Runtime/RigVM/Public/RigVMFunctions/RigVMDispatch_Array.h
new file mode 100644
--- /dev/null
+++ b/Engine/Source/Runtime/RigVM/Public/RigVMFunctions/RigVMDispatch_Array.h
@@ -0,0 +1,370 @@
+// Copyright Epic Games, Inc. All Rights Reserved.
+
+#pragma once
+
+#include "RigVMCore/RigVMDispatchFactory.h"
+#include "RigVMDispatch_Core.h"
+#include "RigVMCore/RigVMByteCode.h"
+#include "RigVMDispatch_Array.generated.h"
+
+USTRUCT(meta=(Abstract, Category = "Array", Keywords = "List,Collection", NodeColor = "1,1,1,1"))
+struct RIGVM_API FRigVMDispatch_ArrayBase : public FRigVMDispatch_CoreBase
+{
+	GENERATED_BODY()
+
+public:
+	virtual ERigVMOpCode GetOpCode() const { return ERigVMOpCode::Invalid; }
+	static UScriptStruct* GetFactoryDispatchForOpCode(ERigVMOpCode InOpCode);
+	static FName GetFactoryNameForOpCode(ERigVMOpCode InOpCode);
+#if WITH_EDITOR
+	virtual FString GetArgumentDefaultValue(const FName& InArgumentName, TRigVMTypeIndex InTypeIndex) const override;
+#endif
+	virtual bool IsSingleton() const override { return true; } 
+
+protected:
+	static FRigVMTemplateArgument CreateArgument(const FName& InName, ERigVMPinDirection InDirection);
+	static TMap<uint32, int32> GetArrayHash(FScriptArrayHelper& InArrayHelper, const FArrayProperty* InArrayProperty);
+	
+	static const FName& ExecuteName;
+	static const FName ArrayName;
+	static const FName ValuesName;
+	static const FName NumName;
+	static const FName IndexName;
+	static const FName ElementName;
+	static const FName SuccessName;
+	static const FName OtherName;
+	static const FName CloneName;
+	static const FName CountName;
+	static const FName RatioName;
+	static const FName ResultName;
+	static const FName& CompletedName;
+};
+
+USTRUCT(meta=(Abstract))
+struct RIGVM_API FRigVMDispatch_ArrayBaseMutable : public FRigVMDispatch_ArrayBase
+{
+	GENERATED_BODY()
+
+public:
+	virtual TArray<FRigVMExecuteArgument> GetExecuteArguments_Impl(const FRigVMDispatchContext& InContext) const override;
+};
+
+USTRUCT(meta=(DisplayName = "Array", Keywords = "Make,MakeArray"))
+struct RIGVM_API FRigVMDispatch_ArrayMake : public FRigVMDispatch_ArrayBase
+{
+	GENERATED_BODY()
+
+public:
+	virtual TArray<FRigVMTemplateArgument> GetArguments() const override;
+	virtual FRigVMTemplateTypeMap OnNewArgumentType(const FName& InArgumentName, TRigVMTypeIndex InTypeIndex) const override;
+#if WITH_EDITOR
+	virtual FText GetNodeTooltip(const FRigVMTemplateTypeMap& InTypes) const override;
+	virtual FText GetArgumentTooltip(const FName& InArgumentName, TRigVMTypeIndex InTypeIndex) const override;
+	virtual FString GetArgumentMetaData(const FName& InArgumentName, const FName& InMetaDataKey) const override;
+	virtual FString GetArgumentDefaultValue(const FName& InArgumentName, TRigVMTypeIndex InTypeIndex) const override;
+#endif
+	virtual FName GetArgumentNameForOperandIndex(int32 InOperandIndex, int32 InTotalOperands) const override;
+
+protected:
+	virtual FRigVMFunctionPtr GetDispatchFunctionImpl(const FRigVMTemplateTypeMap& InTypes) const override { return &FRigVMDispatch_ArrayMake::Execute; }
+	static void Execute(FRigVMExtendedExecuteContext& InContext, FRigVMMemoryHandleArray Handles);
+};
+
+USTRUCT(meta=(DisplayName = "Reset"))
+struct RIGVM_API FRigVMDispatch_ArrayReset : public FRigVMDispatch_ArrayBaseMutable
+{
+	GENERATED_BODY()
+
+public:
+	virtual ERigVMOpCode GetOpCode() const override { return ERigVMOpCode::ArrayReset; }
+	virtual TArray<FRigVMTemplateArgument> GetArguments() const override;
+	virtual FRigVMTemplateTypeMap OnNewArgumentType(const FName& InArgumentName, TRigVMTypeIndex InTypeIndex) const override;
+#if WITH_EDITOR
+	virtual FText GetNodeTooltip(const FRigVMTemplateTypeMap& InTypes) const override;
+	virtual FText GetArgumentTooltip(const FName& InArgumentName, TRigVMTypeIndex InTypeIndex) const override;
+#endif
+
+protected:
+	virtual FRigVMFunctionPtr GetDispatchFunctionImpl(const FRigVMTemplateTypeMap& InTypes) const override { return &FRigVMDispatch_ArrayReset::Execute; }
+	static void Execute(FRigVMExtendedExecuteContext& InContext, FRigVMMemoryHandleArray Handles);
+};
+
+USTRUCT(meta=(DisplayName = "Num"))
+struct RIGVM_API FRigVMDispatch_ArrayGetNum : public FRigVMDispatch_ArrayBase
+{
+	GENERATED_BODY()
+
+public:
+	virtual ERigVMOpCode GetOpCode() const override { return ERigVMOpCode::ArrayReset; }
+	virtual TArray<FRigVMTemplateArgument> GetArguments() const override;
+	virtual FRigVMTemplateTypeMap OnNewArgumentType(const FName& InArgumentName, TRigVMTypeIndex InTypeIndex) const override;
+#if WITH_EDITOR
+	virtual FText GetNodeTooltip(const FRigVMTemplateTypeMap& InTypes) const override;
+	virtual FText GetArgumentTooltip(const FName& InArgumentName, TRigVMTypeIndex InTypeIndex) const override;
+#endif
+
+protected:
+	virtual FRigVMFunctionPtr GetDispatchFunctionImpl(const FRigVMTemplateTypeMap& InTypes) const override { return &FRigVMDispatch_ArrayGetNum::Execute; }
+	static void Execute(FRigVMExtendedExecuteContext& InContext, FRigVMMemoryHandleArray Handles);
+};
+
+USTRUCT(meta=(DisplayName = "Set Num"))
+struct RIGVM_API FRigVMDispatch_ArraySetNum : public FRigVMDispatch_ArrayBaseMutable
+{
+	GENERATED_BODY()
+
+public:
+	virtual ERigVMOpCode GetOpCode() const override { return ERigVMOpCode::ArraySetNum; }
+	virtual TArray<FRigVMTemplateArgument> GetArguments() const override;
+	virtual FRigVMTemplateTypeMap OnNewArgumentType(const FName& InArgumentName, TRigVMTypeIndex InTypeIndex) const override;
+#if WITH_EDITOR
+	virtual FText GetNodeTooltip(const FRigVMTemplateTypeMap& InTypes) const override;
+	virtual FText GetArgumentTooltip(const FName& InArgumentName, TRigVMTypeIndex InTypeIndex) const override;
+#endif
+
+protected:
+	virtual FRigVMFunctionPtr GetDispatchFunctionImpl(const FRigVMTemplateTypeMap& InTypes) const override { return &FRigVMDispatch_ArraySetNum::Execute; }
+	static void Execute(FRigVMExtendedExecuteContext& InContext, FRigVMMemoryHandleArray Handles);
+};
+
+USTRUCT(meta=(DisplayName = "At"))
+struct RIGVM_API FRigVMDispatch_ArrayGetAtIndex : public FRigVMDispatch_ArrayBase
+{
+	GENERATED_BODY()
+
+public:
+	virtual ERigVMOpCode GetOpCode() const override { return ERigVMOpCode::ArrayGetAtIndex; }
+	virtual TArray<FRigVMTemplateArgument> GetArguments() const override;
+	virtual FRigVMTemplateTypeMap OnNewArgumentType(const FName& InArgumentName, TRigVMTypeIndex InTypeIndex) const override;
+#if WITH_EDITOR
+	virtual FText GetNodeTooltip(const FRigVMTemplateTypeMap& InTypes) const override;
+	virtual FText GetArgumentTooltip(const FName& InArgumentName, TRigVMTypeIndex InTypeIndex) const override;
+#endif
+
+protected:
+	virtual FRigVMFunctionPtr GetDispatchFunctionImpl(const FRigVMTemplateTypeMap& InTypes) const override { return &FRigVMDispatch_ArrayGetAtIndex::Execute; }
+	static void Execute(FRigVMExtendedExecuteContext& InContext, FRigVMMemoryHandleArray Handles);
+};
+
+USTRUCT(meta=(DisplayName = "Set At"))
+struct RIGVM_API FRigVMDispatch_ArraySetAtIndex : public FRigVMDispatch_ArrayBaseMutable
+{
+	GENERATED_BODY()
+
+public:
+	virtual ERigVMOpCode GetOpCode() const override { return ERigVMOpCode::ArraySetAtIndex; }
+	virtual TArray<FRigVMTemplateArgument> GetArguments() const override;
+	virtual FRigVMTemplateTypeMap OnNewArgumentType(const FName& InArgumentName, TRigVMTypeIndex InTypeIndex) const override;
+#if WITH_EDITOR
+	virtual FText GetNodeTooltip(const FRigVMTemplateTypeMap& InTypes) const override;
+	virtual FText GetArgumentTooltip(const FName& InArgumentName, TRigVMTypeIndex InTypeIndex) const override;
+#endif
+
+protected:
+	virtual FRigVMFunctionPtr GetDispatchFunctionImpl(const FRigVMTemplateTypeMap& InTypes) const override { return &FRigVMDispatch_ArraySetAtIndex::Execute; }
+	static void Execute(FRigVMExtendedExecuteContext& InContext, FRigVMMemoryHandleArray Handles);
+};
+
+USTRUCT(meta=(DisplayName = "Add"))
+struct RIGVM_API FRigVMDispatch_ArrayAdd : public FRigVMDispatch_ArraySetAtIndex
+{
+	GENERATED_BODY()
+
+public:
+	virtual ERigVMOpCode GetOpCode() const override { return ERigVMOpCode::ArrayAdd; }
+	virtual TArray<FRigVMTemplateArgument> GetArguments() const override;
+#if WITH_EDITOR
+	virtual FText GetNodeTooltip(const FRigVMTemplateTypeMap& InTypes) const override;
+	virtual FText GetArgumentTooltip(const FName& InArgumentName, TRigVMTypeIndex InTypeIndex) const override;
+#endif
+
+protected:
+	virtual FRigVMFunctionPtr GetDispatchFunctionImpl(const FRigVMTemplateTypeMap& InTypes) const override { return &FRigVMDispatch_ArrayAdd::Execute; }
+	static void Execute(FRigVMExtendedExecuteContext& InContext, FRigVMMemoryHandleArray Handles);
+};
+
+USTRUCT(meta=(DisplayName = "Insert"))
+struct RIGVM_API FRigVMDispatch_ArrayInsert : public FRigVMDispatch_ArraySetAtIndex
+{
+	GENERATED_BODY()
+
+public:
+	virtual ERigVMOpCode GetOpCode() const override { return ERigVMOpCode::ArrayInsert; }
+#if WITH_EDITOR
+   	virtual FText GetNodeTooltip(const FRigVMTemplateTypeMap& InTypes) const override;
+   	virtual FText GetArgumentTooltip(const FName& InArgumentName, TRigVMTypeIndex InTypeIndex) const override;
+#endif
+
+protected:
+	virtual FRigVMFunctionPtr GetDispatchFunctionImpl(const FRigVMTemplateTypeMap& InTypes) const override { return &FRigVMDispatch_ArrayInsert::Execute; }
+	static void Execute(FRigVMExtendedExecuteContext& InContext, FRigVMMemoryHandleArray Handles);
+};
+
+USTRUCT(meta=(DisplayName = "Remove"))
+struct RIGVM_API FRigVMDispatch_ArrayRemove : public FRigVMDispatch_ArrayBaseMutable
+{
+	GENERATED_BODY()
+
+public:
+	virtual ERigVMOpCode GetOpCode() const override { return ERigVMOpCode::ArrayRemove; }
+	virtual TArray<FRigVMTemplateArgument> GetArguments() const override;
+	virtual FRigVMTemplateTypeMap OnNewArgumentType(const FName& InArgumentName, TRigVMTypeIndex InTypeIndex) const override;
+#if WITH_EDITOR
+   	virtual FText GetNodeTooltip(const FRigVMTemplateTypeMap& InTypes) const override;
+   	virtual FText GetArgumentTooltip(const FName& InArgumentName, TRigVMTypeIndex InTypeIndex) const override;
+#endif
+
+protected:
+	virtual FRigVMFunctionPtr GetDispatchFunctionImpl(const FRigVMTemplateTypeMap& InTypes) const override { return &FRigVMDispatch_ArrayRemove::Execute; }
+	static void Execute(FRigVMExtendedExecuteContext& InContext, FRigVMMemoryHandleArray Handles);
+};
+
+USTRUCT(meta=(DisplayName = "Reverse"))
+struct RIGVM_API FRigVMDispatch_ArrayReverse : public FRigVMDispatch_ArrayReset
+{
+	GENERATED_BODY()
+
+public:
+	virtual ERigVMOpCode GetOpCode() const override { return ERigVMOpCode::ArrayReverse; }
+#if WITH_EDITOR
+   	virtual FText GetNodeTooltip(const FRigVMTemplateTypeMap& InTypes) const override;
+   	virtual FText GetArgumentTooltip(const FName& InArgumentName, TRigVMTypeIndex InTypeIndex) const override;
+#endif
+
+protected:
+
+	virtual FRigVMFunctionPtr GetDispatchFunctionImpl(const FRigVMTemplateTypeMap& InTypes) const override { return &FRigVMDispatch_ArrayReverse::Execute; }
+	static void Execute(FRigVMExtendedExecuteContext& InContext, FRigVMMemoryHandleArray Handles);
+};
+
+USTRUCT(meta=(DisplayName = "Find"))
+struct RIGVM_API FRigVMDispatch_ArrayFind : public FRigVMDispatch_ArrayBase
+{
+	GENERATED_BODY()
+
+public:
+	virtual ERigVMOpCode GetOpCode() const override { return ERigVMOpCode::ArrayFind; }
+	virtual TArray<FRigVMTemplateArgument> GetArguments() const override;
+	virtual FRigVMTemplateTypeMap OnNewArgumentType(const FName& InArgumentName, TRigVMTypeIndex InTypeIndex) const override;
+#if WITH_EDITOR
+   	virtual FText GetNodeTooltip(const FRigVMTemplateTypeMap& InTypes) const override;
+   	virtual FText GetArgumentTooltip(const FName& InArgumentName, TRigVMTypeIndex InTypeIndex) const override;
+#endif
+
+protected:
+	virtual FRigVMFunctionPtr GetDispatchFunctionImpl(const FRigVMTemplateTypeMap& InTypes) const override { return &FRigVMDispatch_ArrayFind::Execute; }
+	static void Execute(FRigVMExtendedExecuteContext& InContext, FRigVMMemoryHandleArray Handles);
+};
+
+USTRUCT(meta=(DisplayName = "Append"))
+struct RIGVM_API FRigVMDispatch_ArrayAppend : public FRigVMDispatch_ArrayBaseMutable
+{
+	GENERATED_BODY()
+
+public:
+	virtual ERigVMOpCode GetOpCode() const override { return ERigVMOpCode::ArrayAppend; }
+	virtual TArray<FRigVMTemplateArgument> GetArguments() const override;
+	virtual FRigVMTemplateTypeMap OnNewArgumentType(const FName& InArgumentName, TRigVMTypeIndex InTypeIndex) const override;
+#if WITH_EDITOR
+   	virtual FText GetNodeTooltip(const FRigVMTemplateTypeMap& InTypes) const override;
+   	virtual FText GetArgumentTooltip(const FName& InArgumentName, TRigVMTypeIndex InTypeIndex) const override;
+#endif
+
+protected:
+	virtual FRigVMFunctionPtr GetDispatchFunctionImpl(const FRigVMTemplateTypeMap& InTypes) const override { return &FRigVMDispatch_ArrayAppend::Execute; }
+	static void Execute(FRigVMExtendedExecuteContext& InContext, FRigVMMemoryHandleArray Handles);
+};
+
+USTRUCT(meta=(DisplayName = "Clone"))
+struct RIGVM_API FRigVMDispatch_ArrayClone : public FRigVMDispatch_ArrayBase
+{
+	GENERATED_BODY()
+
+public:
+	virtual ERigVMOpCode GetOpCode() const override { return ERigVMOpCode::ArrayClone; }
+	virtual TArray<FRigVMTemplateArgument> GetArguments() const override;
+	virtual FRigVMTemplateTypeMap OnNewArgumentType(const FName& InArgumentName, TRigVMTypeIndex InTypeIndex) const override;
+#if WITH_EDITOR
+   	virtual FText GetNodeTooltip(const FRigVMTemplateTypeMap& InTypes) const override;
+   	virtual FText GetArgumentTooltip(const FName& InArgumentName, TRigVMTypeIndex InTypeIndex) const override;
+#endif
+
+protected:
+	virtual FRigVMFunctionPtr GetDispatchFunctionImpl(const FRigVMTemplateTypeMap& InTypes) const override { return &FRigVMDispatch_ArrayClone::Execute; }
+	static void Execute(FRigVMExtendedExecuteContext& InContext, FRigVMMemoryHandleArray Handles);
+};
+
+USTRUCT(meta=(DisplayName = "Union"))
+struct RIGVM_API FRigVMDispatch_ArrayUnion : public FRigVMDispatch_ArrayAppend
+{
+	GENERATED_BODY()
+
+public:
+	virtual ERigVMOpCode GetOpCode() const override { return ERigVMOpCode::ArrayUnion; }
+#if WITH_EDITOR
+   	virtual FText GetNodeTooltip(const FRigVMTemplateTypeMap& InTypes) const override;
+   	virtual FText GetArgumentTooltip(const FName& InArgumentName, TRigVMTypeIndex InTypeIndex) const override;
+#endif
+
+protected:
+	virtual FRigVMFunctionPtr GetDispatchFunctionImpl(const FRigVMTemplateTypeMap& InTypes) const override { return &FRigVMDispatch_ArrayUnion::Execute; }
+	static void Execute(FRigVMExtendedExecuteContext& InContext, FRigVMMemoryHandleArray Handles);
+};
+
+USTRUCT(meta=(DisplayName = "Difference"))
+struct RIGVM_API FRigVMDispatch_ArrayDifference : public FRigVMDispatch_ArrayBase
+{
+	GENERATED_BODY()
+
+public:
+	virtual ERigVMOpCode GetOpCode() const override { return ERigVMOpCode::ArrayDifference; }
+	virtual TArray<FRigVMTemplateArgument> GetArguments() const override;
+	virtual FRigVMTemplateTypeMap OnNewArgumentType(const FName& InArgumentName, TRigVMTypeIndex InTypeIndex) const override;
+#if WITH_EDITOR
+   	virtual FText GetNodeTooltip(const FRigVMTemplateTypeMap& InTypes) const override;
+   	virtual FText GetArgumentTooltip(const FName& InArgumentName, TRigVMTypeIndex InTypeIndex) const override;
+#endif
+
+protected:
+	virtual FRigVMFunctionPtr GetDispatchFunctionImpl(const FRigVMTemplateTypeMap& InTypes) const override { return &FRigVMDispatch_ArrayDifference::Execute; }
+	static void Execute(FRigVMExtendedExecuteContext& InContext, FRigVMMemoryHandleArray Handles);
+};
+
+USTRUCT(meta=(DisplayName = "Intersection"))
+struct RIGVM_API FRigVMDispatch_ArrayIntersection : public FRigVMDispatch_ArrayDifference
+{
+	GENERATED_BODY()
+
+public:
+	virtual ERigVMOpCode GetOpCode() const override { return ERigVMOpCode::ArrayIntersection; }
+#if WITH_EDITOR
+   	virtual FText GetNodeTooltip(const FRigVMTemplateTypeMap& InTypes) const override;
+   	virtual FText GetArgumentTooltip(const FName& InArgumentName, TRigVMTypeIndex InTypeIndex) const override;
+#endif
+
+protected:
+	virtual FRigVMFunctionPtr GetDispatchFunctionImpl(const FRigVMTemplateTypeMap& InTypes) const override { return &FRigVMDispatch_ArrayIntersection::Execute; }
+	static void Execute(FRigVMExtendedExecuteContext& InContext, FRigVMMemoryHandleArray Handles);
+};
+
+USTRUCT(meta=(DisplayName = "For Each", Icon="EditorStyle|GraphEditor.Macro.ForEach_16x"))
+struct RIGVM_API FRigVMDispatch_ArrayIterator : public FRigVMDispatch_ArrayBaseMutable
+{
+	GENERATED_BODY()
+
+public:
+	virtual ERigVMOpCode GetOpCode() const override { return ERigVMOpCode::ArrayIterator; }
+	virtual TArray<FRigVMTemplateArgument> GetArguments() const override;
+	virtual TArray<FRigVMExecuteArgument> GetExecuteArguments_Impl(const FRigVMDispatchContext& InContext) const override;
+	virtual const TArray<FName>& GetControlFlowBlocks_Impl(const FRigVMDispatchContext& InContext) const override;
+	virtual const bool IsControlFlowBlockSliced(const FName& InBlockName) const override;
+	virtual FRigVMTemplateTypeMap OnNewArgumentType(const FName& InArgumentName, TRigVMTypeIndex InTypeIndex) const override;
+#if WITH_EDITOR
+   	virtual FText GetNodeTooltip(const FRigVMTemplateTypeMap& InTypes) const override;
+   	virtual FText GetArgumentTooltip(const FName& InArgumentName, TRigVMTypeIndex InTypeIndex) const override;
+#endif
+
+protected:
+	virtual FRigVMFunctionPtr GetDispatchFunctionImpl(const FRigVMTemplateTypeMap& InTypes) const override { return &FRigVMDispatch_ArrayIterator::Execute; }
+	static void Execute(FRigVMExtendedExecuteContext& InContext, FRigVMMemoryHandleArray Handles);
+};
diff --git a/Engine/Source/Runtime/RigVM/Public/RigVMFunctions/RigVMDispatch_If.h b/Engine/Source/Runtime/RigVM/Public/RigVMFunctions/RigVMDispatch_If.h
--- a/Engine/Source/Runtime/RigVM/Public/RigVMFunctions/RigVMDispatch_If.h
+++ b/Engine/Source/Runtime/RigVM/Public/RigVMFunctions/RigVMDispatch_If.h
@@ -18,6 +18,7 @@ public:
 
 	virtual TArray<FRigVMTemplateArgument> GetArguments() const override;
 	virtual FRigVMTemplateTypeMap OnNewArgumentType(const FName& InArgumentName, TRigVMTypeIndex InTypeIndex) const override;
+	virtual bool IsSingleton() const override { return true; } 
 
 #if WITH_EDITOR
 	virtual FString GetArgumentMetaData(const FName& InArgumentName, const FName& InMetaDataKey) const override;
diff --git a/Engine/Source/Runtime/RigVM/Public/RigVMFunctions/RigVMDispatch_Print.h b/Engine/Source/Runtime/RigVM/Public/RigVMFunctions/RigVMDispatch_Print.h
--- a/Engine/Source/Runtime/RigVM/Public/RigVMFunctions/RigVMDispatch_Print.h
+++ b/Engine/Source/Runtime/RigVM/Public/RigVMFunctions/RigVMDispatch_Print.h
@@ -18,6 +18,7 @@ public:
 	virtual TArray<FRigVMTemplateArgument> GetArguments() const override;
 	virtual TArray<FRigVMExecuteArgument> GetExecuteArguments_Impl(const FRigVMDispatchContext& InContext) const override;
 	virtual FRigVMTemplateTypeMap OnNewArgumentType(const FName& InArgumentName, TRigVMTypeIndex InTypeIndex) const override;
+	virtual bool IsSingleton() const override { return true; } 
 
 #if WITH_EDITOR
 	virtual FString GetArgumentDefaultValue(const FName& InArgumentName, TRigVMTypeIndex InTypeIndex) const override;
diff --git a/Engine/Source/Runtime/RigVM/Public/RigVMFunctions/RigVMDispatch_Select.h b/Engine/Source/Runtime/RigVM/Public/RigVMFunctions/RigVMDispatch_Select.h
--- a/Engine/Source/Runtime/RigVM/Public/RigVMFunctions/RigVMDispatch_Select.h
+++ b/Engine/Source/Runtime/RigVM/Public/RigVMFunctions/RigVMDispatch_Select.h
@@ -23,6 +23,7 @@ public:
 	virtual FString GetArgumentDefaultValue(const FName& InArgumentName, TRigVMTypeIndex InTypeIndex) const override;
 #endif
 	virtual FName GetArgumentNameForOperandIndex(int32 InOperandIndex, int32 InTotalOperands) const override;
+	virtual bool IsSingleton() const override { return true; } 
 
 protected:
 
diff --git a/Engine/Source/Runtime/RigVM/Public/RigVMFunctions/RigVMDispatch_Switch.h b/Engine/Source/Runtime/RigVM/Public/RigVMFunctions/RigVMDispatch_Switch.h
--- a/Engine/Source/Runtime/RigVM/Public/RigVMFunctions/RigVMDispatch_Switch.h
+++ b/Engine/Source/Runtime/RigVM/Public/RigVMFunctions/RigVMDispatch_Switch.h
@@ -21,6 +21,7 @@ public:
 
 	virtual TArray<FRigVMTemplateArgument> GetArguments() const override;
 	virtual TArray<FRigVMExecuteArgument> GetExecuteArguments_Impl(const FRigVMDispatchContext& InContext) const override;
+	virtual bool IsSingleton() const override { return true; } 
 
 #if WITH_EDITOR
 	virtual FString GetArgumentMetaData(const FName& InArgumentName, const FName& InMetaDataKey) const override;
