RigVM: Avoid discovering types from string when not necessary

diff --git a/Engine/Source/Runtime/RigVM/Private/RigVMCore/RigVMTemplate.cpp b/Engine/Source/Runtime/RigVM/Private/RigVMCore/RigVMTemplate.cpp
--- a/Engine/Source/Runtime/RigVM/Private/RigVMCore/RigVMTemplate.cpp
+++ b/Engine/Source/Runtime/RigVM/Private/RigVMCore/RigVMTemplate.cpp
@@ -31,7 +31,7 @@ FRigVMTemplateArgument::FRigVMTemplateArgument(FProperty* InProperty)
 	FString ExtendedType;
 	const FString CPPType = InProperty->GetCPPType(&ExtendedType);
 	const FName CPPTypeName = *(CPPType + ExtendedType);
-	FRigVMTemplateArgumentType Type(CPPTypeName);
+	UObject* CPPTypeObject = nullptr;
 
 	if (FArrayProperty* ArrayProperty = CastField<FArrayProperty>(InProperty))
 	{
@@ -40,18 +40,18 @@ FRigVMTemplateArgument::FRigVMTemplateArgument(FProperty* InProperty)
 
 	if (FStructProperty* StructProperty = CastField<FStructProperty>(InProperty))
 	{
-		Type.CPPTypeObject = StructProperty->Struct;
+		CPPTypeObject = StructProperty->Struct;
 	}
 	else if (FEnumProperty* EnumProperty = CastField<FEnumProperty>(InProperty))
 	{
-		Type.CPPTypeObject = EnumProperty->GetEnum();
+		CPPTypeObject = EnumProperty->GetEnum();
 	}
 	else if (FByteProperty* ByteProperty = CastField<FByteProperty>(InProperty))
 	{
-		Type.CPPTypeObject = ByteProperty->Enum;
+		CPPTypeObject = ByteProperty->Enum;
 	}
-	Type.CPPType = *RigVMTypeUtils::PostProcessCPPType(Type.CPPType.ToString(), Type.CPPTypeObject);
 
+	const FRigVMTemplateArgumentType Type(CPPTypeName, CPPTypeObject);
 	const TRigVMTypeIndex TypeIndex = FRigVMRegistry::Get().FindOrAddType_Internal(Type, true); 
 
 	TypeIndices.Add(TypeIndex);
diff --git a/Engine/Source/Runtime/RigVM/Public/RigVMCore/RigVMTemplate.h b/Engine/Source/Runtime/RigVM/Public/RigVMCore/RigVMTemplate.h
--- a/Engine/Source/Runtime/RigVM/Public/RigVMCore/RigVMTemplate.h
+++ b/Engine/Source/Runtime/RigVM/Public/RigVMCore/RigVMTemplate.h
@@ -69,9 +69,10 @@ struct RIGVM_API FRigVMTemplateArgumentType
 		// InCppType is unreliable because not all caller knows that
 		// we use generated unique names for user defined structs
 		// so here we override the CppType name with the actual name used in the registry
-		CPPType = *RigVMTypeUtils::PostProcessCPPType(CPPType.ToString(), CPPTypeObject);
+		const FString InCPPTypeString = CPPType.ToString();
+		CPPType = *RigVMTypeUtils::PostProcessCPPType(InCPPTypeString, CPPTypeObject);
 		
-		check(!CPPType.IsNone());
+		checkf(!CPPType.IsNone(), TEXT("FRigVMTemplateArgumentType(): Input CPPType '%s' could not be resolved."), *InCPPTypeString);
 	}
 
 	FRigVMTemplateArgumentType(UScriptStruct* InScriptStruct)
diff --git a/Engine/Source/Runtime/RigVM/Public/RigVMTypeUtils.h b/Engine/Source/Runtime/RigVM/Public/RigVMTypeUtils.h
--- a/Engine/Source/Runtime/RigVM/Public/RigVMTypeUtils.h
+++ b/Engine/Source/Runtime/RigVM/Public/RigVMTypeUtils.h
@@ -167,6 +167,11 @@ namespace RigVMTypeUtils
 			TEXT("TArray<TArray<TObjectPtr<"),
 			TEXT("TScriptInterface<")
 		};
+		static const TArray<FString> PrefixesNotRequiringCPPTypeObject = {
+			TEXT("UInt"), 
+			TEXT("TArray<UInt"),
+			TEXT("TArray<TArray<UInt")
+		};
 		static const TArray<FString> CPPTypesNotRequiringCPPTypeObject = {
 			TEXT("FString"), 
 			TEXT("TArray<FString>"), 
@@ -178,6 +183,13 @@ namespace RigVMTypeUtils
 
 		if(!CPPTypesNotRequiringCPPTypeObject.Contains(InCPPType))
 		{
+			for(const FString& Prefix : PrefixesNotRequiringCPPTypeObject)
+			{
+				if(InCPPType.StartsWith(Prefix, ESearchCase::CaseSensitive))
+				{
+					return false;
+				}
+			}
 			for(const FString& Prefix : PrefixesRequiringCPPTypeObject)
 			{
 				if(InCPPType.StartsWith(Prefix, ESearchCase::CaseSensitive))
