Control Rig: Catch missing user defined data structures

diff --git a/Engine/Source/Runtime/RigVM/Private/RigVMCore/RigVMMemoryStorage.cpp b/Engine/Source/Runtime/RigVM/Private/RigVMCore/RigVMMemoryStorage.cpp
--- a/Engine/Source/Runtime/RigVM/Private/RigVMCore/RigVMMemoryStorage.cpp
+++ b/Engine/Source/Runtime/RigVM/Private/RigVMCore/RigVMMemoryStorage.cpp
@@ -257,12 +257,6 @@ FRigVMPropertyDescription::FRigVMPropertyDescription(const FName& InName, const
 	, Containers()
 	, DefaultValue(InDefaultValue)
 {
-	// sanity check that we have a CPP Type Object for enums, structs and uobjects
-	if(RigVMTypeUtils::RequiresCPPTypeObject(CPPType))
-	{
-		checkf(CPPTypeObject != nullptr, TEXT("CPPType '%s' requires the CPPTypeObject to be provided."), *CPPType);
-	}
-
 	if(CPPTypeObject)
 	{
 		if(CPPTypeObject->IsA<UClass>())
