Control Rig: Remove remapping Space to nodes with no space

diff --git a/Engine/Plugins/Animation/ControlRig/Source/ControlRig/Private/Units/Deprecated/Math/RigUnit_Converter.cpp b/Engine/Plugins/Animation/ControlRig/Source/ControlRig/Private/Units/Deprecated/Math/RigUnit_Converter.cpp
--- a/Engine/Plugins/Animation/ControlRig/Source/ControlRig/Private/Units/Deprecated/Math/RigUnit_Converter.cpp
+++ b/Engine/Plugins/Animation/ControlRig/Source/ControlRig/Private/Units/Deprecated/Math/RigUnit_Converter.cpp
@@ -17,7 +17,6 @@ FRigVMStructUpgradeInfo FRigUnit_ConvertTransform::GetUpgradeInfo() const
 	NewNode.Value = Input;
 
 	FRigVMStructUpgradeInfo Info(*this, NewNode);
-	Info.AddRemappedPin(TEXT("Space"), TEXT("Space.Name"));
 	Info.AddRemappedPin(TEXT("Input"), TEXT("Value"));
 	return Info;
 }
diff --git a/Engine/Source/Runtime/RigVM/Private/RigVMFunctions/Debug/RigVMFunction_DebugLine.cpp b/Engine/Source/Runtime/RigVM/Private/RigVMFunctions/Debug/RigVMFunction_DebugLine.cpp
--- a/Engine/Source/Runtime/RigVM/Private/RigVMFunctions/Debug/RigVMFunction_DebugLine.cpp
+++ b/Engine/Source/Runtime/RigVM/Private/RigVMFunctions/Debug/RigVMFunction_DebugLine.cpp
@@ -27,7 +27,6 @@ FRigVMStructUpgradeInfo FRigVMFunction_DebugLine::GetUpgradeInfo() const
 	NewNode.bEnabled = bEnabled;
 
 	FRigVMStructUpgradeInfo Info(*this, NewNode);
-	Info.AddRemappedPin(TEXT("Space"), TEXT("Space.Name"));
 	return Info;
 }
 
diff --git a/Engine/Source/Runtime/RigVM/Private/RigVMFunctions/Debug/RigVMFunction_DebugLineStrip.cpp b/Engine/Source/Runtime/RigVM/Private/RigVMFunctions/Debug/RigVMFunction_DebugLineStrip.cpp
--- a/Engine/Source/Runtime/RigVM/Private/RigVMFunctions/Debug/RigVMFunction_DebugLineStrip.cpp
+++ b/Engine/Source/Runtime/RigVM/Private/RigVMFunctions/Debug/RigVMFunction_DebugLineStrip.cpp
@@ -25,7 +25,6 @@ FRigVMStructUpgradeInfo FRigVMFunction_DebugLineStrip::GetUpgradeInfo() const
 	NewNode.bEnabled = bEnabled;
 
 	FRigVMStructUpgradeInfo Info(*this, NewNode);
-	Info.AddRemappedPin(TEXT("Space"), TEXT("Space.Name"));
 	return Info;
 }
 
diff --git a/Engine/Source/Runtime/RigVM/Private/RigVMFunctions/Debug/RigVMFunction_DebugPrimitives.cpp b/Engine/Source/Runtime/RigVM/Private/RigVMFunctions/Debug/RigVMFunction_DebugPrimitives.cpp
--- a/Engine/Source/Runtime/RigVM/Private/RigVMFunctions/Debug/RigVMFunction_DebugPrimitives.cpp
+++ b/Engine/Source/Runtime/RigVM/Private/RigVMFunctions/Debug/RigVMFunction_DebugPrimitives.cpp
@@ -27,7 +27,6 @@ FRigVMStructUpgradeInfo FRigVMFunction_DebugRectangle::GetUpgradeInfo() const
 	NewNode.bEnabled = bEnabled;
 
 	FRigVMStructUpgradeInfo Info(*this, NewNode);
-	Info.AddRemappedPin(TEXT("Space"), TEXT("Space.Name"));
 	return Info;
 }
 
@@ -70,7 +69,6 @@ FRigVMStructUpgradeInfo FRigVMFunction_DebugArc::GetUpgradeInfo() const
 	NewNode.bEnabled = bEnabled;
 
 	FRigVMStructUpgradeInfo Info(*this, NewNode);
-	Info.AddRemappedPin(TEXT("Space"), TEXT("Space.Name"));
 	return Info;
 }
 
diff --git a/Engine/Source/Runtime/RigVM/Private/RigVMFunctions/Debug/RigVMFunction_DebugTransform.cpp b/Engine/Source/Runtime/RigVM/Private/RigVMFunctions/Debug/RigVMFunction_DebugTransform.cpp
--- a/Engine/Source/Runtime/RigVM/Private/RigVMFunctions/Debug/RigVMFunction_DebugTransform.cpp
+++ b/Engine/Source/Runtime/RigVM/Private/RigVMFunctions/Debug/RigVMFunction_DebugTransform.cpp
@@ -45,7 +45,6 @@ FRigVMStructUpgradeInfo FRigVMFunction_DebugTransform::GetUpgradeInfo() const
 
 	FRigVMStructUpgradeInfo Info(*this, NewNode);
 	Info.AddRemappedPin(TEXT("Transform"), TEXT("Value"));
-	Info.AddRemappedPin(TEXT("Space"), TEXT("Space.Name"));
 	return Info;
 }
 
@@ -74,7 +73,6 @@ FRigVMStructUpgradeInfo FRigVMFunction_DebugTransformMutable::GetUpgradeInfo() c
 	NewNode.bEnabled = bEnabled;
 
 	FRigVMStructUpgradeInfo Info(*this, NewNode);
-	Info.AddRemappedPin(TEXT("Space"), TEXT("Space.Name"));
 	return Info;
 }
 
@@ -157,7 +155,6 @@ FRigVMStructUpgradeInfo FRigVMFunction_DebugTransformArrayMutable::GetUpgradeInf
 	NewNode.bEnabled = bEnabled;
 
 	FRigVMStructUpgradeInfo Info(*this, NewNode);
-	Info.AddRemappedPin(TEXT("Space"), TEXT("Space.Name"));
 	return Info;
 }
 
