RigVM : Added Array Sum Average Functions

diff --git a/Engine/Source/Runtime/RigVM/Private/RigVMFunctions/Math/RigVMFunction_MathDouble.cpp b/Engine/Source/Runtime/RigVM/Private/RigVMFunctions/Math/RigVMFunction_MathDouble.cpp
--- a/Engine/Source/Runtime/RigVM/Private/RigVMFunctions/Math/RigVMFunction_MathDouble.cpp
+++ b/Engine/Source/Runtime/RigVM/Private/RigVMFunctions/Math/RigVMFunction_MathDouble.cpp
@@ -288,3 +288,33 @@ FRigVMFunction_MathDoubleExponential_Execute()
 	Result = FMath::Exp(Value);
 }
 
+FRigVMFunction_MathDoubleArraySum_Execute()
+{
+	Sum = 0.0;
+	for (double Value : Array)
+	{
+		Sum += Value;
+	}
+}
+
+FRigVMFunction_MathDoubleArrayAverage_Execute()
+{
+	Average = 0.0;
+	if (!Array.IsEmpty())
+	{
+		for (double Value : Array)
+		{
+			Average += Value;
+		}
+		Average = Average / Array.Num();
+	}
+	else
+	{
+		UE_RIGVMSTRUCT_REPORT_WARNING(TEXT("Array is empty"));
+	}
+}
+
+
+
+
+
diff --git a/Engine/Source/Runtime/RigVM/Private/RigVMFunctions/Math/RigVMFunction_MathFloat.cpp b/Engine/Source/Runtime/RigVM/Private/RigVMFunctions/Math/RigVMFunction_MathFloat.cpp
--- a/Engine/Source/Runtime/RigVM/Private/RigVMFunctions/Math/RigVMFunction_MathFloat.cpp
+++ b/Engine/Source/Runtime/RigVM/Private/RigVMFunctions/Math/RigVMFunction_MathFloat.cpp
@@ -299,3 +299,35 @@ FRigVMFunction_MathFloatExponential_Execute()
 	Result = FMath::Exp(Value);
 }
 
+FRigVMFunction_MathFloatArraySum_Execute()
+{
+	Sum = 0.f;
+	for (float Value : Array)
+	{
+		Sum += Value;
+	}
+}
+
+FRigVMFunction_MathFloatArrayAverage_Execute()
+{
+	Average = 0.f;
+	if (!Array.IsEmpty())
+	{
+		for (float Value : Array)
+		{
+			Average += Value;
+		}
+		Average = Average / Array.Num();
+	}
+	else
+	{
+		UE_RIGVMSTRUCT_REPORT_WARNING(TEXT("Array is empty"));
+	}
+}
+
+
+
+
+
+
+
diff --git a/Engine/Source/Runtime/RigVM/Private/RigVMFunctions/Math/RigVMFunction_MathFloatTests.cpp b/Engine/Source/Runtime/RigVM/Private/RigVMFunctions/Math/RigVMFunction_MathFloatTests.cpp
--- a/Engine/Source/Runtime/RigVM/Private/RigVMFunctions/Math/RigVMFunction_MathFloatTests.cpp
+++ b/Engine/Source/Runtime/RigVM/Private/RigVMFunctions/Math/RigVMFunction_MathFloatTests.cpp
@@ -410,4 +410,20 @@ IMPLEMENT_RIGVMSTRUCT_AUTOMATION_TEST(FRigVMFunction_MathFloatLawOfCosine)
 	return true;
 }
 
+IMPLEMENT_RIGVMSTRUCT_AUTOMATION_TEST(FRigVMFunction_MathFloatArraySum)
+{
+	Unit.Array = {1.1f, 2.2f, 3.3f};
+	Execute();
+	AddErrorIfFalse(FMath::IsNearlyEqual(Unit.Sum, 6.6f, 0.01f), TEXT("unexpected sum"));
+	return true;
+}
+
+IMPLEMENT_RIGVMSTRUCT_AUTOMATION_TEST(FRigVMFunction_MathFloatArrayAverage)
+{
+	Unit.Array = { 1.1f, 2.2f, 3.3f };
+	Execute();
+	AddErrorIfFalse(FMath::IsNearlyEqual(Unit.Average, 2.2f, 0.01f), TEXT("unexpected average"));
+	return true;
+}
+
 #endif
diff --git a/Engine/Source/Runtime/RigVM/Private/RigVMFunctions/Math/RigVMFunction_MathInt.cpp b/Engine/Source/Runtime/RigVM/Private/RigVMFunctions/Math/RigVMFunction_MathInt.cpp
--- a/Engine/Source/Runtime/RigVM/Private/RigVMFunctions/Math/RigVMFunction_MathInt.cpp
+++ b/Engine/Source/Runtime/RigVM/Private/RigVMFunctions/Math/RigVMFunction_MathInt.cpp
@@ -131,3 +131,38 @@ FRigVMFunction_MathIntLessEqual_Execute()
 	Result = A <= B;
 }
 
+FRigVMFunction_MathIntArraySum_Execute()
+{
+	Sum = 0;
+	for (int32 Value : Array)
+	{
+		Sum += Value;
+	}
+}
+
+FRigVMFunction_MathIntArrayAverage_Execute()
+{
+	Average = 0;
+	if (!Array.IsEmpty())
+	{
+		for (int32 Value : Array)
+		{
+			Average += Value;
+		}
+		Average = Average / Array.Num();
+	}
+	else
+	{
+		UE_RIGVMSTRUCT_REPORT_WARNING(TEXT("Array is empty"));
+	}
+}
+
+
+
+
+
+
+
+
+
+
diff --git a/Engine/Source/Runtime/RigVM/Private/RigVMFunctions/Math/RigVMFunction_MathVector.cpp b/Engine/Source/Runtime/RigVM/Private/RigVMFunctions/Math/RigVMFunction_MathVector.cpp
--- a/Engine/Source/Runtime/RigVM/Private/RigVMFunctions/Math/RigVMFunction_MathVector.cpp
+++ b/Engine/Source/Runtime/RigVM/Private/RigVMFunctions/Math/RigVMFunction_MathVector.cpp
@@ -525,3 +525,30 @@ FRigVMFunction_MathVectorMirrorTransform_Execute()
 	FRigVMFunction_MathTransformMirrorTransform::StaticExecute(ExecuteContext, Transform, MirrorAxis, AxisToFlip, CentralTransform, Transform);
 	Result = Transform.GetTranslation();
 }
+
+FRigVMFunction_MathVectorArraySum_Execute()
+{
+	Sum = FVector::ZeroVector;
+	for (const FVector& Value : Array)
+	{
+		Sum += Value;
+	}
+}
+
+FRigVMFunction_MathVectorArrayAverage_Execute()
+{
+	Average = FVector::ZeroVector;
+	if (!Array.IsEmpty())
+	{
+		for (const FVector& Value : Array)
+		{
+			Average += Value;
+		}
+		Average = Average / Array.Num();
+	}
+	else
+	{
+		UE_RIGVMSTRUCT_REPORT_WARNING(TEXT("Array is empty"));
+	}
+}
+
diff --git a/Engine/Source/Runtime/RigVM/Private/RigVMFunctions/Math/RigVMFunction_MathVectorTests.cpp b/Engine/Source/Runtime/RigVM/Private/RigVMFunctions/Math/RigVMFunction_MathVectorTests.cpp
--- a/Engine/Source/Runtime/RigVM/Private/RigVMFunctions/Math/RigVMFunction_MathVectorTests.cpp
+++ b/Engine/Source/Runtime/RigVM/Private/RigVMFunctions/Math/RigVMFunction_MathVectorTests.cpp
@@ -417,4 +417,20 @@ IMPLEMENT_RIGVMSTRUCT_AUTOMATION_TEST(FRigVMFunction_MathDistanceToPlane)
 	return true;
 }
 
+IMPLEMENT_RIGVMSTRUCT_AUTOMATION_TEST(FRigVMFunction_MathVectorArraySum)
+{
+	Unit.Array = { FVector(1.f, 2.f, 3.f), FVector(3.f, 2.f, 1.f), FVector(-1.f, 2.f, 3.5f) };
+	Execute();
+	AddErrorIfFalse(FRigVMFunction_MathVectorTest_Utils::IsNearlyEqual(Unit.Sum, FVector(3.f, 6.f, 7.5f), KINDA_SMALL_NUMBER), TEXT("unexpected sum"));
+	return true;
+}
+
+IMPLEMENT_RIGVMSTRUCT_AUTOMATION_TEST(FRigVMFunction_MathVectorArrayAverage)
+{
+	Unit.Array = { FVector(1.f, 2.f, 3.f), FVector(3.f, 2.f, 1.f), FVector(-1.f, 2.f, 3.5f) };
+	Execute();
+	AddErrorIfFalse(FRigVMFunction_MathVectorTest_Utils::IsNearlyEqual(Unit.Average, FVector(1.f, 2.f, 2.5f), KINDA_SMALL_NUMBER), TEXT("unexpected average"));
+	return true;
+}
+
 #endif
diff --git a/Engine/Source/Runtime/RigVM/Public/RigVMFunctions/Math/RigVMFunction_MathDouble.h b/Engine/Source/Runtime/RigVM/Public/RigVMFunctions/Math/RigVMFunction_MathDouble.h
--- a/Engine/Source/Runtime/RigVM/Public/RigVMFunctions/Math/RigVMFunction_MathDouble.h
+++ b/Engine/Source/Runtime/RigVM/Public/RigVMFunctions/Math/RigVMFunction_MathDouble.h
@@ -891,4 +891,52 @@ struct RIGVM_API FRigVMFunction_MathDoubleExponential : public FRigVMFunction_Ma
 
 	RIGVM_METHOD()
 	virtual void Execute() override;
+};
+
+/**
+ * Returns the sum of the given array
+ */
+USTRUCT(meta = (DisplayName = "Array Sum", TemplateName = "ArraySum"))
+struct RIGVM_API FRigVMFunction_MathDoubleArraySum : public FRigVMFunction_MathDoubleBase
+{
+	GENERATED_BODY()
+
+	FRigVMFunction_MathDoubleArraySum()
+	{
+		Sum = 0.0;
+	}
+
+	/** Execute logic for this rig unit */
+	RIGVM_METHOD()
+	virtual void Execute() override;
+
+	UPROPERTY(meta = (Input))
+	TArray<double> Array;
+
+	UPROPERTY(meta = (Output))
+	double Sum;
+};
+
+/**
+ * Returns the average of the given array
+ */
+USTRUCT(meta = (DisplayName = "Array Average", TemplateName = "ArrayAverage"))
+struct RIGVM_API FRigVMFunction_MathDoubleArrayAverage : public FRigVMFunction_MathDoubleBase
+{
+	GENERATED_BODY()
+
+	FRigVMFunction_MathDoubleArrayAverage()
+	{
+		Average = 0.0;
+	}
+
+	/** Execute logic for this rig unit */
+	RIGVM_METHOD()
+	virtual void Execute() override;
+
+	UPROPERTY(meta = (Input))
+	TArray<double> Array;
+
+	UPROPERTY(meta = (Output))
+	double Average;
 };
\ No newline at end of file
diff --git a/Engine/Source/Runtime/RigVM/Public/RigVMFunctions/Math/RigVMFunction_MathFloat.h b/Engine/Source/Runtime/RigVM/Public/RigVMFunctions/Math/RigVMFunction_MathFloat.h
--- a/Engine/Source/Runtime/RigVM/Public/RigVMFunctions/Math/RigVMFunction_MathFloat.h
+++ b/Engine/Source/Runtime/RigVM/Public/RigVMFunctions/Math/RigVMFunction_MathFloat.h
@@ -923,4 +923,52 @@ struct RIGVM_API FRigVMFunction_MathFloatExponential : public FRigVMFunction_Mat
 
 	RIGVM_METHOD()
 	virtual void Execute() override;
+};
+
+/**
+ * Returns the sum of the given array
+ */
+USTRUCT(meta = (DisplayName = "Array Sum", TemplateName = "ArraySum"))
+struct RIGVM_API FRigVMFunction_MathFloatArraySum : public FRigVMFunction_MathFloatBase
+{
+	GENERATED_BODY()
+
+	FRigVMFunction_MathFloatArraySum()
+	{
+		Sum = 0.f;
+	}
+
+	/** Execute logic for this rig unit */
+	RIGVM_METHOD()
+	virtual void Execute() override;
+
+	UPROPERTY(meta = (Input))
+	TArray<float> Array;
+
+	UPROPERTY(meta = (Output))
+	float Sum;
+};
+
+/**
+ * Returns the average of the given array
+ */
+USTRUCT(meta = (DisplayName = "Array Average", TemplateName = "ArrayAverage"))
+struct RIGVM_API FRigVMFunction_MathFloatArrayAverage : public FRigVMFunction_MathFloatBase
+{
+	GENERATED_BODY()
+
+	FRigVMFunction_MathFloatArrayAverage()
+	{
+		Average = 0.f;
+	}
+
+	/** Execute logic for this rig unit */
+	RIGVM_METHOD()
+	virtual void Execute() override;
+
+	UPROPERTY(meta = (Input))
+	TArray<float> Array;
+
+	UPROPERTY(meta = (Output))
+	float Average;
 };
\ No newline at end of file
diff --git a/Engine/Source/Runtime/RigVM/Public/RigVMFunctions/Math/RigVMFunction_MathInt.h b/Engine/Source/Runtime/RigVM/Public/RigVMFunctions/Math/RigVMFunction_MathInt.h
--- a/Engine/Source/Runtime/RigVM/Public/RigVMFunctions/Math/RigVMFunction_MathInt.h
+++ b/Engine/Source/Runtime/RigVM/Public/RigVMFunctions/Math/RigVMFunction_MathInt.h
@@ -485,3 +485,52 @@ struct RIGVM_API FRigVMFunction_MathIntLessEqual : public FRigVMFunction_MathInt
 	UPROPERTY(meta=(Output))
 	bool Result;
 };
+
+/**
+ * Returns the sum of the given array
+ */
+USTRUCT(meta = (DisplayName = "Array Sum", TemplateName = "ArraySum"))
+struct RIGVM_API FRigVMFunction_MathIntArraySum : public FRigVMFunction_MathIntBase
+{
+	GENERATED_BODY()
+
+	FRigVMFunction_MathIntArraySum()
+	{
+		Sum = 0;
+	}
+
+	/** Execute logic for this rig unit */
+	RIGVM_METHOD()
+	virtual void Execute() override;
+
+	UPROPERTY(meta = (Input))
+	TArray<int32> Array;
+
+	UPROPERTY(meta = (Output))
+	int32 Sum;
+};
+
+
+/**
+ * Returns the average of the given array
+ */
+USTRUCT(meta = (DisplayName = "Array Average", TemplateName = "ArrayAverage"))
+struct RIGVM_API FRigVMFunction_MathIntArrayAverage : public FRigVMFunction_MathIntBase
+{
+	GENERATED_BODY()
+
+	FRigVMFunction_MathIntArrayAverage()
+	{
+		Average = 0;
+	}
+
+	/** Execute logic for this rig unit */
+	RIGVM_METHOD()
+	virtual void Execute() override;
+
+	UPROPERTY(meta = (Input))
+	TArray<int32> Array;
+
+	UPROPERTY(meta = (Output))
+	int32 Average;
+};
\ No newline at end of file
diff --git a/Engine/Source/Runtime/RigVM/Public/RigVMFunctions/Math/RigVMFunction_MathVector.h b/Engine/Source/Runtime/RigVM/Public/RigVMFunctions/Math/RigVMFunction_MathVector.h
--- a/Engine/Source/Runtime/RigVM/Public/RigVMFunctions/Math/RigVMFunction_MathVector.h
+++ b/Engine/Source/Runtime/RigVM/Public/RigVMFunctions/Math/RigVMFunction_MathVector.h
@@ -1188,3 +1188,52 @@ struct RIGVM_API FRigVMFunction_MathVectorMirrorTransform : public FRigVMFunctio
 	UPROPERTY(meta=(Output))
 	FVector Result;
 };
+
+/**
+ * Returns the sum of the given array
+ */
+USTRUCT(meta = (DisplayName = "Array Sum", TemplateName = "ArraySum"))
+struct RIGVM_API FRigVMFunction_MathVectorArraySum : public FRigVMFunction_MathVectorBase
+{
+	GENERATED_BODY()
+
+	FRigVMFunction_MathVectorArraySum()
+	{
+		Sum = FVector::ZeroVector;
+	}
+
+	/** Execute logic for this rig unit */
+	RIGVM_METHOD()
+	virtual void Execute() override;
+
+	UPROPERTY(meta = (Input))
+	TArray<FVector> Array;
+
+	UPROPERTY(meta = (Output))
+	FVector Sum;
+};
+
+
+/**
+ * Returns the average of the given array
+ */
+USTRUCT(meta = (DisplayName = "Array Average", TemplateName = "ArrayAverage"))
+struct RIGVM_API FRigVMFunction_MathVectorArrayAverage : public FRigVMFunction_MathVectorBase
+{
+	GENERATED_BODY()
+
+	FRigVMFunction_MathVectorArrayAverage()
+	{
+		Average = FVector::ZeroVector;
+	}
+
+	/** Execute logic for this rig unit */
+	RIGVM_METHOD()
+	virtual void Execute() override;
+
+	UPROPERTY(meta = (Input))
+	TArray<FVector> Array;
+
+	UPROPERTY(meta = (Output))
+	FVector Average;
+};
\ No newline at end of file
