Control Rig: Fix crash when debug RigVM execution is on

diff --git a/Engine/Source/Runtime/RigVM/Private/RigVMCore/RigVM.cpp b/Engine/Source/Runtime/RigVM/Private/RigVMCore/RigVM.cpp
--- a/Engine/Source/Runtime/RigVM/Private/RigVMCore/RigVM.cpp
+++ b/Engine/Source/Runtime/RigVM/Private/RigVMCore/RigVM.cpp
@@ -2175,11 +2175,11 @@ ERigVMExecuteResult URigVM::ExecuteInstructions(int32 InFirstInstruction, int32
 		for (int32 PropertyIndex=0; PropertyIndex<WorkMemory->Num(); ++PropertyIndex, ++LineIndex)
 		{
 			FString Line = FString::Printf(TEXT("%s: %s"), *WorkMemory->GetProperties()[PropertyIndex]->GetFullName(), *WorkMemory->GetDataAsString(PropertyIndex));
-			if (ContextPublicData.PreviousWorkMemory.Num() > 0 && ContextPublicData.PreviousWorkMemory[PropertyIndex].StartsWith(TEXT(" -- ")))
+			if (ContextPublicData.PreviousWorkMemory.Num() > 0 && ContextPublicData.PreviousWorkMemory.IsValidIndex(PropertyIndex) && ContextPublicData.PreviousWorkMemory[PropertyIndex].StartsWith(TEXT(" -- ")))
 			{
 				ContextPublicData.PreviousWorkMemory[PropertyIndex].RightChopInline(4);
 			}
-			if (ContextPublicData.PreviousWorkMemory.Num() == 0 || Line == ContextPublicData.PreviousWorkMemory[PropertyIndex])
+			if (ContextPublicData.PreviousWorkMemory.Num() == 0 || (ContextPublicData.PreviousWorkMemory.IsValidIndex(PropertyIndex) && Line == ContextPublicData.PreviousWorkMemory[PropertyIndex]))
 			{
 				CurrentWorkMemory.Add(Line);
 			}
@@ -2193,11 +2193,11 @@ ERigVMExecuteResult URigVM::ExecuteInstructions(int32 InFirstInstruction, int32
 			FString Value;
 			ExternalVariable.Property->ExportTextItem_Direct(Value, ExternalVariable.Memory, nullptr, nullptr, PPF_None);
 			FString Line = FString::Printf(TEXT("External %s: %s"), *ExternalVariable.Name.ToString(), *Value);
-			if (ContextPublicData.PreviousWorkMemory.Num() > 0 && ContextPublicData.PreviousWorkMemory[LineIndex].StartsWith(TEXT(" -- ")))
+			if (ContextPublicData.PreviousWorkMemory.Num() > 0 && ContextPublicData.PreviousWorkMemory.IsValidIndex(LineIndex) && ContextPublicData.PreviousWorkMemory[LineIndex].StartsWith(TEXT(" -- ")))
 			{
 				ContextPublicData.PreviousWorkMemory[LineIndex].RightChopInline(4);
 			}
-			if (ContextPublicData.PreviousWorkMemory.Num() == 0 || Line == ContextPublicData.PreviousWorkMemory[LineIndex])
+			if (ContextPublicData.PreviousWorkMemory.Num() == 0 || (ContextPublicData.PreviousWorkMemory.IsValidIndex(LineIndex) && Line == ContextPublicData.PreviousWorkMemory[LineIndex]))
 			{
 				CurrentWorkMemory.Add(Line);
 			}
@@ -2208,7 +2208,7 @@ ERigVMExecuteResult URigVM::ExecuteInstructions(int32 InFirstInstruction, int32
 			++LineIndex;
 		}
 		ContextPublicData.DebugMemoryString += FString::Join(CurrentWorkMemory, TEXT("\n")) + FString(TEXT("\n\n"));
-		ContextPublicData.PreviousWorkMemory = CurrentWorkMemory;		
+		ContextPublicData.PreviousWorkMemory = CurrentWorkMemory;
 #endif
 #endif
 	}
