RigVM: Store a hash along side each graph identifying the structural state

Avoid rebuilding graphs if we don't need to. hashes contain the structure of each node, the structure of any used data type and the structure of each used template.

diff --git a/Engine/Plugins/Animation/ControlRig/Source/ControlRigDeveloper/Private/ControlRigBlueprint.cpp b/Engine/Plugins/Animation/ControlRig/Source/ControlRigDeveloper/Private/ControlRigBlueprint.cpp
--- a/Engine/Plugins/Animation/ControlRig/Source/ControlRigDeveloper/Private/ControlRigBlueprint.cpp
+++ b/Engine/Plugins/Animation/ControlRig/Source/ControlRigDeveloper/Private/ControlRigBlueprint.cpp
@@ -1319,6 +1319,32 @@ void UControlRigBlueprint::DecrementVMRecompileBracket()
 void UControlRigBlueprint::RefreshAllModels(EControlRigBlueprintLoadType InLoadType)
 {
 	const bool bIsPostLoad = InLoadType == EControlRigBlueprintLoadType::PostLoad;
+
+	// avoid any compute if the current structure hashes match with the serialized ones
+	if(RigVMClient.GetStructureHash() == RigVMClient.GetSerializedStructureHash())
+	{
+		TArray<URigVMGraph*> ModelGraphs = RigVMClient.GetAllModels(true, true);
+		Algo::Reverse(ModelGraphs);
+		for (URigVMGraph* ModelGraph : ModelGraphs)
+		{
+			URigVMController* Controller = GetOrCreateController(ModelGraph);
+			// at this stage, allow all links to be reattached,
+			// RecomputeAllTemplateFilteredPermutations() later should break any invalid links
+			Controller->ReattachLinksToPinObjects(true /* follow redirectors */, nullptr, true, true);
+
+			if(bIsPostLoad)
+			{
+				for(URigVMNode* ModelNode : ModelGraph->GetNodes())
+				{
+					if (URigVMLibraryNode* LibraryNode = ModelNode->GetTypedOuter<URigVMLibraryNode>())
+					{
+						Controller->UpdateLibraryTemplate(LibraryNode, false);
+					}
+				}
+			}
+		}
+		return;
+	}
 	
 	TGuardValue<bool> IsCompilingGuard(bIsCompiling, true);
 	TGuardValue<bool> ClientIgnoreModificationsGuard(RigVMClient.bIgnoreModelNotifications, true);
diff --git a/Engine/Source/Developer/RigVMDeveloper/Private/RigVMModel/Nodes/RigVMFunctionReferenceNode.cpp b/Engine/Source/Developer/RigVMDeveloper/Private/RigVMModel/Nodes/RigVMFunctionReferenceNode.cpp
--- a/Engine/Source/Developer/RigVMDeveloper/Private/RigVMModel/Nodes/RigVMFunctionReferenceNode.cpp
+++ b/Engine/Source/Developer/RigVMDeveloper/Private/RigVMModel/Nodes/RigVMFunctionReferenceNode.cpp
@@ -122,6 +122,43 @@ FName URigVMFunctionReferenceNode::GetOuterVariableName(const FName& InInnerVari
 	return NAME_None;
 }
 
+uint32 URigVMFunctionReferenceNode::GetStructureHash() const
+{
+	const FRigVMRegistry& Registry = FRigVMRegistry::Get();
+	
+	uint32 Hash = Super::GetStructureHash();
+
+	Hash = HashCombine(Hash, GetTypeHash(ReferencedFunctionHeader.Name.ToString()));
+	Hash = HashCombine(Hash, GetTypeHash(ReferencedFunctionHeader.NodeTitle));
+	Hash = HashCombine(Hash, GetTypeHash(ReferencedFunctionHeader.LibraryPointer.LibraryNode.ToString()));
+	Hash = HashCombine(Hash, GetTypeHash(ReferencedFunctionHeader.Keywords));
+	Hash = HashCombine(Hash, GetTypeHash(ReferencedFunctionHeader.Tooltip.ToString()));
+	Hash = HashCombine(Hash, GetTypeHash(ReferencedFunctionHeader.NodeColor));
+
+	for(const FRigVMGraphFunctionArgument& Argument : ReferencedFunctionHeader.Arguments)
+	{
+		Hash = HashCombine(Hash, GetTypeHash(Argument.Name.ToString()));
+		Hash = HashCombine(Hash, GetTypeHash((int32)Argument.Direction));
+		const TRigVMTypeIndex TypeIndex = Registry.GetTypeIndexFromCPPType(Argument.CPPType.ToString());
+		Hash = HashCombine(Hash, Registry.GetHashForType(TypeIndex));
+
+		for(const TPair<FString, FText>& Pair : Argument.PathToTooltip)
+		{
+			Hash = HashCombine(Hash, GetTypeHash(Pair.Key));
+			Hash = HashCombine(Hash, GetTypeHash(Pair.Value.ToString()));
+		}
+	}
+
+	for(const FRigVMExternalVariable& ExternalVariable : ReferencedFunctionHeader.ExternalVariables)
+	{
+		Hash = HashCombine(Hash, GetTypeHash(ExternalVariable.Name.ToString()));
+		const TRigVMTypeIndex TypeIndex = Registry.GetTypeIndexFromCPPType(ExternalVariable.TypeName.ToString());
+		Hash = HashCombine(Hash, Registry.GetHashForType(TypeIndex));
+	}
+
+	return Hash;
+}
+
 const FRigVMGraphFunctionData* URigVMFunctionReferenceNode::GetReferencedFunctionData() const
 {
 	if (IRigVMGraphFunctionHost* Host = ReferencedFunctionHeader.GetFunctionHost())
diff --git a/Engine/Source/Developer/RigVMDeveloper/Private/RigVMModel/Nodes/RigVMTemplateNode.cpp b/Engine/Source/Developer/RigVMDeveloper/Private/RigVMModel/Nodes/RigVMTemplateNode.cpp
--- a/Engine/Source/Developer/RigVMDeveloper/Private/RigVMModel/Nodes/RigVMTemplateNode.cpp
+++ b/Engine/Source/Developer/RigVMDeveloper/Private/RigVMModel/Nodes/RigVMTemplateNode.cpp
@@ -123,6 +123,18 @@ void URigVMTemplateNode::PostLoad()
 	}
 }
 
+uint32 URigVMTemplateNode::GetStructureHash() const
+{
+	uint32 Hash = Super::GetStructureHash();
+
+	if(const FRigVMTemplate* Template = GetTemplate())
+	{
+		Hash = HashCombine(Hash, GetTypeHash(*Template));
+	}
+
+	return Hash;
+}
+
 UScriptStruct* URigVMTemplateNode::GetScriptStruct() const
 {
 	if(const FRigVMFunction* Function = GetResolvedFunction())
diff --git a/Engine/Source/Developer/RigVMDeveloper/Private/RigVMModel/Nodes/RigVMUnitNode.cpp b/Engine/Source/Developer/RigVMDeveloper/Private/RigVMModel/Nodes/RigVMUnitNode.cpp
--- a/Engine/Source/Developer/RigVMDeveloper/Private/RigVMModel/Nodes/RigVMUnitNode.cpp
+++ b/Engine/Source/Developer/RigVMDeveloper/Private/RigVMModel/Nodes/RigVMUnitNode.cpp
@@ -367,3 +367,13 @@ FRigVMStructUpgradeInfo URigVMUnitNode::GetUpgradeInfo() const
 	return FRigVMStructUpgradeInfo();
 }
 
+uint32 URigVMUnitNode::GetStructureHash() const
+{
+	uint32 Hash = Super::GetStructureHash();
+	if(const UScriptStruct* ScriptStruct = GetScriptStruct())
+	{
+		Hash = HashCombine(Hash, FRigVMRegistry::Get().GetHashForScriptStruct(ScriptStruct)); 
+	}
+	return Hash;
+}
+
diff --git a/Engine/Source/Developer/RigVMDeveloper/Private/RigVMModel/RigVMClient.cpp b/Engine/Source/Developer/RigVMDeveloper/Private/RigVMModel/RigVMClient.cpp
--- a/Engine/Source/Developer/RigVMDeveloper/Private/RigVMModel/RigVMClient.cpp
+++ b/Engine/Source/Developer/RigVMDeveloper/Private/RigVMModel/RigVMClient.cpp
@@ -661,6 +661,28 @@ void FRigVMClient::DestroyObject(UObject* InObject)
 	}
 }
 
+uint32 FRigVMClient::GetStructureHash() const
+{
+	uint32 Hash = 0;
+	const TArray<URigVMGraph*> ModelsAndFunctionLibrary = GetAllModels(true, true);
+	for(const URigVMGraph* Model : ModelsAndFunctionLibrary)
+	{
+		Hash = HashCombine(Hash, Model->GetStructureHash());
+	}
+	return Hash;
+}
+
+uint32 FRigVMClient::GetSerializedStructureHash() const
+{
+	uint32 Hash = 0;
+	const TArray<URigVMGraph*> ModelsAndFunctionLibrary = GetAllModels(true, true);
+	for(const URigVMGraph* Model : ModelsAndFunctionLibrary)
+	{
+		Hash = HashCombine(Hash, Model->GetSerializedStructureHash());
+	}
+	return Hash;
+}
+
 FRigVMClientPatchResult FRigVMClient::PatchModelsOnLoad()
 {
 	FRigVMClientPatchResult Result;
diff --git a/Engine/Source/Developer/RigVMDeveloper/Private/RigVMModel/RigVMGraph.cpp b/Engine/Source/Developer/RigVMDeveloper/Private/RigVMModel/RigVMGraph.cpp
--- a/Engine/Source/Developer/RigVMDeveloper/Private/RigVMModel/RigVMGraph.cpp
+++ b/Engine/Source/Developer/RigVMDeveloper/Private/RigVMModel/RigVMGraph.cpp
@@ -7,6 +7,7 @@
 #include "RigVMModel/Nodes/RigVMFunctionEntryNode.h"
 #include "RigVMModel/Nodes/RigVMFunctionReturnNode.h"
 #include "UObject/Package.h"
+#include "UObject/ObjectSaveContext.h"
 #include "RigVMTypeUtils.h"
 
 #include UE_INLINE_GENERATED_CPP_BY_NAME(RigVMGraph)
@@ -400,6 +401,26 @@ void URigVMGraph::ClearAST(bool bClearDiagnostics, bool bClearRuntime)
 	}
 }
 
+uint32 URigVMGraph::GetStructureHash() const
+{
+	uint32 Hash = GetTypeHash(GetName());
+	for(const URigVMNode* Node : Nodes)
+	{
+		const uint32 NodeHash = Node->GetStructureHash();
+		Hash = HashCombine(Hash, NodeHash);
+	}
+	return Hash;
+}
+
+void URigVMGraph::PreSave(FObjectPreSaveContext SaveContext)
+{
+	UObject::PreSave(SaveContext);
+
+
+	// save the structure hash along side this graph
+	LastStructureHash = GetStructureHash();
+}
+
 bool URigVMGraph::IsNameAvailable(const FString& InName)
 {
 	for (URigVMNode* Node : Nodes)
diff --git a/Engine/Source/Developer/RigVMDeveloper/Private/RigVMModel/RigVMNode.cpp b/Engine/Source/Developer/RigVMDeveloper/Private/RigVMModel/RigVMNode.cpp
--- a/Engine/Source/Developer/RigVMDeveloper/Private/RigVMModel/RigVMNode.cpp
+++ b/Engine/Source/Developer/RigVMDeveloper/Private/RigVMModel/RigVMNode.cpp
@@ -380,6 +380,17 @@ bool URigVMNode::IsLinkedTo(URigVMNode* InNode) const
 	return false;
 }
 
+uint32 URigVMNode::GetStructureHash() const
+{
+	uint32 Hash = GetTypeHash(GetName());
+	for(const URigVMPin* Pin : Pins)
+	{
+		const uint32 PinHash = Pin->GetStructureHash();
+		Hash = HashCombine(Hash, PinHash);
+	}
+	return Hash;
+}
+
 URigVMLibraryNode* URigVMNode::FindFunctionForNode() 
 {
 	UObject* Subject = this;
diff --git a/Engine/Source/Developer/RigVMDeveloper/Private/RigVMModel/RigVMPin.cpp b/Engine/Source/Developer/RigVMDeveloper/Private/RigVMModel/RigVMPin.cpp
--- a/Engine/Source/Developer/RigVMDeveloper/Private/RigVMModel/RigVMPin.cpp
+++ b/Engine/Source/Developer/RigVMDeveloper/Private/RigVMModel/RigVMPin.cpp
@@ -1310,6 +1310,15 @@ bool URigVMPin::IsOrphanPin() const
 	return false;
 }
 
+uint32 URigVMPin::GetStructureHash() const
+{
+	uint32 Hash = GetTypeHash(GetName());
+	Hash = HashCombine(Hash, GetTypeHash(GetCPPType()));
+	Hash = HashCombine(Hash, GetTypeHash((int32)GetDirection()));
+	Hash = HashCombine(Hash, FRigVMRegistry::Get().GetHashForType(GetTypeIndex()));
+	return Hash;
+}
+
 void URigVMPin::UpdateTypeInformationIfRequired() const
 {
 	if (CPPTypeObject == nullptr)
diff --git a/Engine/Source/Developer/RigVMDeveloper/Public/RigVMModel/Nodes/RigVMFunctionReferenceNode.h b/Engine/Source/Developer/RigVMDeveloper/Public/RigVMModel/Nodes/RigVMFunctionReferenceNode.h
--- a/Engine/Source/Developer/RigVMDeveloper/Public/RigVMModel/Nodes/RigVMFunctionReferenceNode.h
+++ b/Engine/Source/Developer/RigVMDeveloper/Public/RigVMModel/Nodes/RigVMFunctionReferenceNode.h
@@ -46,6 +46,7 @@ public:
 	FName GetOuterVariableName(const FName& InInnerVariableName) const;
 	// end Variable remapping
 
+	virtual uint32 GetStructureHash() const override;
 	const FRigVMGraphFunctionHeader& GetReferencedFunctionHeader() const { return ReferencedFunctionHeader; }
 
 	const FRigVMGraphFunctionData* GetReferencedFunctionData() const;
diff --git a/Engine/Source/Developer/RigVMDeveloper/Public/RigVMModel/Nodes/RigVMTemplateNode.h b/Engine/Source/Developer/RigVMDeveloper/Public/RigVMModel/Nodes/RigVMTemplateNode.h
--- a/Engine/Source/Developer/RigVMDeveloper/Public/RigVMModel/Nodes/RigVMTemplateNode.h
+++ b/Engine/Source/Developer/RigVMDeveloper/Public/RigVMModel/Nodes/RigVMTemplateNode.h
@@ -139,6 +139,8 @@ public:
 	// Will prioritize the InPreferredType if available
 	TRigVMTypeIndex TryReduceTypesToSingle(const TArray<TRigVMTypeIndex>& InTypes, const TRigVMTypeIndex PreferredType = TRigVMTypeIndex()) const;
 
+	virtual uint32 GetStructureHash() const override;
+
 protected:
 
 	virtual void InvalidateCache() override;
diff --git a/Engine/Source/Developer/RigVMDeveloper/Public/RigVMModel/Nodes/RigVMUnitNode.h b/Engine/Source/Developer/RigVMDeveloper/Public/RigVMModel/Nodes/RigVMUnitNode.h
--- a/Engine/Source/Developer/RigVMDeveloper/Public/RigVMModel/Nodes/RigVMUnitNode.h
+++ b/Engine/Source/Developer/RigVMDeveloper/Public/RigVMModel/Nodes/RigVMUnitNode.h
@@ -74,6 +74,8 @@ public:
 
 	virtual FRigVMStructUpgradeInfo GetUpgradeInfo() const override;
 
+	virtual uint32 GetStructureHash() const override;
+
 protected:
 
 	virtual FText GetToolTipTextForPin(const URigVMPin* InPin) const override;
diff --git a/Engine/Source/Developer/RigVMDeveloper/Public/RigVMModel/RigVMClient.h b/Engine/Source/Developer/RigVMDeveloper/Public/RigVMModel/RigVMClient.h
--- a/Engine/Source/Developer/RigVMDeveloper/Public/RigVMModel/RigVMClient.h
+++ b/Engine/Source/Developer/RigVMDeveloper/Public/RigVMModel/RigVMClient.h
@@ -139,6 +139,9 @@ public:
 	static FName GetUniqueName(UObject* InOuter, const FName& InDesiredName);
 	static void DestroyObject(UObject* InObject);
 
+	uint32 GetStructureHash() const;
+	uint32 GetSerializedStructureHash() const;
+
 	// backwards compatibility
 	FRigVMClientPatchResult PatchModelsOnLoad();
 
diff --git a/Engine/Source/Developer/RigVMDeveloper/Public/RigVMModel/RigVMGraph.h b/Engine/Source/Developer/RigVMDeveloper/Public/RigVMModel/RigVMGraph.h
--- a/Engine/Source/Developer/RigVMDeveloper/Public/RigVMModel/RigVMGraph.h
+++ b/Engine/Source/Developer/RigVMDeveloper/Public/RigVMModel/RigVMGraph.h
@@ -146,6 +146,11 @@ public:
 	TSharedPtr<FRigVMParserAST> GetRuntimeAST(const FRigVMParserASTSettings& InSettings = FRigVMParserASTSettings::Optimized(), bool bForceRefresh = false);
 	void ClearAST(bool bClearDiagnostics = true, bool bClearRuntime = true);
 
+	virtual uint32 GetStructureHash() const;
+	uint32 GetSerializedStructureHash() const { return LastStructureHash; }
+
+	virtual void PreSave(FObjectPreSaveContext SaveContext) override;
+
 private:
 
 	FRigVMGraphModifiedEvent ModifiedEvent;
@@ -174,6 +179,9 @@ private:
 	TArray<TSharedPtr<FString>> ParameterNames;
 #endif
 
+	UPROPERTY()
+	uint32 LastStructureHash;
+
 	UPROPERTY()
 	bool bEditable;
 
diff --git a/Engine/Source/Developer/RigVMDeveloper/Public/RigVMModel/RigVMNode.h b/Engine/Source/Developer/RigVMDeveloper/Public/RigVMModel/RigVMNode.h
--- a/Engine/Source/Developer/RigVMDeveloper/Public/RigVMModel/RigVMNode.h
+++ b/Engine/Source/Developer/RigVMDeveloper/Public/RigVMModel/RigVMNode.h
@@ -273,6 +273,8 @@ public:
 
 	virtual FRigVMStructUpgradeInfo GetUpgradeInfo() const { return FRigVMStructUpgradeInfo(); }
 
+	virtual uint32 GetStructureHash() const;
+
 private:
 
 	static const FString NodeColorName;
diff --git a/Engine/Source/Developer/RigVMDeveloper/Public/RigVMModel/RigVMPin.h b/Engine/Source/Developer/RigVMDeveloper/Public/RigVMModel/RigVMPin.h
--- a/Engine/Source/Developer/RigVMDeveloper/Public/RigVMModel/RigVMPin.h
+++ b/Engine/Source/Developer/RigVMDeveloper/Public/RigVMModel/RigVMPin.h
@@ -447,6 +447,8 @@ public:
 	// Returns true if the pin has been orphaned
 	bool IsOrphanPin() const;
 
+	uint32 GetStructureHash() const;
+
 private:
 
 	void UpdateTypeInformationIfRequired() const;
diff --git a/Engine/Source/Runtime/RigVM/Private/RigVMCore/RigVMRegistry.cpp b/Engine/Source/Runtime/RigVM/Private/RigVMCore/RigVMRegistry.cpp
--- a/Engine/Source/Runtime/RigVM/Private/RigVMCore/RigVMRegistry.cpp
+++ b/Engine/Source/Runtime/RigVM/Private/RigVMCore/RigVMRegistry.cpp
@@ -76,6 +76,129 @@ const TArray<UScriptStruct*>& FRigVMRegistry::GetMathTypes()
 	return MathTypes;
 }
 
+uint32 FRigVMRegistry::GetHashForType(TRigVMTypeIndex InTypeIndex) const
+{
+	if(!Types.IsValidIndex(InTypeIndex))
+	{
+		return UINT32_MAX;
+	}
+
+	FRigVMRegistry* MutableThis = (FRigVMRegistry*)this; 
+	FTypeInfo& TypeInfo = MutableThis->Types[InTypeIndex];
+	
+	if(TypeInfo.Hash != UINT32_MAX)
+	{
+		return TypeInfo.Hash;
+	}
+
+	uint32 Hash = INDEX_NONE;
+	if(const UScriptStruct* ScriptStruct = Cast<UScriptStruct>(TypeInfo.Type.CPPTypeObject))
+	{
+		Hash = GetHashForScriptStruct(ScriptStruct, false);
+	}
+	else if(const UStruct* Struct = Cast<UStruct>(TypeInfo.Type.CPPTypeObject))
+	{
+		Hash = GetHashForStruct(Struct);
+	}
+	else if(const UEnum* Enum = Cast<UEnum>(TypeInfo.Type.CPPTypeObject))
+    {
+    	Hash = GetHashForEnum(Enum, false);
+    }
+    else
+    {
+    	Hash = GetTypeHash(TypeInfo.Type.CPPType.ToString());
+    }
+
+	// for used defined structs - always recompute it
+	if(Cast<UUserDefinedStruct>(TypeInfo.Type.CPPTypeObject))
+	{
+		return Hash;
+	}
+
+	TypeInfo.Hash = Hash;
+	return Hash;
+}
+
+uint32 FRigVMRegistry::GetHashForScriptStruct(const UScriptStruct* InScriptStruct, bool bCheckTypeIndex) const
+{
+	if(bCheckTypeIndex)
+	{
+		const TRigVMTypeIndex TypeIndex = GetTypeIndex(*InScriptStruct->GetStructCPPName(), (UObject*)InScriptStruct);
+		if(TypeIndex != INDEX_NONE)
+		{
+			return GetHashForType(TypeIndex);
+		}
+	}
+	
+	const uint32 NameHash = GetTypeHash(InScriptStruct->GetStructCPPName());
+	return HashCombine(NameHash, GetHashForStruct(InScriptStruct));
+}
+
+uint32 FRigVMRegistry::GetHashForStruct(const UStruct* InStruct) const
+{
+	uint32 Hash = GetTypeHash(InStruct->GetPathName());
+	for (TFieldIterator<FProperty> It(InStruct); It; ++It)
+	{
+		const FProperty* Property = *It;
+		if(IsAllowedType(Property))
+		{
+			Hash = HashCombine(Hash, GetHashForProperty(Property));
+		}
+	}
+	return Hash;
+}
+
+uint32 FRigVMRegistry::GetHashForEnum(const UEnum* InEnum, bool bCheckTypeIndex) const
+{
+	if(bCheckTypeIndex)
+	{
+		const TRigVMTypeIndex TypeIndex = GetTypeIndex(*InEnum->CppType, (UObject*)InEnum);
+		if(TypeIndex != INDEX_NONE)
+		{
+			return GetHashForType(TypeIndex);
+		}
+	}
+	
+	uint32 Hash = GetTypeHash(InEnum->GetName());
+	for(int32 Index = 0; Index < InEnum->NumEnums(); Index++)
+	{
+		Hash = HashCombine(Hash, GetTypeHash(InEnum->GetValueByIndex(Index)));
+		Hash = HashCombine(Hash, GetTypeHash(InEnum->GetDisplayNameTextByIndex(Index).ToString()));
+	}
+	return Hash;
+}
+
+uint32 FRigVMRegistry::GetHashForProperty(const FProperty* InProperty) const
+{
+	uint32 Hash = GetTypeHash(InProperty->GetName());
+
+	FString ExtendedCPPType;
+	const FString CPPType = InProperty->GetCPPType(&ExtendedCPPType);
+	Hash = HashCombine(Hash, GetTypeHash(CPPType + ExtendedCPPType));
+	
+	if(const FArrayProperty* ArrayProperty = CastField<FArrayProperty>(InProperty))
+	{
+		InProperty = ArrayProperty->Inner;
+	}
+	
+	if(const FStructProperty* StructProperty = CastField<FStructProperty>(InProperty))
+	{
+		Hash = HashCombine(Hash, GetHashForStruct(StructProperty->Struct));
+	}
+	else if(const FByteProperty* ByteProperty = CastField<FByteProperty>(InProperty))
+	{
+		if(ByteProperty->Enum)
+		{
+			Hash = HashCombine(Hash, GetHashForEnum(ByteProperty->Enum));
+		}
+	}
+	else if(const FEnumProperty* EnumProperty = CastField<FEnumProperty>(InProperty))
+	{
+		Hash = HashCombine(Hash, GetHashForEnum(EnumProperty->GetEnum()));
+	}
+	
+	return Hash;
+}
 
 void FRigVMRegistry::InitializeIfNeeded()
 {
diff --git a/Engine/Source/Runtime/RigVM/Private/RigVMCore/RigVMTemplate.cpp b/Engine/Source/Runtime/RigVM/Private/RigVMCore/RigVMTemplate.cpp
--- a/Engine/Source/Runtime/RigVM/Private/RigVMCore/RigVMTemplate.cpp
+++ b/Engine/Source/Runtime/RigVM/Private/RigVMCore/RigVMTemplate.cpp
@@ -393,6 +393,7 @@ TArray<FString> FRigVMTemplateArgument::GetSupportedTypeStrings(const TArray<int
 FRigVMTemplate::FRigVMTemplate()
 	: Index(INDEX_NONE)
 	, Notation(NAME_None)
+	, Hash(UINT32_MAX)
 {
 
 }
@@ -432,6 +433,7 @@ void FRigVMTemplate::Serialize(FArchive& Ar)
 FRigVMTemplate::FRigVMTemplate(UScriptStruct* InStruct, const FString& InTemplateName, int32 InFunctionIndex)
 	: Index(INDEX_NONE)
 	, Notation(NAME_None)
+	, Hash(UINT32_MAX)
 {
 	TArray<FString> ArgumentNotations;
 
@@ -476,6 +478,7 @@ FRigVMTemplate::FRigVMTemplate(UScriptStruct* InStruct, const FString& InTemplat
 FRigVMTemplate::FRigVMTemplate(const FName& InTemplateName, const TArray<FRigVMTemplateArgument>& InArguments, int32 InFunctionIndex)
 	: Index(INDEX_NONE)
 	, Notation(NAME_None)
+	, Hash(UINT32_MAX)
 {
 	TArray<FString> ArgumentNotations;
 	for (const FRigVMTemplateArgument& InArgument : InArguments)
@@ -1533,6 +1536,8 @@ FString FRigVMTemplate::GetKeywords() const
 
 bool FRigVMTemplate::AddTypeForArgument(const FName& InArgumentName, TRigVMTypeIndex InTypeIndex)
 {
+	InvalidateHash();
+	
 	if(OnNewArgumentType().IsBound())
 	{
 		FRigVMTemplateTypeMap Types = OnNewArgumentType().Execute(this, InArgumentName, InTypeIndex);
@@ -1593,6 +1598,8 @@ bool FRigVMTemplate::AddTypeForArgument(const FName& InArgumentName, TRigVMTypeI
 
 void FRigVMTemplate::HandleTypeRemoval(TRigVMTypeIndex InTypeIndex)
 {
+	InvalidateHash();
+
 	TArray<int32> PermutationsToRemove;
 	for (int32 PermutationIndex = 0; PermutationIndex < NumPermutations(); PermutationIndex++)
 	{
@@ -1627,3 +1634,39 @@ void FRigVMTemplate::HandleTypeRemoval(TRigVMTypeIndex InTypeIndex)
 }
 
 
+uint32 GetTypeHash(const FRigVMTemplateArgument& InArgument)
+{
+	const FRigVMRegistry& Registry = FRigVMRegistry::Get();
+	uint32 Hash = GetTypeHash(InArgument.Name.ToString());
+	Hash = HashCombine(Hash, GetTypeHash((int32)InArgument.Direction));
+	for(const TRigVMTypeIndex& TypeIndex : InArgument.TypeIndices)
+	{
+		Hash = HashCombine(Hash, Registry.GetHashForType(TypeIndex));
+	}
+	return Hash;
+}
+
+uint32 GetTypeHash(const FRigVMTemplate& InTemplate)
+{
+	if(InTemplate.Hash != UINT32_MAX)
+	{
+		return InTemplate.Hash;
+	}
+
+	uint32 Hash = GetTypeHash(InTemplate.GetNotation().ToString());
+	for(const FRigVMTemplateArgument& Argument : InTemplate.Arguments)
+	{
+		Hash = HashCombine(Hash, GetTypeHash(Argument));
+	}
+
+	// todo: in Dev-EngineMerge we should add the execute arguments to the hash as well
+
+	if(const FRigVMDispatchFactory* Factory = InTemplate.GetDispatchFactory())
+	{
+		Hash = HashCombine(Hash, GetTypeHash(Factory->GetFactoryName().ToString()));
+	}
+
+	InTemplate.Hash = Hash;
+	return Hash;
+}
+
diff --git a/Engine/Source/Runtime/RigVM/Public/RigVMCore/RigVMRegistry.h b/Engine/Source/Runtime/RigVM/Public/RigVMCore/RigVMRegistry.h
--- a/Engine/Source/Runtime/RigVM/Public/RigVMCore/RigVMRegistry.h
+++ b/Engine/Source/Runtime/RigVM/Public/RigVMCore/RigVMRegistry.h
@@ -259,6 +259,13 @@ public:
 	// Notifies other system that types have been added/removed, and template permutations have been updated
 	FOnRigVMRegistryChanged& OnRigVMRegistryChanged() { return OnRigVMRegistryChangedDelegate; }
 
+	// Returns a unique hash per type index
+	uint32 GetHashForType(TRigVMTypeIndex InTypeIndex) const;
+	uint32 GetHashForScriptStruct(const UScriptStruct* InScriptStruct, bool bCheckTypeIndex = true) const;
+	uint32 GetHashForStruct(const UStruct* InStruct) const;
+	uint32 GetHashForEnum(const UEnum* InEnum, bool bCheckTypeIndex = true) const;
+	uint32 GetHashForProperty(const FProperty* InProperty) const;
+
 private:
 
 	static const FName TemplateNameMetaName;
@@ -278,6 +285,7 @@ private:
 			, ArrayTypeIndex(INDEX_NONE)
 			, bIsArray(false)
 			, bIsExecute(false)
+			, Hash(UINT32_MAX)
 		{}
 		
 		FRigVMTemplateArgumentType Type;
@@ -285,6 +293,7 @@ private:
 		TRigVMTypeIndex ArrayTypeIndex;
 		bool bIsArray;
 		bool bIsExecute;
+		uint32 Hash;
 	};
 
 	TRigVMTypeIndex FindOrAddType_Internal(const FRigVMTemplateArgumentType& InType, bool bForce);
diff --git a/Engine/Source/Runtime/RigVM/Public/RigVMCore/RigVMTemplate.h b/Engine/Source/Runtime/RigVM/Public/RigVMCore/RigVMTemplate.h
--- a/Engine/Source/Runtime/RigVM/Public/RigVMCore/RigVMTemplate.h
+++ b/Engine/Source/Runtime/RigVM/Public/RigVMCore/RigVMTemplate.h
@@ -107,7 +107,7 @@ struct RIGVM_API FRigVMTemplateArgumentType
 		return CPPType != InOther.CPPType;
 	}
 
-	friend uint32 GetTypeHash(const FRigVMTemplateArgumentType& InType)
+	RIGVM_API friend FORCEINLINE uint32 GetTypeHash(const FRigVMTemplateArgumentType& InType)
 	{
 		return GetTypeHash(InType.CPPType);
 	}
@@ -235,6 +235,8 @@ struct RIGVM_API FRigVMTemplateArgument
 	// returns true if the argument uses an array container
 	EArrayType GetArrayType() const;
 
+	RIGVM_API friend uint32 GetTypeHash(const FRigVMTemplateArgument& InArgument);
+
 protected:
 
 	int32 Index;
@@ -417,6 +419,8 @@ public:
 			Delegates.GetDispatchFactoryDelegate.IsBound();
 	}  
 
+	RIGVM_API friend uint32 GetTypeHash(const FRigVMTemplate& InTemplate);
+
 private:
 
 	// Constructor from a struct, a template name and a function index
@@ -427,6 +431,7 @@ private:
 
 	static FLinearColor GetColorFromMetadata(FString InMetadata);
 
+	void InvalidateHash() { Hash = UINT32_MAX; }
 	const TArray<FRigVMExecuteArgument>& GetExecuteArguments(const FRigVMDispatchContext& InContext) const;
 
 	int32 Index;
@@ -434,6 +439,7 @@ private:
 	TArray<FRigVMTemplateArgument> Arguments;
 	mutable TArray<FRigVMExecuteArgument> ExecuteArguments;
 	TArray<int32> Permutations;
+	mutable uint32 Hash;
 
 	FRigVMTemplateDelegates Delegates;
 
