Control Rig: Only report one instance of property copy failure. Otherwise the log output might get overwhelmed.

#rnx

diff --git a/Engine/Source/Runtime/RigVM/Private/RigVMCore/RigVMMemoryStorage.cpp b/Engine/Source/Runtime/RigVM/Private/RigVMCore/RigVMMemoryStorage.cpp
--- a/Engine/Source/Runtime/RigVM/Private/RigVMCore/RigVMMemoryStorage.cpp
+++ b/Engine/Source/Runtime/RigVM/Private/RigVMCore/RigVMMemoryStorage.cpp
@@ -1099,20 +1099,30 @@ bool URigVMMemoryStorage::CopyProperty(
 		// between two properties which are not compatible
 		if(!InTargetProperty->SameType(InSourceProperty))
 		{
+			// Only log the issue once, rather than spam.
+			static TSet<FString> ReportedErrors;
+			
 			FString TargetType, SourceType, TargetExtendedType, SourceExtendedType;
 			TargetType = InTargetProperty->GetCPPType(&TargetExtendedType);
 			SourceType = InSourceProperty->GetCPPType(&SourceExtendedType);
 			TargetType += TargetExtendedType;
 			SourceType += SourceExtendedType;
-
+			
 			UPackage* Package = InTargetProperty->GetOutermost();
-			check(Package);
-			UE_LOG(LogRigVM, Warning, TEXT("Failed to copy %s (%s) to %s (%s) in package %s"),
+			
+			FString Message = FString::Printf(TEXT("Failed to copy %s (%s) to %s (%s) in package %s"),
 				*InSourceProperty->GetName(),
 				*SourceType,
 				*InTargetProperty->GetName(), 
 				*TargetType,
-				*Package->GetName());
+				Package ? *Package->GetName() : TEXT("<Unknown Package>"));
+
+			if (!ReportedErrors.Contains(Message))
+			{
+				UE_LOG(LogRigVM, Warning, TEXT("%s"), *Message);
+				ReportedErrors.Add(Message);
+			}
+			
 			return false;
 		}
 	}
