Control Rig: Get argument names separate from the argument definition, since generating the definition can be slow.

#rnx

diff --git a/Engine/Source/Runtime/RigVM/Private/RigVMFunctions/RigVMDispatch_Array.cpp b/Engine/Source/Runtime/RigVM/Private/RigVMFunctions/RigVMDispatch_Array.cpp
--- a/Engine/Source/Runtime/RigVM/Private/RigVMFunctions/RigVMDispatch_Array.cpp
+++ b/Engine/Source/Runtime/RigVM/Private/RigVMFunctions/RigVMDispatch_Array.cpp
@@ -278,7 +278,8 @@ FName FRigVMDispatch_ArrayMake::GetArgumentNameForOperandIndex(int32 InOperandIn
 	{
 		return FRigVMBranchInfo::GetFixedArrayLabel(ValuesName, *FString::FromInt(InOperandIndex - 1));
 	}
-	return FRigVMDispatch_ArrayBase::GetArgumentNameForOperandIndex(1, 2);
+	check(InOperandIndex == (InTotalOperands - 1));
+	return ArrayName;
 }
 
 void FRigVMDispatch_ArrayMake::Execute(FRigVMExtendedExecuteContext& InContext, FRigVMMemoryHandleArray Handles)
@@ -306,6 +307,12 @@ void FRigVMDispatch_ArrayMake::Execute(FRigVMExtendedExecuteContext& InContext,
 
 ////////////////////////////////////////////////////////////////////////////////////////////////////////////////////////
 
+FName FRigVMDispatch_ArrayReset::GetArgumentNameForOperandIndex(int32 InOperandIndex, int32 InTotalOperands) const
+{
+	check(InTotalOperands == 1);
+	return ArrayName;
+}
+
 TArray<FRigVMTemplateArgument> FRigVMDispatch_ArrayReset::GetArguments() const
 {
 	return {
@@ -348,6 +355,16 @@ FText FRigVMDispatch_ArrayReset::GetArgumentTooltip(const FName& InArgumentName,
 
 ////////////////////////////////////////////////////////////////////////////////////////////////////////////////////////
 
+FName FRigVMDispatch_ArrayGetNum::GetArgumentNameForOperandIndex(int32 InOperandIndex, int32 InTotalOperands) const
+{
+	static const FName ArgumentNames[] = {
+		ArrayName,
+		NumName
+	};
+	check(InTotalOperands == UE_ARRAY_COUNT(ArgumentNames));
+	return ArgumentNames[InOperandIndex];
+}
+
 TArray<FRigVMTemplateArgument> FRigVMDispatch_ArrayGetNum::GetArguments() const
 {
 	return {
@@ -397,6 +414,16 @@ FText FRigVMDispatch_ArrayGetNum::GetArgumentTooltip(const FName& InArgumentName
 
 ////////////////////////////////////////////////////////////////////////////////////////////////////////////////////////
 
+FName FRigVMDispatch_ArraySetNum::GetArgumentNameForOperandIndex(int32 InOperandIndex, int32 InTotalOperands) const
+{
+	static const FName ArgumentNames[] = {
+		ArrayName,
+		NumName
+	};
+	check(InTotalOperands == UE_ARRAY_COUNT(ArgumentNames));
+	return ArgumentNames[InOperandIndex];
+}
+
 TArray<FRigVMTemplateArgument> FRigVMDispatch_ArraySetNum::GetArguments() const
 {
 	return {
@@ -449,6 +476,17 @@ FText FRigVMDispatch_ArraySetNum::GetArgumentTooltip(const FName& InArgumentName
 
 ////////////////////////////////////////////////////////////////////////////////////////////////////////////////////////
 
+FName FRigVMDispatch_ArrayGetAtIndex::GetArgumentNameForOperandIndex(int32 InOperandIndex, int32 InTotalOperands) const
+{
+	static const FName ArgumentNames[] = {
+		ArrayName,
+		IndexName,
+		ElementName
+	};
+	check(InTotalOperands == UE_ARRAY_COUNT(ArgumentNames));
+	return ArgumentNames[InOperandIndex];
+}
+
 TArray<FRigVMTemplateArgument> FRigVMDispatch_ArrayGetAtIndex::GetArguments() const
 {
 	return {
@@ -520,6 +558,17 @@ FText FRigVMDispatch_ArrayGetAtIndex::GetArgumentTooltip(const FName& InArgument
 
 ////////////////////////////////////////////////////////////////////////////////////////////////////////////////////////
 
+FName FRigVMDispatch_ArraySetAtIndex::GetArgumentNameForOperandIndex(int32 InOperandIndex, int32 InTotalOperands) const
+{
+	static const FName ArgumentNames[] = {
+		ArrayName,
+		IndexName,
+		ElementName
+	};
+	check(InTotalOperands == UE_ARRAY_COUNT(ArgumentNames));
+	return ArgumentNames[InOperandIndex];
+}
+
 TArray<FRigVMTemplateArgument> FRigVMDispatch_ArraySetAtIndex::GetArguments() const
 {
 	return {
@@ -591,6 +640,17 @@ FText FRigVMDispatch_ArraySetAtIndex::GetArgumentTooltip(const FName& InArgument
 
 ////////////////////////////////////////////////////////////////////////////////////////////////////////////////////////
 
+FName FRigVMDispatch_ArrayAdd::GetArgumentNameForOperandIndex(int32 InOperandIndex, int32 InTotalOperands) const
+{
+	static const FName ArgumentNames[] = {
+		ArrayName,
+		ElementName,
+		IndexName
+	};
+	check(InTotalOperands == UE_ARRAY_COUNT(ArgumentNames));
+	return ArgumentNames[InOperandIndex];
+}
+
 TArray<FRigVMTemplateArgument> FRigVMDispatch_ArrayAdd::GetArguments() const
 {
 	return {
@@ -698,6 +758,16 @@ FText FRigVMDispatch_ArrayInsert::GetArgumentTooltip(const FName& InArgumentName
 
 ////////////////////////////////////////////////////////////////////////////////////////////////////////////////////////
 
+FName FRigVMDispatch_ArrayRemove::GetArgumentNameForOperandIndex(int32 InOperandIndex, int32 InTotalOperands) const
+{
+	static const FName ArgumentNames[] = {
+		ArrayName,
+		IndexName
+	};
+	check(InTotalOperands == UE_ARRAY_COUNT(ArgumentNames));
+	return ArgumentNames[InOperandIndex];
+}
+
 TArray<FRigVMTemplateArgument> FRigVMDispatch_ArrayRemove::GetArguments() const
 {
 	return {
@@ -782,6 +852,18 @@ FText FRigVMDispatch_ArrayReverse::GetArgumentTooltip(const FName& InArgumentNam
 
 ////////////////////////////////////////////////////////////////////////////////////////////////////////////////////////
 
+FName FRigVMDispatch_ArrayFind::GetArgumentNameForOperandIndex(int32 InOperandIndex, int32 InTotalOperands) const
+{
+	static const FName ArgumentNames[] = {
+		ArrayName,
+		ElementName,
+		IndexName,
+		SuccessName
+	};
+	check(InTotalOperands == UE_ARRAY_COUNT(ArgumentNames));
+	return ArgumentNames[InOperandIndex];
+}
+
 TArray<FRigVMTemplateArgument> FRigVMDispatch_ArrayFind::GetArguments() const
 {
 	return {
@@ -881,6 +963,16 @@ FText FRigVMDispatch_ArrayFind::GetArgumentTooltip(const FName& InArgumentName,
 
 ////////////////////////////////////////////////////////////////////////////////////////////////////////////////////////
 
+FName FRigVMDispatch_ArrayAppend::GetArgumentNameForOperandIndex(int32 InOperandIndex, int32 InTotalOperands) const
+{
+	static const FName ArgumentNames[] = {
+		ArrayName,
+		OtherName
+	};
+	check(InTotalOperands == UE_ARRAY_COUNT(ArgumentNames));
+	return ArgumentNames[InOperandIndex];
+}
+
 TArray<FRigVMTemplateArgument> FRigVMDispatch_ArrayAppend::GetArguments() const
 {
 	return {
@@ -949,6 +1041,16 @@ FText FRigVMDispatch_ArrayAppend::GetArgumentTooltip(const FName& InArgumentName
 
 ////////////////////////////////////////////////////////////////////////////////////////////////////////////////////////
 
+FName FRigVMDispatch_ArrayClone::GetArgumentNameForOperandIndex(int32 InOperandIndex, int32 InTotalOperands) const
+{
+	static const FName ArgumentNames[] = {
+		ArrayName,
+		CloneName
+	};
+	check(InTotalOperands == UE_ARRAY_COUNT(ArgumentNames));
+	return ArgumentNames[InOperandIndex];
+}
+
 TArray<FRigVMTemplateArgument> FRigVMDispatch_ArrayClone::GetArguments() const
 {
 	return {
@@ -1102,6 +1204,17 @@ FText FRigVMDispatch_ArrayUnion::GetArgumentTooltip(const FName& InArgumentName,
 
 ////////////////////////////////////////////////////////////////////////////////////////////////////////////////////////
 
+FName FRigVMDispatch_ArrayDifference::GetArgumentNameForOperandIndex(int32 InOperandIndex, int32 InTotalOperands) const
+{
+	static const FName ArgumentNames[] = {
+		ArrayName,
+		OtherName,
+		ResultName
+	};
+	check(InTotalOperands == UE_ARRAY_COUNT(ArgumentNames));
+	return ArgumentNames[InOperandIndex];
+}
+
 TArray<FRigVMTemplateArgument> FRigVMDispatch_ArrayDifference::GetArguments() const
 {
 	return {
@@ -1256,6 +1369,20 @@ FText FRigVMDispatch_ArrayIntersection::GetArgumentTooltip(const FName& InArgume
 
 ////////////////////////////////////////////////////////////////////////////////////////////////////////////////////////
 
+FName FRigVMDispatch_ArrayIterator::GetArgumentNameForOperandIndex(int32 InOperandIndex, int32 InTotalOperands) const
+{
+	static const FName ArgumentNames[] = {
+		ArrayName,
+		ElementName,
+		IndexName,
+		CountName,
+		RatioName,
+		FRigVMStruct::ControlFlowBlockToRunName
+	};
+	check(InTotalOperands == UE_ARRAY_COUNT(ArgumentNames));
+	return ArgumentNames[InOperandIndex];
+}
+
 TArray<FRigVMTemplateArgument> FRigVMDispatch_ArrayIterator::GetArguments() const
 {
 	return {
diff --git a/Engine/Source/Runtime/RigVM/Private/RigVMFunctions/RigVMDispatch_Core.cpp b/Engine/Source/Runtime/RigVM/Private/RigVMFunctions/RigVMDispatch_Core.cpp
--- a/Engine/Source/Runtime/RigVM/Private/RigVMFunctions/RigVMDispatch_Core.cpp
+++ b/Engine/Source/Runtime/RigVM/Private/RigVMFunctions/RigVMDispatch_Core.cpp
@@ -6,6 +6,23 @@
 
 #include UE_INLINE_GENERATED_CPP_BY_NAME(RigVMDispatch_Core)
 
+
+const FName FRigVMDispatch_CoreEquals::AName = TEXT("A");
+const FName FRigVMDispatch_CoreEquals::BName = TEXT("B");
+const FName FRigVMDispatch_CoreEquals::ResultName = TEXT("Result");
+
+
+FName FRigVMDispatch_CoreEquals::GetArgumentNameForOperandIndex(int32 InOperandIndex, int32 InTotalOperands) const
+{
+	static const FName ArgumentNames[] = {
+		AName,
+		BName,
+		ResultName
+	};
+	check(InTotalOperands == UE_ARRAY_COUNT(ArgumentNames));
+	return ArgumentNames[InOperandIndex];
+}
+
 TArray<FRigVMTemplateArgument> FRigVMDispatch_CoreEquals::GetArguments() const
 {
 	const TArray<FRigVMTemplateArgument::ETypeCategory> ValueCategories = {
@@ -13,9 +30,9 @@ TArray<FRigVMTemplateArgument> FRigVMDispatch_CoreEquals::GetArguments() const
 		FRigVMTemplateArgument::ETypeCategory_ArrayAnyValue
 	};
 	return {
-		FRigVMTemplateArgument(TEXT("A"), ERigVMPinDirection::Input, ValueCategories),
-		FRigVMTemplateArgument(TEXT("B"), ERigVMPinDirection::Input, ValueCategories),
-		FRigVMTemplateArgument(TEXT("Result"), ERigVMPinDirection::Output, RigVMTypeUtils::TypeIndex::Bool)
+		FRigVMTemplateArgument(AName, ERigVMPinDirection::Input, ValueCategories),
+		FRigVMTemplateArgument(BName, ERigVMPinDirection::Input, ValueCategories),
+		FRigVMTemplateArgument(ResultName, ERigVMPinDirection::Output, RigVMTypeUtils::TypeIndex::Bool)
 	};
 }
 
@@ -23,17 +40,17 @@ FRigVMTemplateTypeMap FRigVMDispatch_CoreEquals::OnNewArgumentType(const FName&
 	TRigVMTypeIndex InTypeIndex) const
 {
 	FRigVMTemplateTypeMap Types;
-	Types.Add(TEXT("A"), InTypeIndex);
-	Types.Add(TEXT("B"), InTypeIndex);
-	Types.Add(TEXT("Result"), RigVMTypeUtils::TypeIndex::Bool);
+	Types.Add(AName, InTypeIndex);
+	Types.Add(BName, InTypeIndex);
+	Types.Add(ResultName, RigVMTypeUtils::TypeIndex::Bool);
 	return Types;
 }
 
 FRigVMFunctionPtr FRigVMDispatch_CoreEquals::GetDispatchFunctionImpl(const FRigVMTemplateTypeMap& InTypes) const
 {
-	const TRigVMTypeIndex TypeIndex = InTypes.FindChecked(TEXT("A"));
-	check(TypeIndex == InTypes.FindChecked(TEXT("B")));
-	check(InTypes.FindChecked(TEXT("Result")) == RigVMTypeUtils::TypeIndex::Bool);
+	const TRigVMTypeIndex TypeIndex = InTypes.FindChecked(AName);
+	check(TypeIndex == InTypes.FindChecked(BName));
+	check(InTypes.FindChecked(ResultName) == RigVMTypeUtils::TypeIndex::Bool);
 
 	if(TypeIndex == RigVMTypeUtils::TypeIndex::Float)
 	{
@@ -103,7 +120,7 @@ void FRigVMDispatch_CoreEquals::Execute(FRigVMExtendedExecuteContext& InContext,
 
 	const uint8* A = Handles[0].GetData();
 	const uint8* B = Handles[1].GetData();
-	bool& Result = *(bool*)Handles[2].GetData();
+	bool& Result = *reinterpret_cast<bool*>(Handles[2].GetData());
 	Result = PropertyA->Identical(A, B);
 	Result = AdaptResult(Result, InContext);
 }
diff --git a/Engine/Source/Runtime/RigVM/Private/RigVMFunctions/RigVMDispatch_If.cpp b/Engine/Source/Runtime/RigVM/Private/RigVMFunctions/RigVMDispatch_If.cpp
--- a/Engine/Source/Runtime/RigVM/Private/RigVMFunctions/RigVMDispatch_If.cpp
+++ b/Engine/Source/Runtime/RigVM/Private/RigVMFunctions/RigVMDispatch_If.cpp
@@ -11,6 +11,18 @@ const FName FRigVMDispatch_If::TrueName = TEXT("True");
 const FName FRigVMDispatch_If::FalseName = TEXT("False");
 const FName FRigVMDispatch_If::ResultName = TEXT("Result");
 
+FName FRigVMDispatch_If::GetArgumentNameForOperandIndex(int32 InOperandIndex, int32 InTotalOperands) const
+{
+	static const FName ArgumentNames[] = {
+		ConditionName,
+		TrueName,
+		FalseName,
+		ResultName
+	};
+	check(InTotalOperands == UE_ARRAY_COUNT(ArgumentNames));
+	return ArgumentNames[InOperandIndex];
+}
+
 TArray<FRigVMTemplateArgument> FRigVMDispatch_If::GetArguments() const
 {
 	static const TArray<FRigVMTemplateArgument::ETypeCategory> ValueCategories = {
diff --git a/Engine/Source/Runtime/RigVM/Private/RigVMFunctions/RigVMDispatch_Print.cpp b/Engine/Source/Runtime/RigVM/Private/RigVMFunctions/RigVMDispatch_Print.cpp
--- a/Engine/Source/Runtime/RigVM/Private/RigVMFunctions/RigVMDispatch_Print.cpp
+++ b/Engine/Source/Runtime/RigVM/Private/RigVMFunctions/RigVMDispatch_Print.cpp
@@ -7,6 +7,27 @@
 
 #include UE_INLINE_GENERATED_CPP_BY_NAME(RigVMDispatch_Print)
 
+
+const FName FRigVMDispatch_Print::PrefixName = TEXT("Prefix");
+const FName FRigVMDispatch_Print::ValueName = TEXT("Value");
+const FName FRigVMDispatch_Print::EnabledName = TEXT("Enabled");
+const FName FRigVMDispatch_Print::ScreenDurationName = TEXT("ScreenDuration");
+const FName FRigVMDispatch_Print::ScreenColorName = TEXT("ScreenColor");
+
+
+FName FRigVMDispatch_Print::GetArgumentNameForOperandIndex(int32 InOperandIndex, int32 InTotalOperands) const
+{
+	static const FName ArgumentNames[] = {
+		PrefixName,
+		ValueName,
+		EnabledName,
+		ScreenDurationName,
+		ScreenColorName
+	};
+	check(InTotalOperands == UE_ARRAY_COUNT(ArgumentNames));
+	return ArgumentNames[InOperandIndex];
+}
+
 TArray<FRigVMTemplateArgument> FRigVMDispatch_Print::GetArguments() const
 {
 	const TArray<FRigVMTemplateArgument::ETypeCategory> ValueCategories = {
@@ -14,11 +35,11 @@ TArray<FRigVMTemplateArgument> FRigVMDispatch_Print::GetArguments() const
 		FRigVMTemplateArgument::ETypeCategory_ArrayAnyValue
 	};
 	return {
-		FRigVMTemplateArgument(TEXT("Prefix"), ERigVMPinDirection::Input, RigVMTypeUtils::TypeIndex::FString),
-		FRigVMTemplateArgument(TEXT("Value"), ERigVMPinDirection::Input, ValueCategories),
-		FRigVMTemplateArgument(TEXT("Enabled"), ERigVMPinDirection::Input, RigVMTypeUtils::TypeIndex::Bool),
-		FRigVMTemplateArgument(TEXT("ScreenDuration"), ERigVMPinDirection::Input, RigVMTypeUtils::TypeIndex::Float),
-		FRigVMTemplateArgument(TEXT("ScreenColor"), ERigVMPinDirection::Input, FRigVMRegistry::Get().GetTypeIndex<FLinearColor>())
+		FRigVMTemplateArgument(PrefixName, ERigVMPinDirection::Input, RigVMTypeUtils::TypeIndex::FString),
+		FRigVMTemplateArgument(ValueName, ERigVMPinDirection::Input, ValueCategories),
+		FRigVMTemplateArgument(EnabledName, ERigVMPinDirection::Input, RigVMTypeUtils::TypeIndex::Bool),
+		FRigVMTemplateArgument(ScreenDurationName, ERigVMPinDirection::Input, RigVMTypeUtils::TypeIndex::Float),
+		FRigVMTemplateArgument(ScreenColorName, ERigVMPinDirection::Input, FRigVMRegistry::Get().GetTypeIndex<FLinearColor>())
 	};
 }
 
@@ -31,11 +52,11 @@ FRigVMTemplateTypeMap FRigVMDispatch_Print::OnNewArgumentType(const FName& InArg
                                                               TRigVMTypeIndex InTypeIndex) const
 {
 	FRigVMTemplateTypeMap Types;
-	Types.Add(TEXT("Prefix"), RigVMTypeUtils::TypeIndex::FString);
-	Types.Add(TEXT("Value"), InTypeIndex);
-	Types.Add(TEXT("Enabled"), RigVMTypeUtils::TypeIndex::Bool);
-	Types.Add(TEXT("ScreenDuration"), RigVMTypeUtils::TypeIndex::Float);
-	Types.Add(TEXT("ScreenColor"), FRigVMRegistry::Get().GetTypeIndex<FLinearColor>());
+	Types.Add(PrefixName, RigVMTypeUtils::TypeIndex::FString);
+	Types.Add(ValueName, InTypeIndex);
+	Types.Add(EnabledName, RigVMTypeUtils::TypeIndex::Bool);
+	Types.Add(ScreenDurationName, RigVMTypeUtils::TypeIndex::Float);
+	Types.Add(ScreenColorName, FRigVMRegistry::Get().GetTypeIndex<FLinearColor>());
 	return Types;
 }
 
@@ -43,11 +64,11 @@ FRigVMTemplateTypeMap FRigVMDispatch_Print::OnNewArgumentType(const FName& InArg
 
 FString FRigVMDispatch_Print::GetArgumentDefaultValue(const FName& InArgumentName, TRigVMTypeIndex InTypeIndex) const
 {
-	if(InArgumentName == TEXT("Enabled"))
+	if(InArgumentName == EnabledName)
 	{
 		return TEXT("True");
 	}
-	if(InArgumentName == TEXT("ScreenDuration"))
+	if(InArgumentName == ScreenDurationName)
 	{
 		return TEXT("0.050000");
 	}
@@ -56,7 +77,7 @@ FString FRigVMDispatch_Print::GetArgumentDefaultValue(const FName& InArgumentNam
 
 FString FRigVMDispatch_Print::GetArgumentMetaData(const FName& InArgumentName, const FName& InMetaDataKey) const
 {
-	if(InArgumentName == TEXT("ScreenDuration") || InArgumentName == TEXT("ScreenColor"))
+	if(InArgumentName == ScreenDurationName || InArgumentName == ScreenColorName)
 	{
 		if(InMetaDataKey == FRigVMStruct::DetailsOnlyMetaName)
 		{
diff --git a/Engine/Source/Runtime/RigVM/Private/RigVMFunctions/RigVMDispatch_Select.cpp b/Engine/Source/Runtime/RigVM/Private/RigVMFunctions/RigVMDispatch_Select.cpp
--- a/Engine/Source/Runtime/RigVM/Private/RigVMFunctions/RigVMDispatch_Select.cpp
+++ b/Engine/Source/Runtime/RigVM/Private/RigVMFunctions/RigVMDispatch_Select.cpp
@@ -11,6 +11,18 @@ const FName FRigVMDispatch_SelectInt32::IndexName = TEXT("Index");
 const FName FRigVMDispatch_SelectInt32::ValuesName = TEXT("Values");
 const FName FRigVMDispatch_SelectInt32::ResultName = TEXT("Result");
 
+
+FName FRigVMDispatch_SelectInt32::GetArgumentNameForOperandIndex(int32 InOperandIndex, int32 InTotalOperands) const
+{
+	if(InOperandIndex > 0 && InOperandIndex < InTotalOperands - 1)
+	{
+		return FRigVMBranchInfo::GetFixedArrayLabel(ValuesName, *FString::FromInt(InOperandIndex - 1));
+	}
+	
+	return InOperandIndex == 0 ? IndexName : ResultName;
+}
+
+
 TArray<FRigVMTemplateArgument> FRigVMDispatch_SelectInt32::GetArguments() const
 {
 	static const TArray<FRigVMTemplateArgument::ETypeCategory> ArrayCategories = {
@@ -91,14 +103,6 @@ FString FRigVMDispatch_SelectInt32::GetArgumentDefaultValue(const FName& InArgum
 
 #endif
 
-FName FRigVMDispatch_SelectInt32::GetArgumentNameForOperandIndex(int32 InOperandIndex, int32 InTotalOperands) const
-{
-	if(InOperandIndex > 0 && InOperandIndex < InTotalOperands - 1)
-	{
-		return FRigVMBranchInfo::GetFixedArrayLabel(ValuesName, *FString::FromInt(InOperandIndex - 1));
-	}
-	return FRigVMDispatch_CoreBase::GetArgumentNameForOperandIndex(InOperandIndex == 0 ? 0 : 2, 3);
-}
 
 FRigVMFunctionPtr FRigVMDispatch_SelectInt32::GetDispatchFunctionImpl(const FRigVMTemplateTypeMap& InTypes) const
 {
diff --git a/Engine/Source/Runtime/RigVM/Private/RigVMFunctions/RigVMDispatch_Switch.cpp b/Engine/Source/Runtime/RigVM/Private/RigVMFunctions/RigVMDispatch_Switch.cpp
--- a/Engine/Source/Runtime/RigVM/Private/RigVMFunctions/RigVMDispatch_Switch.cpp
+++ b/Engine/Source/Runtime/RigVM/Private/RigVMFunctions/RigVMDispatch_Switch.cpp
@@ -10,6 +10,16 @@
 const FName FRigVMDispatch_SwitchInt32::IndexName = TEXT("Index");
 const FName FRigVMDispatch_SwitchInt32::CasesName = TEXT("Cases");
 
+FName FRigVMDispatch_SwitchInt32::GetArgumentNameForOperandIndex(int32 InOperandIndex, int32 InTotalOperands) const
+{
+	static const FName ArgumentNames[] = {
+		IndexName,
+		FRigVMStruct::ControlFlowBlockToRunName
+	};
+	check(InTotalOperands == UE_ARRAY_COUNT(ArgumentNames));
+	return ArgumentNames[InOperandIndex];
+}
+
 TArray<FRigVMTemplateArgument> FRigVMDispatch_SwitchInt32::GetArguments() const
 {
 	return {
diff --git a/Engine/Source/Runtime/RigVM/Public/RigVMFunctions/RigVMDispatch_Array.h b/Engine/Source/Runtime/RigVM/Public/RigVMFunctions/RigVMDispatch_Array.h
--- a/Engine/Source/Runtime/RigVM/Public/RigVMFunctions/RigVMDispatch_Array.h
+++ b/Engine/Source/Runtime/RigVM/Public/RigVMFunctions/RigVMDispatch_Array.h
@@ -77,6 +77,7 @@ struct RIGVM_API FRigVMDispatch_ArrayReset : public FRigVMDispatch_ArrayBaseMuta
 
 public:
 	virtual ERigVMOpCode GetOpCode() const override { return ERigVMOpCode::ArrayReset; }
+	virtual FName GetArgumentNameForOperandIndex(int32 InOperandIndex, int32 InTotalOperands) const override;
 	virtual TArray<FRigVMTemplateArgument> GetArguments() const override;
 	virtual FRigVMTemplateTypeMap OnNewArgumentType(const FName& InArgumentName, TRigVMTypeIndex InTypeIndex) const override;
 #if WITH_EDITOR
@@ -96,6 +97,7 @@ struct RIGVM_API FRigVMDispatch_ArrayGetNum : public FRigVMDispatch_ArrayBase
 
 public:
 	virtual ERigVMOpCode GetOpCode() const override { return ERigVMOpCode::ArrayReset; }
+	virtual FName GetArgumentNameForOperandIndex(int32 InOperandIndex, int32 InTotalOperands) const override;
 	virtual TArray<FRigVMTemplateArgument> GetArguments() const override;
 	virtual FRigVMTemplateTypeMap OnNewArgumentType(const FName& InArgumentName, TRigVMTypeIndex InTypeIndex) const override;
 #if WITH_EDITOR
@@ -115,6 +117,7 @@ struct RIGVM_API FRigVMDispatch_ArraySetNum : public FRigVMDispatch_ArrayBaseMut
 
 public:
 	virtual ERigVMOpCode GetOpCode() const override { return ERigVMOpCode::ArraySetNum; }
+	virtual FName GetArgumentNameForOperandIndex(int32 InOperandIndex, int32 InTotalOperands) const override;
 	virtual TArray<FRigVMTemplateArgument> GetArguments() const override;
 	virtual FRigVMTemplateTypeMap OnNewArgumentType(const FName& InArgumentName, TRigVMTypeIndex InTypeIndex) const override;
 #if WITH_EDITOR
@@ -134,6 +137,7 @@ struct RIGVM_API FRigVMDispatch_ArrayGetAtIndex : public FRigVMDispatch_ArrayBas
 
 public:
 	virtual ERigVMOpCode GetOpCode() const override { return ERigVMOpCode::ArrayGetAtIndex; }
+	virtual FName GetArgumentNameForOperandIndex(int32 InOperandIndex, int32 InTotalOperands) const override;
 	virtual TArray<FRigVMTemplateArgument> GetArguments() const override;
 	virtual FRigVMTemplateTypeMap OnNewArgumentType(const FName& InArgumentName, TRigVMTypeIndex InTypeIndex) const override;
 #if WITH_EDITOR
@@ -153,6 +157,7 @@ struct RIGVM_API FRigVMDispatch_ArraySetAtIndex : public FRigVMDispatch_ArrayBas
 
 public:
 	virtual ERigVMOpCode GetOpCode() const override { return ERigVMOpCode::ArraySetAtIndex; }
+	virtual FName GetArgumentNameForOperandIndex(int32 InOperandIndex, int32 InTotalOperands) const override;
 	virtual TArray<FRigVMTemplateArgument> GetArguments() const override;
 	virtual FRigVMTemplateTypeMap OnNewArgumentType(const FName& InArgumentName, TRigVMTypeIndex InTypeIndex) const override;
 #if WITH_EDITOR
@@ -172,6 +177,7 @@ struct RIGVM_API FRigVMDispatch_ArrayAdd : public FRigVMDispatch_ArraySetAtIndex
 
 public:
 	virtual ERigVMOpCode GetOpCode() const override { return ERigVMOpCode::ArrayAdd; }
+	virtual FName GetArgumentNameForOperandIndex(int32 InOperandIndex, int32 InTotalOperands) const override;
 	virtual TArray<FRigVMTemplateArgument> GetArguments() const override;
 #if WITH_EDITOR
 	virtual FText GetNodeTooltip(const FRigVMTemplateTypeMap& InTypes) const override;
@@ -207,6 +213,7 @@ struct RIGVM_API FRigVMDispatch_ArrayRemove : public FRigVMDispatch_ArrayBaseMut
 
 public:
 	virtual ERigVMOpCode GetOpCode() const override { return ERigVMOpCode::ArrayRemove; }
+	virtual FName GetArgumentNameForOperandIndex(int32 InOperandIndex, int32 InTotalOperands) const override;
 	virtual TArray<FRigVMTemplateArgument> GetArguments() const override;
 	virtual FRigVMTemplateTypeMap OnNewArgumentType(const FName& InArgumentName, TRigVMTypeIndex InTypeIndex) const override;
 #if WITH_EDITOR
@@ -244,6 +251,7 @@ struct RIGVM_API FRigVMDispatch_ArrayFind : public FRigVMDispatch_ArrayBase
 
 public:
 	virtual ERigVMOpCode GetOpCode() const override { return ERigVMOpCode::ArrayFind; }
+	virtual FName GetArgumentNameForOperandIndex(int32 InOperandIndex, int32 InTotalOperands) const override;
 	virtual TArray<FRigVMTemplateArgument> GetArguments() const override;
 	virtual FRigVMTemplateTypeMap OnNewArgumentType(const FName& InArgumentName, TRigVMTypeIndex InTypeIndex) const override;
 #if WITH_EDITOR
@@ -263,6 +271,7 @@ struct RIGVM_API FRigVMDispatch_ArrayAppend : public FRigVMDispatch_ArrayBaseMut
 
 public:
 	virtual ERigVMOpCode GetOpCode() const override { return ERigVMOpCode::ArrayAppend; }
+	virtual FName GetArgumentNameForOperandIndex(int32 InOperandIndex, int32 InTotalOperands) const override;
 	virtual TArray<FRigVMTemplateArgument> GetArguments() const override;
 	virtual FRigVMTemplateTypeMap OnNewArgumentType(const FName& InArgumentName, TRigVMTypeIndex InTypeIndex) const override;
 #if WITH_EDITOR
@@ -282,6 +291,7 @@ struct RIGVM_API FRigVMDispatch_ArrayClone : public FRigVMDispatch_ArrayBase
 
 public:
 	virtual ERigVMOpCode GetOpCode() const override { return ERigVMOpCode::ArrayClone; }
+	virtual FName GetArgumentNameForOperandIndex(int32 InOperandIndex, int32 InTotalOperands) const override;
 	virtual TArray<FRigVMTemplateArgument> GetArguments() const override;
 	virtual FRigVMTemplateTypeMap OnNewArgumentType(const FName& InArgumentName, TRigVMTypeIndex InTypeIndex) const override;
 #if WITH_EDITOR
@@ -318,6 +328,7 @@ struct RIGVM_API FRigVMDispatch_ArrayDifference : public FRigVMDispatch_ArrayBas
 
 public:
 	virtual ERigVMOpCode GetOpCode() const override { return ERigVMOpCode::ArrayDifference; }
+	virtual FName GetArgumentNameForOperandIndex(int32 InOperandIndex, int32 InTotalOperands) const override;
 	virtual TArray<FRigVMTemplateArgument> GetArguments() const override;
 	virtual FRigVMTemplateTypeMap OnNewArgumentType(const FName& InArgumentName, TRigVMTypeIndex InTypeIndex) const override;
 #if WITH_EDITOR
@@ -354,6 +365,7 @@ struct RIGVM_API FRigVMDispatch_ArrayIterator : public FRigVMDispatch_ArrayBaseM
 
 public:
 	virtual ERigVMOpCode GetOpCode() const override { return ERigVMOpCode::ArrayIterator; }
+	virtual FName GetArgumentNameForOperandIndex(int32 InOperandIndex, int32 InTotalOperands) const override;
 	virtual TArray<FRigVMTemplateArgument> GetArguments() const override;
 	virtual TArray<FRigVMExecuteArgument> GetExecuteArguments_Impl(const FRigVMDispatchContext& InContext) const override;
 	virtual const TArray<FName>& GetControlFlowBlocks_Impl(const FRigVMDispatchContext& InContext) const override;
diff --git a/Engine/Source/Runtime/RigVM/Public/RigVMFunctions/RigVMDispatch_Core.h b/Engine/Source/Runtime/RigVM/Public/RigVMFunctions/RigVMDispatch_Core.h
--- a/Engine/Source/Runtime/RigVM/Public/RigVMFunctions/RigVMDispatch_Core.h
+++ b/Engine/Source/Runtime/RigVM/Public/RigVMFunctions/RigVMDispatch_Core.h
@@ -21,6 +21,7 @@ struct RIGVM_API FRigVMDispatch_CoreEquals : public FRigVMDispatch_CoreBase
 
 public:
 
+	virtual FName GetArgumentNameForOperandIndex(int32 InOperandIndex, int32 InTotalOperands) const override;
 	virtual TArray<FRigVMTemplateArgument> GetArguments() const override;
 	virtual FRigVMTemplateTypeMap OnNewArgumentType(const FName& InArgumentName, TRigVMTypeIndex InTypeIndex) const override;
 
@@ -75,6 +76,10 @@ protected:
 		bool& Result = *(bool*)Handles[2].GetData();
 		Result = AdaptResult(A.Equals(B), InContext);
 	}
+
+	static const FName AName;
+	static const FName BName;
+	static const FName ResultName;
 };
 
 /*
diff --git a/Engine/Source/Runtime/RigVM/Public/RigVMFunctions/RigVMDispatch_If.h b/Engine/Source/Runtime/RigVM/Public/RigVMFunctions/RigVMDispatch_If.h
--- a/Engine/Source/Runtime/RigVM/Public/RigVMFunctions/RigVMDispatch_If.h
+++ b/Engine/Source/Runtime/RigVM/Public/RigVMFunctions/RigVMDispatch_If.h
@@ -16,6 +16,7 @@ struct RIGVM_API FRigVMDispatch_If : public FRigVMDispatch_CoreBase
 
 public:
 
+	virtual FName GetArgumentNameForOperandIndex(int32 InOperandIndex, int32 InTotalOperands) const override;
 	virtual TArray<FRigVMTemplateArgument> GetArguments() const override;
 	virtual FRigVMTemplateTypeMap OnNewArgumentType(const FName& InArgumentName, TRigVMTypeIndex InTypeIndex) const override;
 	virtual bool IsSingleton() const override { return true; } 
diff --git a/Engine/Source/Runtime/RigVM/Public/RigVMFunctions/RigVMDispatch_Print.h b/Engine/Source/Runtime/RigVM/Public/RigVMFunctions/RigVMDispatch_Print.h
--- a/Engine/Source/Runtime/RigVM/Public/RigVMFunctions/RigVMDispatch_Print.h
+++ b/Engine/Source/Runtime/RigVM/Public/RigVMFunctions/RigVMDispatch_Print.h
@@ -15,6 +15,7 @@ struct RIGVM_API FRigVMDispatch_Print : public FRigVMDispatchFactory
 
 public:
 
+	virtual FName GetArgumentNameForOperandIndex(int32 InOperandIndex, int32 InTotalOperands) const override;
 	virtual TArray<FRigVMTemplateArgument> GetArguments() const override;
 	virtual TArray<FRigVMExecuteArgument> GetExecuteArguments_Impl(const FRigVMDispatchContext& InContext) const override;
 	virtual FRigVMTemplateTypeMap OnNewArgumentType(const FName& InArgumentName, TRigVMTypeIndex InTypeIndex) const override;
@@ -32,4 +33,10 @@ protected:
 		return &FRigVMDispatch_Print::Execute;
 	}
 	static void Execute(FRigVMExtendedExecuteContext& InContext, FRigVMMemoryHandleArray Handles);
+
+	static const FName PrefixName;
+	static const FName ValueName;
+	static const FName EnabledName;
+	static const FName ScreenDurationName;
+	static const FName ScreenColorName;
 };
diff --git a/Engine/Source/Runtime/RigVM/Public/RigVMFunctions/RigVMDispatch_Switch.h b/Engine/Source/Runtime/RigVM/Public/RigVMFunctions/RigVMDispatch_Switch.h
--- a/Engine/Source/Runtime/RigVM/Public/RigVMFunctions/RigVMDispatch_Switch.h
+++ b/Engine/Source/Runtime/RigVM/Public/RigVMFunctions/RigVMDispatch_Switch.h
@@ -19,6 +19,7 @@ public:
 	FRigVMDispatch_SwitchInt32()
 	{}
 
+	virtual FName GetArgumentNameForOperandIndex(int32 InOperandIndex, int32 InTotalOperands) const override;
 	virtual TArray<FRigVMTemplateArgument> GetArguments() const override;
 	virtual TArray<FRigVMExecuteArgument> GetExecuteArguments_Impl(const FRigVMDispatchContext& InContext) const override;
 	virtual bool IsSingleton() const override { return true; } 
