Control Rig: Disable UE_RIGVMCONTROLLER_VERBOSE_REPOPULATE by default

diff --git a/Engine/Source/Runtime/RigVM/Public/RigVMDefines.h b/Engine/Source/Runtime/RigVM/Public/RigVMDefines.h
--- a/Engine/Source/Runtime/RigVM/Public/RigVMDefines.h
+++ b/Engine/Source/Runtime/RigVM/Public/RigVMDefines.h
@@ -31,5 +31,5 @@
 #endif
 
 #ifndef UE_RIGVMCONTROLLER_VERBOSE_REPOPULATE
-#define UE_RIGVMCONTROLLER_VERBOSE_REPOPULATE UE_BUILD_DEBUG
+#define UE_RIGVMCONTROLLER_VERBOSE_REPOPULATE 0
 #endif
