Control Rig: Don't persist a VM on the CR BP skeleton class. It's not necessary, and also during duplication of those classes, GetDefaultObject doesn't work.

#rnx

diff --git a/Engine/Source/Runtime/RigVM/Private/RigVMBlueprintGeneratedClass.cpp b/Engine/Source/Runtime/RigVM/Private/RigVMBlueprintGeneratedClass.cpp
--- a/Engine/Source/Runtime/RigVM/Private/RigVMBlueprintGeneratedClass.cpp
+++ b/Engine/Source/Runtime/RigVM/Private/RigVMBlueprintGeneratedClass.cpp
@@ -45,23 +45,42 @@ void URigVMBlueprintGeneratedClass::Serialize(FArchive& Ar)
 		return;
 	}
 
+	// If we're the blueprint skeleton class, we don't need the VM. We should be an empty vessel. Ideally there should
+	// be a class flag to check for this, but there isn't, so we're left with name matching. 
+	auto IsSkeletonClass = [](const UClass* InClass)
+	{
+#if WITH_EDITOR 
+		return InClass->ClassGeneratedBy &&
+			(InClass->GetName().StartsWith(TEXT("SKEL_")) || InClass->GetName().StartsWith(TEXT("REINST_SKEL_")));
+#else
+		// For non-editor targets, this is all irrelevant, since the skeleton class will never exist for them.
+		return false;
+#endif
+	};
+
 	URigVM* VM = NewObject<URigVM>(GetTransientPackage());
 
-	if (const URigVMHost* CDO = Cast<URigVMHost>(GetDefaultObject(true)))
+	if (!IsSkeletonClass(this))
 	{
-		if (Ar.IsSaving() && CDO->VM)
+		if (const URigVMHost* CDO = Cast<URigVMHost>(GetDefaultObject(true)))
 		{
-			VM->CopyFrom(CDO->VM);
+			if (Ar.IsSaving() && CDO->VM)
+			{
+				VM->CopyFrom(CDO->VM);
+			}
 		}
 	}
 	
 	VM->Serialize(Ar);
 
-	if (const URigVMHost* CDO = Cast<URigVMHost>(GetDefaultObject(false)))
+	if (!IsSkeletonClass(this))
 	{
-		if (Ar.IsLoading())
+		if (const URigVMHost* CDO = Cast<URigVMHost>(GetDefaultObject(false)))
 		{
-			CDO->VM->CopyFrom(VM);
+			if (Ar.IsLoading())
+			{
+				CDO->VM->CopyFrom(VM);
+			}
 		}
 	}
 
