RigVM: Support object property as part of template argument

diff --git a/Engine/Source/Runtime/RigVM/Private/RigVMCore/RigVMTemplate.cpp b/Engine/Source/Runtime/RigVM/Private/RigVMCore/RigVMTemplate.cpp
--- a/Engine/Source/Runtime/RigVM/Private/RigVMCore/RigVMTemplate.cpp
+++ b/Engine/Source/Runtime/RigVM/Private/RigVMCore/RigVMTemplate.cpp
@@ -50,7 +50,14 @@ FRigVMTemplateArgument::FRigVMTemplateArgument(FProperty* InProperty)
 	{
 		CPPTypeObject = ByteProperty->Enum;
 	}
-
+	else if (FObjectProperty* ObjectProperty = CastField<FObjectProperty>(InProperty))
+	{
+		if(RigVMCore::SupportsUObjects())
+		{
+			CPPTypeObject = ObjectProperty->PropertyClass;
+		}
+	}
+	
 	const FRigVMTemplateArgumentType Type(CPPTypeName, CPPTypeObject);
 	const TRigVMTypeIndex TypeIndex = FRigVMRegistry::Get().FindOrAddType_Internal(Type, true); 
 
