Control Rig: Fix ensure when opening DEMO_Constraints_FBIK_CtrlRig

diff --git a/Engine/Source/Runtime/RigVM/Private/RigVMCore/RigVMRegistry.cpp b/Engine/Source/Runtime/RigVM/Private/RigVMCore/RigVMRegistry.cpp
--- a/Engine/Source/Runtime/RigVM/Private/RigVMCore/RigVMRegistry.cpp
+++ b/Engine/Source/Runtime/RigVM/Private/RigVMCore/RigVMRegistry.cpp
@@ -863,6 +863,11 @@ bool FRigVMRegistry::IsArrayType(TRigVMTypeIndex InTypeIndex) const
 
 bool FRigVMRegistry::IsExecuteType(TRigVMTypeIndex InTypeIndex) const
 {
+	if(InTypeIndex == INDEX_NONE)
+	{
+		return false;
+	}
+	
 	if(ensure(Types.IsValidIndex(InTypeIndex)))
 	{
 		return Types[InTypeIndex].bIsExecute;
