RigVM: Added an engine test for PropertyBag usage example with serialization

diff --git a/Engine/Source/Runtime/RigVM/Public/RigVMDefines.h b/Engine/Source/Runtime/RigVM/Public/RigVMDefines.h
--- a/Engine/Source/Runtime/RigVM/Public/RigVMDefines.h
+++ b/Engine/Source/Runtime/RigVM/Public/RigVMDefines.h
@@ -33,3 +33,7 @@
 #ifndef UE_RIGVMCONTROLLER_VERBOSE_REPOPULATE
 #define UE_RIGVMCONTROLLER_VERBOSE_REPOPULATE 0
 #endif
+
+#ifndef UE_RIGVM_PROPERTY_BAG_STORAGE_ENABLED
+#define UE_RIGVM_PROPERTY_BAG_STORAGE_ENABLED 0
+#endif
