RigVM: Support for uint32 types

diff --git a/Engine/Source/Programs/Shared/EpicGames.UHT/Types/Properties/UhtDelegateProperty.cs b/Engine/Source/Programs/Shared/EpicGames.UHT/Types/Properties/UhtDelegateProperty.cs
--- a/Engine/Source/Programs/Shared/EpicGames.UHT/Types/Properties/UhtDelegateProperty.cs
+++ b/Engine/Source/Programs/Shared/EpicGames.UHT/Types/Properties/UhtDelegateProperty.cs
@@ -43,7 +43,7 @@ namespace EpicGames.UHT.Types
 			Function = function;
 			HeaderFile.AddReferencedHeader(function);
 			PropertyCaps |= UhtPropertyCaps.PassCppArgsByRef | UhtPropertyCaps.IsParameterSupportedByBlueprint |
-				UhtPropertyCaps.IsMemberSupportedByBlueprint | UhtPropertyCaps.SupportsRigVM;
+				UhtPropertyCaps.IsMemberSupportedByBlueprint;
 		}
 
 		/// <inheritdoc/>
diff --git a/Engine/Source/Programs/Shared/EpicGames.UHT/Types/Properties/UhtDoubleProperty.cs b/Engine/Source/Programs/Shared/EpicGames.UHT/Types/Properties/UhtDoubleProperty.cs
--- a/Engine/Source/Programs/Shared/EpicGames.UHT/Types/Properties/UhtDoubleProperty.cs
+++ b/Engine/Source/Programs/Shared/EpicGames.UHT/Types/Properties/UhtDoubleProperty.cs
@@ -29,7 +29,7 @@ namespace EpicGames.UHT.Types
 		/// <param name="propertySettings">Property settings</param>
 		public UhtDoubleProperty(UhtPropertySettings propertySettings) : base(propertySettings, UhtPropertyIntType.None)
 		{
-			PropertyCaps |= UhtPropertyCaps.IsParameterSupportedByBlueprint | UhtPropertyCaps.IsMemberSupportedByBlueprint;
+			PropertyCaps |= UhtPropertyCaps.IsParameterSupportedByBlueprint | UhtPropertyCaps.IsMemberSupportedByBlueprint | UhtPropertyCaps.SupportsRigVM;
 		}
 
 		/// <inheritdoc/>
diff --git a/Engine/Source/Programs/Shared/EpicGames.UHT/Types/Properties/UhtFloatProperty.cs b/Engine/Source/Programs/Shared/EpicGames.UHT/Types/Properties/UhtFloatProperty.cs
--- a/Engine/Source/Programs/Shared/EpicGames.UHT/Types/Properties/UhtFloatProperty.cs
+++ b/Engine/Source/Programs/Shared/EpicGames.UHT/Types/Properties/UhtFloatProperty.cs
@@ -28,7 +28,7 @@ namespace EpicGames.UHT.Types
 		/// <param name="propertySettings">Property settings</param>
 		public UhtFloatProperty(UhtPropertySettings propertySettings) : base(propertySettings, UhtPropertyIntType.None)
 		{
-			PropertyCaps |= UhtPropertyCaps.CanExposeOnSpawn | UhtPropertyCaps.IsParameterSupportedByBlueprint | UhtPropertyCaps.IsMemberSupportedByBlueprint;
+			PropertyCaps |= UhtPropertyCaps.CanExposeOnSpawn | UhtPropertyCaps.IsParameterSupportedByBlueprint | UhtPropertyCaps.IsMemberSupportedByBlueprint | UhtPropertyCaps.SupportsRigVM;
 		}
 
 		/// <inheritdoc/>
diff --git a/Engine/Source/Programs/Shared/EpicGames.UHT/Types/Properties/UhtIntProperty.cs b/Engine/Source/Programs/Shared/EpicGames.UHT/Types/Properties/UhtIntProperty.cs
--- a/Engine/Source/Programs/Shared/EpicGames.UHT/Types/Properties/UhtIntProperty.cs
+++ b/Engine/Source/Programs/Shared/EpicGames.UHT/Types/Properties/UhtIntProperty.cs
@@ -28,7 +28,7 @@ namespace EpicGames.UHT.Types
 		/// <param name="intType">Integer type</param>
 		public UhtIntProperty(UhtPropertySettings propertySettings, UhtPropertyIntType intType) : base(propertySettings, intType)
 		{
-			PropertyCaps |= UhtPropertyCaps.CanExposeOnSpawn | UhtPropertyCaps.IsParameterSupportedByBlueprint | UhtPropertyCaps.IsMemberSupportedByBlueprint;
+			PropertyCaps |= UhtPropertyCaps.CanExposeOnSpawn | UhtPropertyCaps.IsParameterSupportedByBlueprint | UhtPropertyCaps.IsMemberSupportedByBlueprint | UhtPropertyCaps.SupportsRigVM;
 		}
 
 		/// <inheritdoc/>
diff --git a/Engine/Source/Programs/Shared/EpicGames.UHT/Types/Properties/UhtMulticastDelegateProperty.cs b/Engine/Source/Programs/Shared/EpicGames.UHT/Types/Properties/UhtMulticastDelegateProperty.cs
--- a/Engine/Source/Programs/Shared/EpicGames.UHT/Types/Properties/UhtMulticastDelegateProperty.cs
+++ b/Engine/Source/Programs/Shared/EpicGames.UHT/Types/Properties/UhtMulticastDelegateProperty.cs
@@ -43,7 +43,7 @@ namespace EpicGames.UHT.Types
 			Function = function;
 			HeaderFile.AddReferencedHeader(function);
 			PropertyCaps |= UhtPropertyCaps.PassCppArgsByRef | UhtPropertyCaps.CanBeBlueprintAssignable | UhtPropertyCaps.CanBeBlueprintCallable |
-				UhtPropertyCaps.CanBeBlueprintAuthorityOnly | UhtPropertyCaps.IsMemberSupportedByBlueprint | UhtPropertyCaps.SupportsRigVM;
+				UhtPropertyCaps.CanBeBlueprintAuthorityOnly | UhtPropertyCaps.IsMemberSupportedByBlueprint;
 			PropertyCaps &= ~(UhtPropertyCaps.IsParameterSupportedByBlueprint);
 		}
 
diff --git a/Engine/Source/Programs/Shared/EpicGames.UHT/Types/Properties/UhtNumericProperty.cs b/Engine/Source/Programs/Shared/EpicGames.UHT/Types/Properties/UhtNumericProperty.cs
--- a/Engine/Source/Programs/Shared/EpicGames.UHT/Types/Properties/UhtNumericProperty.cs
+++ b/Engine/Source/Programs/Shared/EpicGames.UHT/Types/Properties/UhtNumericProperty.cs
@@ -31,7 +31,7 @@ namespace EpicGames.UHT.Types
 		protected UhtNumericProperty(UhtPropertySettings propertySettings, UhtPropertyIntType intType) : base(propertySettings)
 		{
 			IntType = intType;
-			PropertyCaps |= UhtPropertyCaps.RequiresNullConstructorArg | UhtPropertyCaps.SupportsRigVM;
+			PropertyCaps |= UhtPropertyCaps.RequiresNullConstructorArg;
 		}
 
 		/// <inheritdoc/>
diff --git a/Engine/Source/Programs/Shared/EpicGames.UHT/Types/Properties/UhtTextProperty.cs b/Engine/Source/Programs/Shared/EpicGames.UHT/Types/Properties/UhtTextProperty.cs
--- a/Engine/Source/Programs/Shared/EpicGames.UHT/Types/Properties/UhtTextProperty.cs
+++ b/Engine/Source/Programs/Shared/EpicGames.UHT/Types/Properties/UhtTextProperty.cs
@@ -37,7 +37,7 @@ namespace EpicGames.UHT.Types
 		public UhtTextProperty(UhtPropertySettings propertySettings) : base(propertySettings)
 		{
 			PropertyCaps |= UhtPropertyCaps.PassCppArgsByRef | UhtPropertyCaps.CanExposeOnSpawn | UhtPropertyCaps.IsParameterSupportedByBlueprint |
-				UhtPropertyCaps.IsMemberSupportedByBlueprint | UhtPropertyCaps.SupportsRigVM;
+				UhtPropertyCaps.IsMemberSupportedByBlueprint;
 			PropertyCaps &= ~(UhtPropertyCaps.CanBeContainerKey);
 		}
 
diff --git a/Engine/Source/Programs/Shared/EpicGames.UHT/Types/Properties/UhtUInt32Property.cs b/Engine/Source/Programs/Shared/EpicGames.UHT/Types/Properties/UhtUInt32Property.cs
--- a/Engine/Source/Programs/Shared/EpicGames.UHT/Types/Properties/UhtUInt32Property.cs
+++ b/Engine/Source/Programs/Shared/EpicGames.UHT/Types/Properties/UhtUInt32Property.cs
@@ -29,6 +29,7 @@ namespace EpicGames.UHT.Types
 		/// <param name="intType">Integer type</param>
 		public UhtUInt32Property(UhtPropertySettings propertySettings, UhtPropertyIntType intType) : base(propertySettings, intType)
 		{
+			PropertyCaps |= UhtPropertyCaps.SupportsRigVM;
 		}
 
 		/// <inheritdoc/>
diff --git a/Engine/Source/Runtime/RigVM/Private/RigVMCore/RigVMMemoryStorage.cpp b/Engine/Source/Runtime/RigVM/Private/RigVMCore/RigVMMemoryStorage.cpp
--- a/Engine/Source/Runtime/RigVM/Private/RigVMCore/RigVMMemoryStorage.cpp
+++ b/Engine/Source/Runtime/RigVM/Private/RigVMCore/RigVMMemoryStorage.cpp
@@ -780,39 +780,37 @@ FProperty* URigVMMemoryStorageGeneratorClass::AddProperty(URigVMMemoryStorageGen
 		}
 		else // take care of simple types...
 		{
-			static FString BoolString = TEXT("bool");
-			static FString Int32String = TEXT("int32");
-			static FString IntString = TEXT("int");
-			static FString FloatString = TEXT("float");
-			static FString DoubleString = TEXT("double");
-			static FString StringString = TEXT("FString");
-			static FString NameString = TEXT("FName");
+			static const FString IntString = TEXT("int");
 
 			const FString BaseCPPType = InProperty.GetTailCPPType();
-			if(BaseCPPType.Equals(BoolString, ESearchCase::IgnoreCase))
+			if(BaseCPPType.Equals(RigVMTypeUtils::BoolType, ESearchCase::IgnoreCase))
 			{
 				FBoolProperty* BoolProperty = new FBoolProperty(PropertyOwner, InProperty.Name, RF_Public); 
 				BoolProperty->SetBoolSize(sizeof(bool), true);
 				(*ValuePropertyPtr) = BoolProperty;
 			}
-			else if(BaseCPPType.Equals(Int32String, ESearchCase::IgnoreCase) ||
+			else if(BaseCPPType.Equals(RigVMTypeUtils::Int32Type, ESearchCase::IgnoreCase) ||
 				BaseCPPType.Equals(IntString, ESearchCase::IgnoreCase))
 			{
 				(*ValuePropertyPtr) = new FIntProperty(PropertyOwner, InProperty.Name, RF_Public);;
 			}
-			else if(BaseCPPType.Equals(FloatString, ESearchCase::IgnoreCase))
+			else if(BaseCPPType.Equals(RigVMTypeUtils::UInt32Type, ESearchCase::IgnoreCase))
+			{
+				(*ValuePropertyPtr) = new FUInt32Property(PropertyOwner, InProperty.Name, RF_Public);;
+			}
+			else if(BaseCPPType.Equals(RigVMTypeUtils::FloatType, ESearchCase::IgnoreCase))
 			{
 				(*ValuePropertyPtr) = new FFloatProperty(PropertyOwner, InProperty.Name, RF_Public);;
 			}
-			else if(BaseCPPType.Equals(DoubleString, ESearchCase::IgnoreCase))
+			else if(BaseCPPType.Equals(RigVMTypeUtils::DoubleType, ESearchCase::IgnoreCase))
 			{
 				(*ValuePropertyPtr) = new FDoubleProperty(PropertyOwner, InProperty.Name, RF_Public);;
 			}
-			else if(BaseCPPType.Equals(StringString, ESearchCase::IgnoreCase))
+			else if(BaseCPPType.Equals(RigVMTypeUtils::FStringType, ESearchCase::IgnoreCase))
 			{
 				(*ValuePropertyPtr) = new FStrProperty(PropertyOwner, InProperty.Name, RF_Public);;
 			}
-			else if(BaseCPPType.Equals(NameString, ESearchCase::IgnoreCase))
+			else if(BaseCPPType.Equals(RigVMTypeUtils::FNameType, ESearchCase::IgnoreCase))
 			{
 				(*ValuePropertyPtr) = new FNameProperty(PropertyOwner, InProperty.Name, RF_Public);;
 			}
diff --git a/Engine/Source/Runtime/RigVM/Private/RigVMCore/RigVMRegistry.cpp b/Engine/Source/Runtime/RigVM/Private/RigVMCore/RigVMRegistry.cpp
--- a/Engine/Source/Runtime/RigVM/Private/RigVMCore/RigVMRegistry.cpp
+++ b/Engine/Source/Runtime/RigVM/Private/RigVMCore/RigVMRegistry.cpp
@@ -258,6 +258,7 @@ void FRigVMRegistry::InitializeIfNeeded()
 	RigVMTypeUtils::TypeIndex::Float = FindOrAddType(FRigVMTemplateArgumentType(RigVMTypeUtils::FloatTypeName, nullptr));
 	RigVMTypeUtils::TypeIndex::Double = FindOrAddType(FRigVMTemplateArgumentType(RigVMTypeUtils::DoubleTypeName, nullptr));
 	RigVMTypeUtils::TypeIndex::Int32 = FindOrAddType(FRigVMTemplateArgumentType(RigVMTypeUtils::Int32TypeName, nullptr));
+	RigVMTypeUtils::TypeIndex::UInt32 = FindOrAddType(FRigVMTemplateArgumentType(RigVMTypeUtils::UInt32TypeName, nullptr));
 	RigVMTypeUtils::TypeIndex::UInt8 = FindOrAddType(FRigVMTemplateArgumentType(RigVMTypeUtils::UInt8TypeName, nullptr));
 	RigVMTypeUtils::TypeIndex::FName = FindOrAddType(FRigVMTemplateArgumentType(RigVMTypeUtils::FNameTypeName, nullptr));
 	RigVMTypeUtils::TypeIndex::FString = FindOrAddType(FRigVMTemplateArgumentType(RigVMTypeUtils::FStringTypeName, nullptr));
@@ -266,6 +267,7 @@ void FRigVMRegistry::InitializeIfNeeded()
 	RigVMTypeUtils::TypeIndex::FloatArray = FindOrAddType(FRigVMTemplateArgumentType(RigVMTypeUtils::FloatArrayTypeName, nullptr));
 	RigVMTypeUtils::TypeIndex::DoubleArray = FindOrAddType(FRigVMTemplateArgumentType(RigVMTypeUtils::DoubleArrayTypeName, nullptr));
 	RigVMTypeUtils::TypeIndex::Int32Array = FindOrAddType(FRigVMTemplateArgumentType(RigVMTypeUtils::Int32ArrayTypeName, nullptr));
+	RigVMTypeUtils::TypeIndex::UInt32Array = FindOrAddType(FRigVMTemplateArgumentType(RigVMTypeUtils::UInt32ArrayTypeName, nullptr));
 	RigVMTypeUtils::TypeIndex::UInt8Array = FindOrAddType(FRigVMTemplateArgumentType(RigVMTypeUtils::UInt8ArrayTypeName, nullptr));
 	RigVMTypeUtils::TypeIndex::FNameArray = FindOrAddType(FRigVMTemplateArgumentType(RigVMTypeUtils::FNameArrayTypeName, nullptr));
 	RigVMTypeUtils::TypeIndex::FStringArray = FindOrAddType(FRigVMTemplateArgumentType(RigVMTypeUtils::FStringArrayTypeName, nullptr));
diff --git a/Engine/Source/Runtime/RigVM/Private/RigVMTypeUtils.cpp b/Engine/Source/Runtime/RigVM/Private/RigVMTypeUtils.cpp
--- a/Engine/Source/Runtime/RigVM/Private/RigVMTypeUtils.cpp
+++ b/Engine/Source/Runtime/RigVM/Private/RigVMTypeUtils.cpp
@@ -14,6 +14,7 @@ namespace RigVMTypeUtils
 	TRigVMTypeIndex TypeIndex::Float = INDEX_NONE;	
 	TRigVMTypeIndex TypeIndex::Double = INDEX_NONE;	
 	TRigVMTypeIndex TypeIndex::Int32 = INDEX_NONE;	
+	TRigVMTypeIndex TypeIndex::UInt32 = INDEX_NONE;	
 	TRigVMTypeIndex TypeIndex::UInt8 = INDEX_NONE;	
 	TRigVMTypeIndex TypeIndex::FName = INDEX_NONE;	
 	TRigVMTypeIndex TypeIndex::FString = INDEX_NONE;
@@ -22,6 +23,7 @@ namespace RigVMTypeUtils
 	TRigVMTypeIndex TypeIndex::FloatArray = INDEX_NONE;	
 	TRigVMTypeIndex TypeIndex::DoubleArray = INDEX_NONE;	
 	TRigVMTypeIndex TypeIndex::Int32Array = INDEX_NONE;	
+	TRigVMTypeIndex TypeIndex::UInt32Array = INDEX_NONE;	
 	TRigVMTypeIndex TypeIndex::UInt8Array = INDEX_NONE;	
 	TRigVMTypeIndex TypeIndex::FNameArray = INDEX_NONE;	
 	TRigVMTypeIndex TypeIndex::FStringArray = INDEX_NONE;	
diff --git a/Engine/Source/Runtime/RigVM/Public/RigVMTypeUtils.h b/Engine/Source/Runtime/RigVM/Public/RigVMTypeUtils.h
--- a/Engine/Source/Runtime/RigVM/Public/RigVMTypeUtils.h
+++ b/Engine/Source/Runtime/RigVM/Public/RigVMTypeUtils.h
@@ -23,6 +23,7 @@ namespace RigVMTypeUtils
 	const FString FloatType = TEXT("float");
 	const FString DoubleType = TEXT("double");
 	const FString Int32Type = TEXT("int32");
+	const FString UInt32Type = TEXT("uint32");
 	const FString UInt8Type = TEXT("uint8");
 	const FString FNameType = TEXT("FName");
 	const FString FStringType = TEXT("FString");
@@ -30,6 +31,7 @@ namespace RigVMTypeUtils
 	const FString FloatArrayType = TEXT("TArray<float>");
 	const FString DoubleArrayType = TEXT("TArray<double>");
 	const FString Int32ArrayType = TEXT("TArray<int32>");
+	const FString UInt32ArrayType = TEXT("TArray<uint32>");
 	const FString UInt8ArrayType = TEXT("TArray<uint8>");
 	const FString FNameArrayType = TEXT("TArray<FName>");
 	const FString FStringArrayType = TEXT("TArray<FString>");
@@ -38,6 +40,7 @@ namespace RigVMTypeUtils
 	const FName FloatTypeName = *FloatType;
 	const FName DoubleTypeName = *DoubleType;
 	const FName Int32TypeName = *Int32Type;
+	const FName UInt32TypeName = *UInt32Type;
 	const FName UInt8TypeName = *UInt8Type;
 	const FName FNameTypeName = *FNameType;
 	const FName FStringTypeName = *FStringType;
@@ -45,6 +48,7 @@ namespace RigVMTypeUtils
 	const FName FloatArrayTypeName = *FloatArrayType;
 	const FName DoubleArrayTypeName = *DoubleArrayType;
 	const FName Int32ArrayTypeName = *Int32ArrayType;
+	const FName UInt32ArrayTypeName = *UInt32ArrayType;
 	const FName UInt8ArrayTypeName = *UInt8ArrayType;
 	const FName FNameArrayTypeName = *FNameArrayType;
 	const FName FStringArrayTypeName = *FStringArrayType;
@@ -58,6 +62,7 @@ namespace RigVMTypeUtils
 		static TRigVMTypeIndex Float;	
 		static TRigVMTypeIndex Double;	
 		static TRigVMTypeIndex Int32;	
+		static TRigVMTypeIndex UInt32;	
 		static TRigVMTypeIndex UInt8;	
 		static TRigVMTypeIndex FName;	
 		static TRigVMTypeIndex FString;
@@ -66,6 +71,7 @@ namespace RigVMTypeUtils
 		static TRigVMTypeIndex FloatArray;	
 		static TRigVMTypeIndex DoubleArray;	
 		static TRigVMTypeIndex Int32Array;	
+		static TRigVMTypeIndex UInt32Array;	
 		static TRigVMTypeIndex UInt8Array;	
 		static TRigVMTypeIndex FNameArray;	
 		static TRigVMTypeIndex FStringArray;	
@@ -185,14 +191,14 @@ namespace RigVMTypeUtils
 		{
 			for(const FString& Prefix : PrefixesNotRequiringCPPTypeObject)
 			{
-				if(InCPPType.StartsWith(Prefix, ESearchCase::CaseSensitive))
+				if(InCPPType.StartsWith(Prefix, ESearchCase::IgnoreCase))
 				{
 					return false;
 				}
 			}
 			for(const FString& Prefix : PrefixesRequiringCPPTypeObject)
 			{
-				if(InCPPType.StartsWith(Prefix, ESearchCase::CaseSensitive))
+				if(InCPPType.StartsWith(Prefix, ESearchCase::IgnoreCase))
 				{
 					return true;
 				}
