RigVM: Fix crash on editor startup

diff --git a/Engine/Source/Runtime/RigVM/Public/RigVMTypeUtils.h b/Engine/Source/Runtime/RigVM/Public/RigVMTypeUtils.h
--- a/Engine/Source/Runtime/RigVM/Public/RigVMTypeUtils.h
+++ b/Engine/Source/Runtime/RigVM/Public/RigVMTypeUtils.h
@@ -176,7 +176,10 @@ namespace RigVMTypeUtils
 		static const TArray<FString> PrefixesNotRequiringCPPTypeObject = {
 			TEXT("UInt"), 
 			TEXT("TArray<UInt"),
-			TEXT("TArray<TArray<UInt")
+			TEXT("TArray<TArray<UInt"),
+			TEXT("uint"), 
+			TEXT("TArray<uint"),
+			TEXT("TArray<TArray<uint")
 		};
 		static const TArray<FString> CPPTypesNotRequiringCPPTypeObject = {
 			TEXT("FString"), 
@@ -191,14 +194,14 @@ namespace RigVMTypeUtils
 		{
 			for(const FString& Prefix : PrefixesNotRequiringCPPTypeObject)
 			{
-				if(InCPPType.StartsWith(Prefix, ESearchCase::IgnoreCase))
+				if(InCPPType.StartsWith(Prefix, ESearchCase::CaseSensitive))
 				{
 					return false;
 				}
 			}
 			for(const FString& Prefix : PrefixesRequiringCPPTypeObject)
 			{
-				if(InCPPType.StartsWith(Prefix, ESearchCase::IgnoreCase))
+				if(InCPPType.StartsWith(Prefix, ESearchCase::CaseSensitive))
 				{
 					return true;
 				}
