RigVM: Ensure to load missing enums + user defined structs

diff --git a/Engine/Source/Runtime/RigVM/Private/RigVMCore/RigVMRegistry.cpp b/Engine/Source/Runtime/RigVM/Private/RigVMCore/RigVMRegistry.cpp
--- a/Engine/Source/Runtime/RigVM/Private/RigVMCore/RigVMRegistry.cpp
+++ b/Engine/Source/Runtime/RigVM/Private/RigVMCore/RigVMRegistry.cpp
@@ -310,35 +310,31 @@ void FRigVMRegistry::RefreshEngineTypes()
 	TArray<UScriptStruct*> DispatchFactoriesToRegister;
 	DispatchFactoriesToRegister.Reserve(32);
 
-	static TArray<const UClass *> ClassesToEnumerate{UScriptStruct::StaticClass(), UEnum::StaticClass()};
-	ForEachObjectOfClasses(
-		ClassesToEnumerate,
-		[this, &DispatchFactoriesToRegister](UObject* InObject) -> void
+	for (TObjectIterator<UScriptStruct> ScriptStructIt; ScriptStructIt; ++ScriptStructIt)
+	{
+		UScriptStruct* ScriptStruct = *ScriptStructIt;
+		
+		// if this is a C++ type - skip it
+		if(ScriptStruct->IsA<UUserDefinedStruct>() || ScriptStruct->IsChildOf(FRigVMExecuteContext::StaticStruct()))
 		{
-			if (UScriptStruct* ScriptStruct = Cast<UScriptStruct>(InObject))
-			{
-				// if this is a C++ type - skip it
-				if(ScriptStruct->IsA<UUserDefinedStruct>() || ScriptStruct->IsChildOf(FRigVMExecuteContext::StaticStruct()))
-				{
-					FindOrAddType(FRigVMTemplateArgumentType(ScriptStruct));
-				}
-				else if (ScriptStruct != FRigVMDispatchFactory::StaticStruct() &&
-						 ScriptStruct->IsChildOf(FRigVMDispatchFactory::StaticStruct()))
-				{
-					DispatchFactoriesToRegister.Add(ScriptStruct);
-				}
-			}
-			else if (UEnum* Enum = Cast<UEnum>(InObject))
-			{
-				// Ignore reflected C++ enums.
-				if(IsAllowedType(Enum) && !Enum->IsNative())
-				{
-					const FString CPPType = Enum->CppType.IsEmpty() ? Enum->GetName() : Enum->CppType;
-					FindOrAddType(FRigVMTemplateArgumentType(*CPPType, Enum));
-				}
-			}
+			FindOrAddType(FRigVMTemplateArgumentType(ScriptStruct));
+		}
+		else if (ScriptStruct != FRigVMDispatchFactory::StaticStruct() &&
+				 ScriptStruct->IsChildOf(FRigVMDispatchFactory::StaticStruct()))
+		{
+			DispatchFactoriesToRegister.Add(ScriptStruct);
 		}
-	);
+	}
+
+	for (TObjectIterator<UEnum> EnumIt; EnumIt; ++EnumIt)
+	{
+		UEnum* Enum = *EnumIt;
+		if(IsAllowedType(Enum))
+		{
+			const FString CPPType = Enum->CppType.IsEmpty() ? Enum->GetName() : Enum->CppType;
+			FindOrAddType(FRigVMTemplateArgumentType(*CPPType, Enum));
+		}
+	}
 	
 	// Register all dispatch factories only after all other types have been registered.
 	for (UScriptStruct* DispatchFactoryStruct: DispatchFactoriesToRegister)
@@ -836,10 +832,25 @@ TRigVMTypeIndex FRigVMRegistry::GetTypeIndexFromCPPType(const FString& InCPPType
 	if(ensure(!InCPPType.IsEmpty()))
 	{
 		const FName CPPTypeName = *InCPPType;
-		Result = Types.IndexOfByPredicate([CPPTypeName](const FTypeInfo& Info) -> bool
+
+		auto Predicate = [CPPTypeName](const FTypeInfo& Info) -> bool
 		{
 			return Info.Type.CPPType == CPPTypeName;
-		});
+		};
+
+		Result = Types.IndexOfByPredicate(Predicate);
+
+#if !WITH_EDITOR
+		// in game / non-editor it's possible that a user defined struct or enum 
+		// has not been registered. thus we'll call RefreshEngineTypes to bring
+		// things up to date here. 
+		if(Result == INDEX_NONE)
+		{
+			// we may need ot 
+			FRigVMRegistry::Get().RefreshEngineTypes();
+			Result = Types.IndexOfByPredicate(Predicate);
+		}
+#endif
 
 		// If not found, try to find a redirect
 		if (Result == INDEX_NONE)
